New material translator:
- Added support for opaque aggregate values HLSL generation.
- Added support for user aggregate typed dynamic branches.
- Fixed a bug in Subscript when applied to aggregates.

#rnx

#Synced-CL 50714579

#Synced-CL 50714579

diff --git a/Engine/Shaders/Private/MaterialTemplate.ush b/Engine/Shaders/Private/MaterialTemplate.ush
--- a/Engine/Shaders/Private/MaterialTemplate.ush
+++ b/Engine/Shaders/Private/MaterialTemplate.ush
@@ -360,6 +360,7 @@ SHADER_DISABLE_WARNINGS
 
 /** FMaterialAttributes utilities */
 %{material_attributes_utilities}
+%{material_aggregates_declarations}
 
 /** Custom output declarations */
 %{custom_outputs}
diff --git a/Engine/Source/Runtime/Engine/Private/Materials/MaterialExpressions.cpp b/Engine/Source/Runtime/Engine/Private/Materials/MaterialExpressions.cpp
--- a/Engine/Source/Runtime/Engine/Private/Materials/MaterialExpressions.cpp
+++ b/Engine/Source/Runtime/Engine/Private/Materials/MaterialExpressions.cpp
@@ -21809,6 +21809,12 @@ EMaterialValueType UMaterialExpressionIfThenElse::GetInputValueType(int32 InputI
 		default: checkNoEntry(); return {};
 	}
 }
+	
+EMaterialValueType UMaterialExpressionIfThenElse::GetOutputValueType(int32 OutputIndex)
+{
+	return MCT_Unknown;
+}
+
 #endif
 
 ///////////////////////////////////////////////////////////////////////////////
diff --git a/Engine/Source/Runtime/Engine/Private/Materials/MaterialExpressionsIR.cpp b/Engine/Source/Runtime/Engine/Private/Materials/MaterialExpressionsIR.cpp
--- a/Engine/Source/Runtime/Engine/Private/Materials/MaterialExpressionsIR.cpp
+++ b/Engine/Source/Runtime/Engine/Private/Materials/MaterialExpressionsIR.cpp
@@ -3160,7 +3160,7 @@ static FValueRef BlendAggregate(MIR::FEmitter& Em,
 static FValueRef Blend(MIR::FEmitter& Em, EMaterialExpressionBlendMode PixelAttributeBlendMode, EMaterialExpressionBlendMode VertexAttributeBlendMode, FValueRef A, FValueRef B, FValueRef Alpha)
 {
 	// Find the common type between arguments
-	if (MIR::FType CommonType = Em.GetCommonType(A->Type, B->Type))
+	if (MIR::FType CommonType = Em.GetCommonType(A->Type, B->Type); !CommonType.IsPoison())
 	{
 		// And cast both arguments to the common type
 		A = Em.Cast(A, CommonType);
diff --git a/Engine/Source/Runtime/Engine/Private/Materials/MaterialIRDebug.cpp b/Engine/Source/Runtime/Engine/Private/Materials/MaterialIRDebug.cpp
--- a/Engine/Source/Runtime/Engine/Private/Materials/MaterialIRDebug.cpp
+++ b/Engine/Source/Runtime/Engine/Private/Materials/MaterialIRDebug.cpp
@@ -65,7 +65,7 @@ struct FDebugDumpIRUseGraphState
 		Out.Appendf(TEXT("\"%p\" [label=< <b>%s</b>  (%s) <br/> "),
 						Value,
 						LexToString(Value->Kind),
-						Value->Type ? *Value->Type.GetSpelling() : TEXT("???"));
+						!Value->Type.IsPoison() ? *Value->Type.GetSpelling() : TEXT("???"));
 
 		DumpValueInfo(Value);
 
diff --git a/Engine/Source/Runtime/Engine/Private/Materials/MaterialIREmitter.cpp b/Engine/Source/Runtime/Engine/Private/Materials/MaterialIREmitter.cpp
--- a/Engine/Source/Runtime/Engine/Private/Materials/MaterialIREmitter.cpp
+++ b/Engine/Source/Runtime/Engine/Private/Materials/MaterialIREmitter.cpp
@@ -402,7 +402,7 @@ FType FEmitter::TryGetCommonType(FType A, FType B)
 
 FType FEmitter::GetCommonType(FType A, FType B)
 {
-	if (FType CommonType = TryGetCommonType(A, B))
+	if (FType CommonType = TryGetCommonType(A, B); !CommonType.IsPoison())
 	{
 		return CommonType;
 	}
@@ -2594,7 +2594,7 @@ static FType ValidateOperatorAndGetResultType(FEmitter& Emitter, EOperator Op, F
 			FType::MakeVector(CPrimitiveType.ScalarKind, MaxNumComponents));
 		
 		// Getting the common type should always be possible.
-		check(CommonTypeBetweenSecondAndThirdArguments);
+		check(!CommonTypeBetweenSecondAndThirdArguments.IsPoison());
 
 		// Cast second and third arguments to their common type.
 		Arguments[1] = Emitter.Cast(Arguments[1], CommonTypeBetweenSecondAndThirdArguments);
@@ -2975,7 +2975,7 @@ FValueRef FEmitter::Operator(EOperator Op, FValueRef A, FValueRef B, FValueRef C
 
 	// Validate the operation and retrieve the result type.
 	FType ResultType = ValidateOperatorAndGetResultType(*this, Op, A, B, C);
-	if (!ResultType)
+	if (ResultType.IsPoison())
 	{
 		return Poison();
 	}
@@ -3022,7 +3022,7 @@ FValueRef FEmitter::Branch(FValueRef Condition, FValueRef True, FValueRef False)
 	if (!Condition)
 	{
 		return Poison();
-	}	
+	}
 
 	// If the condition is a scalar constant, then simply evaluate the result now.
 	if (const FConstant* ConstCondition = As<FConstant>(Condition))
@@ -3033,8 +3033,15 @@ FValueRef FEmitter::Branch(FValueRef Condition, FValueRef True, FValueRef False)
 	// If the condition is not static, make both true and false arguments have the same type,
 	// by casting false argument into the true's type.
 	FType CommonType = GetCommonType({ True, False });
-	if (!CommonType)
+	if (CommonType.IsPoison())
+	{
+		return Poison();
+	}
+
+	// Branching on material attributes isn't supported yet.
+	if (CommonType.AsAggregate() == MaterialAttributesAggregate::Get())
 	{
+		Error(TEXT("Dynamic branching on MaterialAttributes is not yet supported due to the Normal attribute being evaluated before all other attributes. Branching on any other user defined aggregate is supported."));
 		return Poison();
 	}
 
@@ -3056,52 +3063,71 @@ FValueRef FEmitter::Branch(FValueRef Condition, FValueRef True, FValueRef False)
 
 FValueRef FEmitter::Subscript(FValueRef Value, int32 Index)
 {
+	check(Index >= 0);
+
 	if (!Value.IsValid())
 	{
 		return Value;
 	}
-	
+
 	// Subscripting a composite by index is always possible and simply yields the i-th component.
 	if (FComposite* Composite = As<FComposite>(Value))
 	{
 		check(Index < Composite->GetComponents().Num());
 		return Value.To(Composite->GetComponents()[Index]);
 	}
-	
-	// Other operations supported if value is primitive. Check it first.
-	TOptional<FPrimitive> PrimitiveType = Value->Type.AsPrimitive();
-	if (!PrimitiveType)
-	{
-		Errorf(Value, TEXT("Value of type '%s' cannot be subscripted."), *Value->Type.GetSpelling());
-		return Value.ToPoison();
-	}
 
-	// Getting first component and Value is already a scalar, just return itself.
-	if (Index == 0 && PrimitiveType->IsScalar())
-	{
-		return Value;
-	}
+	// The final value of the subscript instruction should it be emitted 
+	FType SubscriptType;
 
-	if (Index >= PrimitiveType->NumComponents())
+	if (const UMaterialAggregate* Aggregate = Value->Type.AsAggregate())
 	{
-		Errorf(Value, TEXT("Value of type '%s' has fewer dimensions than subscript index `%d`."), *Value->Type.GetSpelling(), Index);
-		return Value.ToPoison();
-	}
+		// Make sure the attribute index is valid
+		if (!Aggregate->Attributes.IsValidIndex(Index))
+		{
+			Errorf(Value, TEXT("Attribute index %d is out of range for material aggregate '%s'."), Index, *Aggregate->GetName());
+			return Value.ToPoison();
+		}
 
-	if (PrimitiveType->IsMatrix() && PrimitiveType->IsDouble())
-	{
-		Errorf(Value, TEXT("Cannot subscript a double-precision matrix."));
-		return Value.ToPoison();
+		// Get the type of the material 
+		SubscriptType = GetMaterialAggregateAttributeType(Aggregate, Index);
 	}
+	else if (TOptional<FPrimitive> PrimitiveType = Value->Type.AsPrimitive())
+	{
+		SubscriptType = PrimitiveType->ToScalar();
+
+		// Getting first component and Value is already a scalar, just return itself.
+		if (Index == 0 && PrimitiveType->IsScalar())
+		{
+			return Value;
+		}
+
+		if (Index >= PrimitiveType->NumComponents())
+		{
+			Errorf(Value, TEXT("Value of type '%s' has fewer dimensions than subscript index `%d`."), *Value->Type.GetSpelling(), Index);
+			return Value.ToPoison();
+		}
+
+		if (PrimitiveType->IsMatrix() && PrimitiveType->IsDouble())
+		{
+			Errorf(Value, TEXT("Cannot subscript a double-precision matrix."));
+			return Value.ToPoison();
+		}
 
-	// Avoid subscripting a subscript (e.g. no value.xy.x)
-	if (FSubscript* Subscript = As<FSubscript>(Value))
+		// Avoid double subscripting a primitive value (e.g. no value.xy.x)
+		if (FSubscript* Subscript = As<FSubscript>(Value); Subscript && Subscript->Arg->Type.AsPrimitive())
+		{
+			Value = Value.To(Subscript->Arg);
+		}
+	}
+	else
 	{
-		Value = Value.To(Subscript->Arg);
+		Errorf(Value, TEXT("Value of type '%s' cannot be subscripted."), *Value->Type.GetSpelling());
+		return Value.ToPoison();
 	}
 
 	// We can't resolve it at compile time: emit subscript value.
-	FSubscript Prototype = MakePrototype<FSubscript>(PrimitiveType->ToScalar());
+	FSubscript Prototype = MakePrototype<FSubscript>(SubscriptType);
 	Prototype.Arg = Value;
 	Prototype.Index = Index;
 
diff --git a/Engine/Source/Runtime/Engine/Private/Materials/MaterialIRToHLSLTranslator.cpp b/Engine/Source/Runtime/Engine/Private/Materials/MaterialIRToHLSLTranslator.cpp
--- a/Engine/Source/Runtime/Engine/Private/Materials/MaterialIRToHLSLTranslator.cpp
+++ b/Engine/Source/Runtime/Engine/Private/Materials/MaterialIRToHLSLTranslator.cpp
@@ -45,9 +45,11 @@ struct FHLSLPrinter
 	FString Buffer;
 	bool bFirstListItem = false;
 	int32 Tabs = 0;
-
-	FHLSLPrinter(int32 Tabs = 0)
+	TArray<const UMaterialAggregate*>& MaterialAggregatesRequiringDeclaration;
+	
+	FHLSLPrinter(TArray<const UMaterialAggregate*>& MaterialAggregatesRequiringDeclaration, int32 Tabs = 0)
 		: Tabs{ Tabs }
+		, MaterialAggregatesRequiringDeclaration { MaterialAggregatesRequiringDeclaration }
 	{
 		*this << Indentation;
 	}
@@ -175,6 +177,18 @@ struct FHLSLPrinter
 	}
 };
 
+static void PrintUserIdentifier(FHLSLPrinter& Printer, FStringView Identifier)
+{
+	if (!Identifier.IsEmpty() && FChar::IsDigit(Identifier[0]))
+	{
+		Printer << TEXT('_');
+	}
+	for (TCHAR Char : Identifier)
+	{
+		Printer << (FChar::IsAlnum(Char) || Char == TEXT('_') ? Char : TEXT('_'));
+	}
+}
+
 // Prints the MIR type HLSL declaration to the specified printer.
 static FHLSLPrinter& operator<<(FHLSLPrinter& Printer, const MIR::FType& Type)
 {
@@ -218,6 +232,12 @@ static FHLSLPrinter& operator<<(FHLSLPrinter& Printer, const MIR::FType& Type)
 			}
 		}
 	}
+	else if (const UMaterialAggregate* Aggregate = Type.AsAggregate())
+	{
+		Printer << TEXT('F');
+		PrintUserIdentifier(Printer, Aggregate->GetName());
+		Printer.MaterialAggregatesRequiringDeclaration.AddUnique(Aggregate);
+	}
 	else
 	{
 		switch (Type.GetKind())
@@ -251,7 +271,7 @@ struct FEntryPointCodeGen
 	const FMaterialIRModule* Module{};
 
 	// Utility for emitting HLSL code
-	FHLSLPrinter Printer{};
+	FHLSLPrinter Printer;
 
 	// Used internally. See FResult::OutputNormalAttributeCode
 	FString OutputNormalAttributeCode;
@@ -276,15 +296,14 @@ struct FEntryPointCodeGen
 	// Used to increase generated HLSL readability, so that we add a newline before the first instruction that follows a SetMateiralOutput.
 	bool bLastInstructionWasSetMaterialOutput{};
 
-	FEntryPointCodeGen(const FMaterialIRModule* Module)
+	FEntryPointCodeGen(const FMaterialIRModule* Module, TArray<const UMaterialAggregate*>& MaterialAggregateRequiringDeclaration)
 	: Module { Module }
+	, Printer{ MaterialAggregateRequiringDeclaration, 1 }
 	{}
 
 	// Generates the full HLSL of the specified entry point.
 	FResult GenerateEntryPoint(int32 ModuleEntryPointIndex, bool bPreviousFrame = false)
 	{
-		Printer = { 1 };
-
 		const FMaterialIRModule::FEntryPoint& EntryPoint = Module->GetEntryPoint(ModuleEntryPointIndex);
 
 		CurrentEntryPointIndex = ModuleEntryPointIndex;
@@ -350,7 +369,6 @@ struct FEntryPointCodeGen
 				// "Normal" is treated in a special way because the rest of the attributes may lead back to reading it.
 				// Therefore, in the way MaterialTemplate.ush is structured, it needs to be evaluated before other attributes.
 				OutputNormalAttributeCode = MoveTemp(Printer.Buffer);
-				Printer = { 1 };
 				bLastInstructionWasSetMaterialOutput = false;
 			}
 		}
@@ -1052,9 +1070,24 @@ struct FEntryPointCodeGen
 
 	void LowerComposite(const MIR::FComposite* Composite)
 	{
-		TConstArrayView<MIR::FValue*> Components = Composite->GetComponents();
-		MIR::FPrimitive PrimitiveType = Composite->Type.GetPrimitive();
+		if (TOptional<MIR::FPrimitive> PrimitiveType = Composite->Type.AsPrimitive())
+		{
+			LowerPrimitiveComposite(Composite, *PrimitiveType);
+		}
+		else if (const UMaterialAggregate* Aggregate = Composite->Type.AsAggregate())
+		{
+			LowerAggregateComposite(Composite, Aggregate);
+		}
+		else
+		{
+			UE_MIR_UNREACHABLE();
+		}
+	}
+
+	void LowerPrimitiveComposite(const MIR::FComposite* Composite, MIR::FPrimitive PrimitiveType)
+	{
 		check(!PrimitiveType.IsScalar());
+		TConstArrayView<MIR::FValue*> Components = Composite->GetComponents();
 		
 		// In order to generate smaller and tidier HLSL, first check whether all components
 		// of this dimensional are actually the same. If so, we can simply emit the 
@@ -1174,6 +1207,25 @@ struct FEntryPointCodeGen
 		}
 	}
 
+	void LowerAggregateComposite(const MIR::FComposite* Composite, const UMaterialAggregate* Aggregate)
+	{
+		Printer << EndOfStatement; // end the declaration
+
+		int Local = InstrToLocalIndex[Composite];
+
+		auto Components = Composite->GetComponents();
+		for (int i = 0; i < Components.Num(); ++i)
+		{
+			Printer << TEXT('_') << Local << TEXT('.');
+			PrintUserIdentifier(Printer, Aggregate->Attributes[i].Name.ToString()); 
+			Printer << TEXT(" = ") << LowerValue(Components[i]);
+			if (i < Components.Num() - 1)
+			{
+				Printer << EndOfStatement;
+			}
+		}
+	}
+
 	void LowerSetMaterialOutput(const MIR::FSetMaterialOutput* SetMaterialOutput)
 	{
 		if (SetMaterialOutput->Property >= MP_CustomizedUVs0 && SetMaterialOutput->Property <= MP_CustomizedUVs7)
@@ -1405,32 +1457,32 @@ struct FEntryPointCodeGen
 	{
 		if (IsFoldable(Branch, CurrentEntryPointIndex))
 		{
-			if (Branch->TrueArg->Type.IsDouble())
+			if (Branch->Type.IsDouble())
 			{
 				Printer << TEXTVIEW("WSSelect(") << LowerValue(Branch->ConditionArg)
 					<< TEXTVIEW(", ") << LowerValue(Branch->TrueArg)
 					<< TEXTVIEW(", ") << LowerValue(Branch->FalseArg)
 					<< TEXTVIEW(")");
+				return;
 			}
-			else
+			else if (Branch->Type.AsPrimitive())
 			{
-				Printer << LowerValue(Branch->ConditionArg)
+				Printer << TEXT('(') << LowerValue(Branch->ConditionArg)
 					<< TEXTVIEW(" ? ") << LowerValue(Branch->TrueArg)
-					<< TEXTVIEW(" : ") << LowerValue(Branch->FalseArg);
+					<< TEXTVIEW(" : ") << LowerValue(Branch->FalseArg) << TEXT(')');
+				return;
 			}
 		}
-		else
-		{
-			Printer << EndOfStatement;
-			Printer << TEXTVIEW("if (") << LowerValue(Branch->ConditionArg) << TEXTVIEW(")") << OpenBlock;
-			Printer << LowerBlock(Branch->TrueBlock[CurrentEntryPointIndex]);
-			Printer << TEXTVIEW("_") << InstrToLocalIndex[Branch] << TEXTVIEW(" = ") << LowerValue(Branch->TrueArg) << EndOfStatement;
-			Printer << CloseBlock << NewLine;
-			Printer << TEXTVIEW("else") << OpenBlock;
-			Printer << LowerBlock(Branch->FalseBlock[CurrentEntryPointIndex]);
-			Printer << TEXTVIEW("_") << InstrToLocalIndex[Branch] << TEXTVIEW(" = ") << LowerValue(Branch->FalseArg) << EndOfStatement;
-			Printer << CloseBlock;
-		}
+
+		Printer << EndOfStatement;
+		Printer << TEXTVIEW("if (") << LowerValue(Branch->ConditionArg) << TEXTVIEW(")") << OpenBlock;
+		Printer << LowerBlock(Branch->TrueBlock[CurrentEntryPointIndex]);
+		Printer << TEXTVIEW("_") << InstrToLocalIndex[Branch] << TEXTVIEW(" = ") << LowerValue(Branch->TrueArg) << EndOfStatement;
+		Printer << CloseBlock << NewLine;
+		Printer << TEXTVIEW("else") << OpenBlock;
+		Printer << LowerBlock(Branch->FalseBlock[CurrentEntryPointIndex]);
+		Printer << TEXTVIEW("_") << InstrToLocalIndex[Branch] << TEXTVIEW(" = ") << LowerValue(Branch->FalseArg) << EndOfStatement;
+		Printer << CloseBlock;
 	}
 			
 	void LowerSubscript(const MIR::FSubscript* Subscript)
@@ -1464,6 +1516,12 @@ struct FEntryPointCodeGen
 			check(Subscript->Index < MatrixType->NumComponents());
 			Printer << TEXTVIEW("._m") << (Subscript->Index % MatrixType->NumRows) << (Subscript->Index / MatrixType->NumRows);
 		}
+		else if (const UMaterialAggregate* Aggregate = Subscript->Arg->Type.AsAggregate())
+		{
+			LowerValue(Subscript->Arg);
+			Printer << TEXT('.');
+			PrintUserIdentifier(Printer, Aggregate->Attributes[Subscript->Index].Name.ToString());
+		}
 		else
 		{
 			// The builder should never emit subscripts of scalar types.
@@ -2227,7 +2285,11 @@ struct FEntryPointCodeGen
 	{
 		if (auto Branch = Instr->As<MIR::FBranch>())
 		{
-			return Branch->TrueBlock[EntryPointIndex].Instructions || Branch->FalseBlock[EntryPointIndex].Instructions;
+			return !Branch->Type.IsPrimitive() || Branch->TrueBlock[EntryPointIndex].Instructions || Branch->FalseBlock[EntryPointIndex].Instructions;
+		}
+		if (Instr->Type.AsAggregate() && Instr->As<MIR::FComposite>())
+		{
+			return true;
 		}
 		return false;
 	}
@@ -2263,9 +2325,9 @@ struct FEntryPointCodeGen
 
 /////////////////////////////////////////////////////////////////////////////////////////////////////////////////////////////
 
-static void GenerateStageEntryPoints(const FMaterialIRModule* Module, TMap<FString, FString>& OutParameters)
+static void GenerateStageEntryPoints(const FMaterialIRModule* Module, TMap<FString, FString>& OutParameters, TArray<const UMaterialAggregate*>& MaterialAggregatesRequiringDeclaration)
 {
-	FEntryPointCodeGen CodeGen { Module };
+	FEntryPointCodeGen CodeGen { Module, MaterialAggregatesRequiringDeclaration };
 	FEntryPointCodeGen::FResult Result;
 
 	// Generate current frame vertex stage HLSL
@@ -2290,12 +2352,12 @@ static void GenerateStageEntryPoints(const FMaterialIRModule* Module, TMap<FStri
 	OutParameters.Add(TEXT("calc_pixel_material_inputs_analytic_derivatives_other_inputs"), MoveTemp(Result.NonNormalHLSL));
 }
 
-static void GenerateCustomOutputDeclarations(const FMaterial* Material, const FMaterialIRModule* Module, TMap<FString, FString>& OutParameters)
+static void GenerateCustomOutputDeclarations(const FMaterial* Material, const FMaterialIRModule* Module, TMap<FString, FString>& OutParameters, TArray<const UMaterialAggregate*>& MaterialAggregatesRequiringDeclaration)
 {
 	// Generate the custom output fields of the struct FMaterialPixelParameters. These fields will be computed and 
 	// set by the function that calculates all material outputs (e.g. BaseColor).
-	FHLSLPrinter VertexPrinter{ 1 };
-	FHLSLPrinter PixelPrinter{ 1 };
+	FHLSLPrinter VertexPrinter{ MaterialAggregatesRequiringDeclaration, 1 };
+	FHLSLPrinter PixelPrinter{ MaterialAggregatesRequiringDeclaration, 1 };
 
 	const int NumCustomizedUVs = Material->GetNumCustomizedUVs();
 	if (NumCustomizedUVs > 0)
@@ -2317,7 +2379,7 @@ static void GenerateCustomOutputDeclarations(const FMaterial* Material, const FM
 	// Development Note: this is really just cruft necessary for backwards compatibility with the way the
 	// 	  old translator supported custom outputs and could be removed entirely with a little refactoring
 	//    of MaterialTemplate.ush.
-	FHLSLPrinter Printer = {};
+	FHLSLPrinter Printer = { MaterialAggregatesRequiringDeclaration };
 	for (const FMaterialIRModule::FCustomOutputGroup& Group : Module->GetCustomOutputGroups())
 	{
 		Printer << TEXTVIEW("#define NUM_MATERIAL_OUTPUTS_GET") << *FString(Group.Name).ToUpper() << TEXT(' ') << Group.NumOutputs << NewLine;
@@ -2361,9 +2423,9 @@ static void GenerateCustomOutputDeclarations(const FMaterial* Material, const FM
 
 // Generates the definitions of all the custom HLSL functions in the module and puts the resulting string
 // into the "custom_functions" source template parameter.
-static void GenerateCustomFunctionsHLSL(const FMaterialIRModule* Module, TMap<FString, FString>& OutParameters)
+static void GenerateCustomFunctionsHLSL(const FMaterialIRModule* Module, TMap<FString, FString>& OutParameters, TArray<const UMaterialAggregate*>& MaterialAggregatesRequiringDeclaration)
 {
-	FHLSLPrinter Printer = { 0 };
+	FHLSLPrinter Printer = { MaterialAggregatesRequiringDeclaration };
 
 	for (const MIR::FFunctionHLSL* Function : Module->GetFunctionHLSLs())
 	{
@@ -2423,28 +2485,14 @@ static void GenerateCustomFunctionsHLSL(const FMaterialIRModule* Module, TMap<FS
 	OutParameters.Add(TEXT("custom_functions"), MoveTemp(Printer.Buffer));
 }
 	
-static void GenerateOtherTemplateStringParameters(const FMaterial* Material, const FMaterialIRModule* Module, TMap<FString, FString>& OutParameters)
+static void GenerateOtherTemplateStringParameters(const FMaterial* Material, const FMaterialIRModule* Module, TMap<FString, FString>& OutParameters, TArray<const UMaterialAggregate*>& MaterialAggregatesRequiringDeclaration)
 {
 	const FMaterialIRModule::FStatistics ModuleStatistics = Module->GetStatistics();
 	const FMaterialIRModule::FSubstrateData& SubstrateData = Module->GetSubstrateData();
 
-	// material_declarations
-	{
-		FHLSLPrinter Printer = {};
-		Printer << TEXT("struct FMaterialAttributes") << OpenBlock;
-		for (const FGuid& AttributeID : FMaterialAttributeDefinitionMap::GetOrderedVisibleAttributeList())
-		{
-			const EMaterialValueType PropertyType = FMaterialAttributeDefinitionMap::GetValueType(AttributeID);
-			Printer << MIR::FType::FromMaterialValueType(PropertyType) << TEXT(' ') << *FMaterialAttributeDefinitionMap::GetAttributeName(AttributeID) << EndOfStatement;
-		}
-		Printer << CloseBlock << TEXT(";\n");
-
-		OutParameters.Add(TEXT("material_declarations"), MoveTemp(Printer.Buffer));
-	}
-
 	// pixel_material_inputs (field declarations of the FPixelMaterialInputs struct)
 	{
-		FHLSLPrinter Printer = { 1 };
+		FHLSLPrinter Printer { MaterialAggregatesRequiringDeclaration, 1 };
 		for (EMaterialProperty Property : MaterialAttributesAggregate::GetMaterialProperties())
 		{
 			if (Property == MP_WorldPositionOffset || (Property >= MP_CustomizedUVs0 && Property <= MP_CustomizedUVs7)) // todo: fixme
@@ -2468,7 +2516,7 @@ static void GenerateOtherTemplateStringParameters(const FMaterial* Material, con
 
 	// get_material_customized_u_vs
 	{
-		FHLSLPrinter Printer = { 1 };
+		FHLSLPrinter Printer { MaterialAggregatesRequiringDeclaration, 1 };
 		for (int32 i = 0; i < ShaderTexCoordsArrayLength; i++)
 		{
 			const bool bIsCustomizedUV = i < MIR::MaxNumCustomizedUVs && Module->IsMaterialPropertyUsed((EMaterialProperty)(MP_CustomizedUVs0 + i));
@@ -2534,6 +2582,74 @@ static const TCHAR* GetShadingModelParameterName(EMaterialShadingModel InModel)
 	}
 }
 
+static void GenerateMaterialAggregatesDeclarations(TMap<FString, FString>& OutParameters, TArray<const UMaterialAggregate*>& MaterialAggregatesRequiringDeclaration)
+{
+	FHLSLPrinter Printer { MaterialAggregatesRequiringDeclaration };
+
+	// Recursively push material aggregates referenced by attributes of existing aggregates
+	for (int32 i = 0; i < MaterialAggregatesRequiringDeclaration.Num(); ++i)
+	{
+		for (const FMaterialAggregateAttribute& Attribute : MaterialAggregatesRequiringDeclaration[i]->Attributes)
+		{
+			if (Attribute.Type == EMaterialAggregateAttributeType::Aggregate)
+			{
+				MaterialAggregatesRequiringDeclaration.AddUnique(Attribute.Aggregate.Get());
+			}
+		}
+	}
+
+	// Generate the aggregates "struct XXX {};" declarations
+	for (int32 i = MaterialAggregatesRequiringDeclaration.Num() - 1; i >= 0; --i)
+	{
+		Printer << TEXTVIEW("\nstruct F");
+		PrintUserIdentifier(Printer, MaterialAggregatesRequiringDeclaration[i]->GetName());
+		Printer << OpenBlock;
+
+		for (const FMaterialAggregateAttribute& Attribute : MaterialAggregatesRequiringDeclaration[i]->Attributes)
+		{
+			switch (Attribute.Type)
+			{
+				case EMaterialAggregateAttributeType::Bool1: Printer << TEXTVIEW("bool"); break;
+				case EMaterialAggregateAttributeType::Bool2: Printer << TEXTVIEW("bool2"); break;
+				case EMaterialAggregateAttributeType::Bool3: Printer << TEXTVIEW("bool3"); break;
+				case EMaterialAggregateAttributeType::Bool4: Printer << TEXTVIEW("bool4"); break;
+
+				case EMaterialAggregateAttributeType::UInt1: Printer << TEXTVIEW("uint"); break;
+				case EMaterialAggregateAttributeType::UInt2: Printer << TEXTVIEW("uint2"); break;
+				case EMaterialAggregateAttributeType::UInt3: Printer << TEXTVIEW("uint3"); break;
+				case EMaterialAggregateAttributeType::UInt4: Printer << TEXTVIEW("uint4"); break;
+
+				case EMaterialAggregateAttributeType::Float1: Printer << TEXTVIEW("MaterialFloat"); break;
+				case EMaterialAggregateAttributeType::Float2: Printer << TEXTVIEW("MaterialFloat2"); break;
+				case EMaterialAggregateAttributeType::Float3: Printer << TEXTVIEW("MaterialFloat3"); break;
+				case EMaterialAggregateAttributeType::Float4: Printer << TEXTVIEW("MaterialFloat4"); break;
+
+				case EMaterialAggregateAttributeType::ShadingModel: Printer << TEXTVIEW("int"); break;
+				case EMaterialAggregateAttributeType::MaterialAttributes: Printer << TEXTVIEW("FMaterialAttributes"); break;
+				case EMaterialAggregateAttributeType::SubstrateData: Printer << TEXTVIEW("FSubstrateData"); break;
+				
+				case EMaterialAggregateAttributeType::Aggregate:
+					Printer << TEXT('F');
+					PrintUserIdentifier(Printer, Attribute.Aggregate->GetName());
+					break;
+
+				default:
+					UE_MIR_UNREACHABLE();
+			}
+
+			Printer << TEXT(' ');
+			
+			PrintUserIdentifier(Printer, Attribute.Name.ToString());
+
+			Printer << EndOfStatement;
+		}
+
+		Printer << CloseBlock << TEXT(";\n");
+	}
+
+	OutParameters.Add(TEXT("material_aggregates_declarations"), MoveTemp(Printer.Buffer));
+}
+
 static void GetShaderCompilerEnvironment(const FMaterial* Material, const FMaterialIRModule* Module, const ITargetPlatform* TargetPlatform, FShaderCompilerEnvironment& OutEnvironment)
 {
 	const FMaterialCompilationOutput& CompilationOutput = Module->GetCompilationOutput();
@@ -2758,10 +2874,14 @@ void FMaterialIRToHLSLTranslation::Run(TMap<FString, FString>& OutShaderParamete
 {
 	OutShaderParameters.Empty();
 
-	GenerateStageEntryPoints(Module, OutShaderParameters);
-	GenerateCustomOutputDeclarations(Material, Module, OutShaderParameters);
-	GenerateCustomFunctionsHLSL(Module, OutShaderParameters);
-	GenerateOtherTemplateStringParameters(Material, Module, OutShaderParameters);
+	TArray<const UMaterialAggregate*> MaterialAggregatesRequiringDeclaration;
+	MaterialAggregatesRequiringDeclaration.Add(MaterialAttributesAggregate::Get());
+
+	GenerateStageEntryPoints(Module, OutShaderParameters, MaterialAggregatesRequiringDeclaration);
+	GenerateCustomOutputDeclarations(Material, Module, OutShaderParameters, MaterialAggregatesRequiringDeclaration);
+	GenerateCustomFunctionsHLSL(Module, OutShaderParameters, MaterialAggregatesRequiringDeclaration);
+	GenerateOtherTemplateStringParameters(Material, Module, OutShaderParameters, MaterialAggregatesRequiringDeclaration);
+	GenerateMaterialAggregatesDeclarations(OutShaderParameters, MaterialAggregatesRequiringDeclaration);
 	GetShaderCompilerEnvironment(Material, Module, TargetPlatform, OutEnvironment);
 }
 
diff --git a/Engine/Source/Runtime/Engine/Public/Materials/MaterialExpressionIfThenElse.h b/Engine/Source/Runtime/Engine/Public/Materials/MaterialExpressionIfThenElse.h
--- a/Engine/Source/Runtime/Engine/Public/Materials/MaterialExpressionIfThenElse.h
+++ b/Engine/Source/Runtime/Engine/Public/Materials/MaterialExpressionIfThenElse.h
@@ -31,6 +31,7 @@ public:
 	virtual void Build(MIR::FEmitter& Emitter) override;
 	virtual void GetCaption(TArray<FString>& OutCaptions) const override;
 	virtual EMaterialValueType GetInputValueType(int32 InputIndex) override;
+	virtual EMaterialValueType GetOutputValueType(int32 OutputIndex) override;
 #endif
 	//~ End UMaterialExpression Interface
 };
diff --git a/Engine/Source/Runtime/Engine/Public/Materials/MaterialIRTypes.h b/Engine/Source/Runtime/Engine/Public/Materials/MaterialIRTypes.h
--- a/Engine/Source/Runtime/Engine/Public/Materials/MaterialIRTypes.h
+++ b/Engine/Source/Runtime/Engine/Public/Materials/MaterialIRTypes.h
@@ -288,9 +288,6 @@ struct FType
 	// Converts this type to a corresponding UE::Shader::EValueType.
 	UE::Shader::EValueType ToValueType() const;
 
-	// Conversion to bool, returns false if the type is Poison, true otherwise.
-	operator bool() const { return Kind != ETypeKind::Poison; }
-
 	bool operator==(FType Other) const;
 	bool operator!=(FType Other) const { return !operator==(Other); }
 
