Enable animation frame stripping cvar support in editor
Useful for debugging purposes

#Synced-CL 50652508

diff --git a/Engine/Source/Runtime/Engine/Private/Animation/AnimSequence.cpp b/Engine/Source/Runtime/Engine/Private/Animation/AnimSequence.cpp
--- a/Engine/Source/Runtime/Engine/Private/Animation/AnimSequence.cpp
+++ b/Engine/Source/Runtime/Engine/Private/Animation/AnimSequence.cpp
@@ -5033,8 +5033,8 @@ FString UAnimSequence::CreateDerivedDataKeyString(const ITargetPlatform* TargetP
 	const char AdditiveType = bIsValidAdditive ? NibbleToTChar(AdditiveAnimType) : '0';
 	const char RefType = bIsValidAdditive ? NibbleToTChar(RefPoseType) : '0';
 
-	bool bPerformFrameStripping = false;
-	bool bPerformFrameStrippingOnOddNumberedFrames = false;
+	bool bPerformFrameStripping = GPerformFrameStripping == 1;
+	bool bPerformFrameStrippingOnOddNumberedFrames = GPerformFrameStrippingOddFramedAnimations == 1;
 	{
 		if (UDeviceProfile* DeviceProfile = UDeviceProfileManager::Get().FindProfile(TargetPlatform->IniPlatformName()))
 		{
@@ -5189,9 +5189,9 @@ FIoHash UAnimSequence::BeginCacheDerivedData(const ITargetPlatform* TargetPlatfo
 	check(TargetData);
 	TargetData->Reset();
 
-	bool bPerformFrameStripping  = false;
+	bool bPerformFrameStripping  = GPerformFrameStripping == 1;
 	{
-		bool bPerformFrameStrippingOnOddNumberedFrames = false;
+		bool bPerformFrameStrippingOnOddNumberedFrames = GPerformFrameStrippingOddFramedAnimations == 1;
 		if (UDeviceProfile* DeviceProfile = UDeviceProfileManager::Get().FindProfile(TargetPlatform->IniPlatformName()))
 		{
 			int32 CVarPlatformFrameStrippingValue = 0;
