Exporting some WP BP functionality so it can be linked.

#Synced-CL 50710515

diff --git a/Engine/Source/Runtime/Engine/Public/WorldPartition/WorldPartitionBlueprintLibrary.h b/Engine/Source/Runtime/Engine/Public/WorldPartition/WorldPartitionBlueprintLibrary.h
--- a/Engine/Source/Runtime/Engine/Public/WorldPartition/WorldPartitionBlueprintLibrary.h
+++ b/Engine/Source/Runtime/Engine/Public/WorldPartition/WorldPartitionBlueprintLibrary.h
@@ -22,10 +22,10 @@ struct FActorDesc
 {
 	GENERATED_USTRUCT_BODY()
 
-	FActorDesc();
+	ENGINE_API FActorDesc();
 
 #if WITH_EDITOR
-	FActorDesc(const FWorldPartitionActorDescInstance& InActorDesc);
+	ENGINE_API FActorDesc(const FWorldPartitionActorDescInstance& InActorDesc);
 #endif
 
 	/** The actor GUID of this descriptor. */
