Fixing crash caused by assets being garbage collected with while the `NotifyToolsOfObjectReplacement` event is pending on main thread.

We are now wrapping the object pointers passed to the GameThread lambda to catch any object becoming invalid before the execution.

#Synced-CL 50672290
#Synced-CL 50672290

diff --git a/Engine/Source/Runtime/Engine/Private/UnrealEngine.cpp b/Engine/Source/Runtime/Engine/Private/UnrealEngine.cpp
--- a/Engine/Source/Runtime/Engine/Private/UnrealEngine.cpp
+++ b/Engine/Source/Runtime/Engine/Private/UnrealEngine.cpp
@@ -18725,11 +18725,27 @@ void UEngine::CopyPropertiesForUnrelatedObjects(UObject* OldObject, UObject* New
 		}
 		else
 		{
+			TMap<FWeakObjectPtr, FWeakObjectPtr> WeakReferenceReplacementMap;
+			Algo::Transform(ReferenceReplacementMap, WeakReferenceReplacementMap, [](const auto& InPair)
+				{ 
+					return TTuple<FWeakObjectPtr, FWeakObjectPtr>(InPair.Key, InPair.Value);
+				});
+
 			ExecuteOnGameThread(
 				TEXT("NotifyToolsOfObjectReplacement"),
-				[ReferenceReplacementMap]()
+				[WeakReferenceReplacementMap = MoveTemp(WeakReferenceReplacementMap)]()
 				{
-					GEngine->NotifyToolsOfObjectReplacement(ReferenceReplacementMap);
+					TMap<UObject*, UObject*> ValidReferenceReplacementMap;
+					Algo::TransformIf(WeakReferenceReplacementMap, ValidReferenceReplacementMap, 
+						[](const auto& InPair)
+						{
+							return InPair.Key.IsValid();
+						},
+						[](const auto& InPair)
+						{
+							return TTuple<UObject*, UObject*>( InPair.Key.Get(), InPair.Value.Get() );
+						});
+					GEngine->NotifyToolsOfObjectReplacement(ValidReferenceReplacementMap);
 				}
 			);
 		}
