Fix textures issue that was happening when loading assets cooked for client(explicitly without EditorOnly data).
In such a case, the ReferencedDefaultTextures array has the required textures.

[RN] minor | Rendering Architecture

#Synced-CL 50642529

diff --git a/Engine/Source/Runtime/Engine/Private/Materials/Material.cpp b/Engine/Source/Runtime/Engine/Private/Materials/Material.cpp
--- a/Engine/Source/Runtime/Engine/Private/Materials/Material.cpp
+++ b/Engine/Source/Runtime/Engine/Private/Materials/Material.cpp
@@ -2249,8 +2249,12 @@ void UMaterial::UpdateCachedExpressionData()
 
 	if (bLoadedCachedExpressionData)
 	{
-		// Always use Cached Data textures in Editor
-		ReferencedDefaultTextures = CachedExpressionData->EditorOnlyData->EditorReferencedDefaultTextures;
+		const FMaterialCachedExpressionData& UsedCachedData = GetCachedExpressionData();
+		// Try to always use Editor Only Data textures in Editor, if it exists
+		if (UsedCachedData.EditorOnlyData->EditorReferencedDefaultTextures.Num() != 0)
+		{
+			ReferencedDefaultTextures = UsedCachedData.EditorOnlyData->EditorReferencedDefaultTextures;
+		}
 
 		// Don't need to rebuild cached data if it was serialized
 		return;
diff --git a/Engine/Source/Runtime/Engine/Private/Materials/MaterialInterface.cpp b/Engine/Source/Runtime/Engine/Private/Materials/MaterialInterface.cpp
--- a/Engine/Source/Runtime/Engine/Private/Materials/MaterialInterface.cpp
+++ b/Engine/Source/Runtime/Engine/Private/Materials/MaterialInterface.cpp
@@ -327,8 +327,16 @@ bool UMaterialInterface::ContainsDefaultTexture(UObject* InTexture) const
 const TArray<TObjectPtr<UObject>>& UMaterialInterface::GetDefaultTextures() const
 {
 	const FMaterialCachedExpressionData& UsedCachedData = GetCachedExpressionData();
+
+	// Can happen inside a Material with bLoadedCachedExpressionData and no EditorOnlyData
+	if (UsedCachedData.EditorOnlyData->EditorReferencedDefaultTextures.Num() == 0)
+	{
+		return ReferencedDefaultTextures;
+	}
+
 	return UsedCachedData.EditorOnlyData->EditorReferencedDefaultTextures;
 }
+
 #endif // WITH_EDITORONLY_DATA
 
 bool UMaterialInterface::IsUsingNewHLSLGenerator() const
