Fix ensure triggered by cvar being read from async loading thread.
Now reads cvar from AnyThread.
#rnx
#Synced-CL 50764949

diff --git a/Engine/Source/Runtime/Engine/Private/VT/MeshPaintVirtualTexture.cpp b/Engine/Source/Runtime/Engine/Private/VT/MeshPaintVirtualTexture.cpp
--- a/Engine/Source/Runtime/Engine/Private/VT/MeshPaintVirtualTexture.cpp
+++ b/Engine/Source/Runtime/Engine/Private/VT/MeshPaintVirtualTexture.cpp
@@ -94,7 +94,7 @@ namespace MeshPaintVirtualTexture
 
 	static bool UseCompressedTextureFormat()
 	{
-		return CVarMeshPaintVirtualTextureUseCompression.GetValueOnGameThread();
+		return CVarMeshPaintVirtualTextureUseCompression.GetValueOnAnyThread();
 	}
 
 	static uint32 GetTileSize()
