RHI: Add the ability to disable specialization constants for a shader even if they are supported by the platforms.  Allows for the sharing of the same code between shaders with specialization constants enabled or disabled.

#rnx

#Synced-CL 50685202

diff --git a/Engine/Shaders/Public/Platform.ush b/Engine/Shaders/Public/Platform.ush
--- a/Engine/Shaders/Public/Platform.ush
+++ b/Engine/Shaders/Public/Platform.ush
@@ -1494,5 +1494,9 @@ uint WaveGetActiveLaneIndexLast()
 	#define UE_DECL(Name)
 #endif
 #ifndef UE_EVAL
-	#define UE_EVAL(Name) Name##_DEFAULT
+	#if ENABLE_SPECIALIZATION_CONSTANTS
+		#define UE_EVAL(Name) Name##_DEFAULT
+	#else
+		#define UE_EVAL(Name) Name
+	#endif
 #endif
\ No newline at end of file
diff --git a/Engine/Source/Runtime/Engine/Private/ShaderCompiler/ShaderCompiler.cpp b/Engine/Source/Runtime/Engine/Private/ShaderCompiler/ShaderCompiler.cpp
--- a/Engine/Source/Runtime/Engine/Private/ShaderCompiler/ShaderCompiler.cpp
+++ b/Engine/Source/Runtime/Engine/Private/ShaderCompiler/ShaderCompiler.cpp
@@ -4171,7 +4171,11 @@ void GlobalBeginCompileShader(
 	SET_SHADER_DEFINE(Input.Environment, PLATFORM_SUPPORTS_BARYCENTRICS_SEMANTIC, FDataDrivenShaderPlatformInfo::GetSupportsBarycentricsSemantic(EShaderPlatform(Target.Platform)) != ERHIFeatureSupport::Unsupported);
 	SET_SHADER_DEFINE(Input.Environment, PLATFORM_SUPPORTS_BINDLESS, FDataDrivenShaderPlatformInfo::GetSupportsBindless(ShaderPlatform));
 	SET_SHADER_DEFINE(Input.Environment, PLATFORM_SUPPORTS_RAYTRACING_CLUSTER_OPS, FDataDrivenShaderPlatformInfo::GetSupportsRayTracingClusterOps(ShaderPlatform));
+
+	// Differentiate between specialization constants being supported by the platform vs being enabled for a shader.
 	SET_SHADER_DEFINE(Input.Environment, PLATFORM_SUPPORTS_SPECIALIZATION_CONSTANTS, FDataDrivenShaderPlatformInfo::GetSupportsSpecializationConstants(ShaderPlatform));
+	const bool bEnableSpecializationConstants = Input.Environment.SpecializationConstantValues.Num() > 0;
+	SET_SHADER_DEFINE(Input.Environment, ENABLE_SPECIALIZATION_CONSTANTS, bEnableSpecializationConstants ? 1 : 0);
 
 	SET_SHADER_DEFINE(Input.Environment, UE_BINDLESS_ENABLED, Input.bBindlessEnabled);
 
