RemoteObjects:  Serialize non-UProperty AnimInstance fields when migrating objects.

#Synced-CL 50628995

diff --git a/Engine/Source/Runtime/Engine/Classes/Animation/AnimInstance.h b/Engine/Source/Runtime/Engine/Classes/Animation/AnimInstance.h
--- a/Engine/Source/Runtime/Engine/Classes/Animation/AnimInstance.h
+++ b/Engine/Source/Runtime/Engine/Classes/Animation/AnimInstance.h
@@ -406,7 +406,7 @@ class UAnimInstance : public UObject
 
 private:
 	// Should UpdateAnimation be called
-	uint8 bUpdateAnimationEnabled;
+	uint8 bUpdateAnimationEnabled : 1;
 
 	/** True when Montages are being ticked, and Montage Events should be queued. 
 	 * When Montage are being ticked, we queue AnimNotifies and Events. We trigger notifies first, then Montage events. */
diff --git a/Engine/Source/Runtime/Engine/Private/Animation/AnimInstance.cpp b/Engine/Source/Runtime/Engine/Private/Animation/AnimInstance.cpp
--- a/Engine/Source/Runtime/Engine/Private/Animation/AnimInstance.cpp
+++ b/Engine/Source/Runtime/Engine/Private/Animation/AnimInstance.cpp
@@ -1605,9 +1605,33 @@ void UAnimInstance::Serialize(FArchive& Ar)
 {
 	Super::Serialize(Ar);
 
-	if (!Ar.IsLoading() || !Ar.IsSaving())
+	// There's speculation we wanted 'if (!IsPersistent())' here, but changing this will require bumping the serialization version,
+	// so we're keeping this always-serialized until we find a reason not to.
+	// if (!Ar.IsLoading() || !Ar.IsSaving()) [This should always return true]
+	Ar << GetProxyOnAnyThread<FAnimInstanceProxy>().GetRequiredBones();
+
+	if (Ar.IsMigratingRemoteObjects())
 	{
-		Ar << GetProxyOnAnyThread<FAnimInstanceProxy>().GetRequiredBones();
+		uint8 BitFields = 0;
+		BitFields |= (bUpdateAnimationEnabled << 0);
+		BitFields |= (bUninitialized << 1);
+		BitFields |= (bHasDispatchedAnimGraphNotifiesThisTick << 2);
+		BitFields |= (bHasDispatchedMontageNotifiesThisTick << 3);
+
+		Ar << BitFields;
+
+		if (Ar.IsLoading())
+		{
+			bUpdateAnimationEnabled 				= (BitFields & (1 << 0)) ? 1 : 0;
+			bUninitialized 							= (BitFields & (1 << 1)) ? 1 : 0;
+			bHasDispatchedAnimGraphNotifiesThisTick = (BitFields & (1 << 2)) ? 1 : 0;
+			bHasDispatchedMontageNotifiesThisTick 	= (BitFields & (1 << 3)) ? 1 : 0;
+		}
+
+		// These are guard flags, so don't serialize them
+		//Ar << bInitializing;
+		//Ar << bPostUpdatingAnimation;
+		//Ar << bUpdatingAnimation;
 	}
 }
 
