[LevelStreaming] Force a FlushAsyncLoading call before calling World->AddToWorld in ELevelStreamingState::MakingVisible to prevent a crash in which we attempt to add entities to the world before they have been fully loaded. This is intended as a temporary workaround to unblock certain impacted flows (CORO-8868)

#rnx

#Synced-CL 50788317

diff --git a/Engine/Source/Runtime/Engine/Private/LevelStreaming.cpp b/Engine/Source/Runtime/Engine/Private/LevelStreaming.cpp
--- a/Engine/Source/Runtime/Engine/Private/LevelStreaming.cpp
+++ b/Engine/Source/Runtime/Engine/Private/LevelStreaming.cpp
@@ -98,6 +98,14 @@ namespace LevelStreamingCVars
 		TEXT("PriorityBias added to level streaming priority when making async load requests"),
 		ECVF_Default
 	);
+
+	static bool bForceFlushLoadsBeforeAddToWorld = false;
+	FAutoConsoleVariableRef CVarForceFlushLoadsBeforeAddToWorld(
+		TEXT("LevelStreaming.ForceFlushLoadsBeforeAddToWorld"),
+		bForceFlushLoadsBeforeAddToWorld,
+		TEXT("Force a full FlushAsyncLoads() call before AddToWorld in MakingVisible.\n")
+		TEXT("0: Disable, 1: Enable"),
+		ECVF_Default);
 }
 
 bool ULevelStreaming::DefaultAllowClientUseMakingInvisibleTransactionRequests()
@@ -1059,6 +1067,12 @@ void ULevelStreaming::UpdateStreamingState(bool& bOutUpdateAgain, bool& bOutRede
 					break;
 				}
 
+				if (LevelStreamingCVars::bForceFlushLoadsBeforeAddToWorld)
+				{
+					// This is intended as a temporary workaround (CORO-8867). A proper fix is being tracked via CORO-8868.
+					FlushAsyncLoading();
+				}
+
 				World->AddToWorld(LoadedLevel, LevelTransform, !bShouldBlockOnLoad, InExternalTimeout, TransactionId, this);
 
 				if (LoadedLevel->bIsVisible)
