New Material Translator: Fixed error reporting when a non MaterialAttributes aggregate value flows through the MaterialAttributes material output pin.

#rnx

#Synced-CL 50775131

diff --git a/Engine/Source/Runtime/Engine/Private/Materials/MaterialIRModuleBuilder.cpp b/Engine/Source/Runtime/Engine/Private/Materials/MaterialIRModuleBuilder.cpp
--- a/Engine/Source/Runtime/Engine/Private/Materials/MaterialIRModuleBuilder.cpp
+++ b/Engine/Source/Runtime/Engine/Private/Materials/MaterialIRModuleBuilder.cpp
@@ -688,7 +688,12 @@ struct FMaterialIRModuleBuilderImpl
 			}
 
 			// Make sure a valid value is present and it is of the correct type.
-			check(!MaterialAttributesValue->IsPoison() && MaterialAttributesValue->Type.AsAggregate() == MaterialAttributesAggregate::Get());
+			check(!MaterialAttributesValue->IsPoison());
+			
+			if (!MaterialAttributesValue->Type.AsAggregate() || MaterialAttributesValue->Type.AsAggregate() != MaterialAttributesAggregate::Get())
+			{
+				Module->AddError(nullptr, FString::Printf(TEXT("Expression connected to the MaterialAttributes material output does is not a MaterialAttributes value (it is a '%s' instead)."), *MaterialAttributesValue->Type.GetSpelling()));
+			}
 		}
 
 		for (EMaterialProperty Property : MaterialAttributesAggregate::GetMaterialProperties())
