MaterialFunction reports an error when its editor only data hasn't been loaded, for instance becasue the editor is cooking a cooked function.

#rnx

#Synced-CL 50775131

diff --git a/Engine/Source/Runtime/Engine/Private/Materials/MaterialExpressions.cpp b/Engine/Source/Runtime/Engine/Private/Materials/MaterialExpressions.cpp
--- a/Engine/Source/Runtime/Engine/Private/Materials/MaterialExpressions.cpp
+++ b/Engine/Source/Runtime/Engine/Private/Materials/MaterialExpressions.cpp
@@ -13967,6 +13967,11 @@ void UMaterialFunction::PostLoad()
 			}
 		}
 	}
+
+	if (!EditorOnly || (GetPackage()->HasAnyFlags(EObjectFlags::RF_WasLoaded) && !EditorOnly->HasAnyFlags(EObjectFlags::RF_WasLoaded)))
+	{
+		UE_LOG(LogMaterial, Error, TEXT("Material function '%s' was loaded, but it is missing editor-only data. Materials calling this function will not translate correctly."), *GetFullName());
+	}
 #endif // WITH_EDITORONLY_DATA
 
 #if WITH_EDITOR
