[WorldPartition] Deprecation cleanup

#Synced-CL 50613390

diff --git a/Engine/Source/Runtime/Engine/Private/WorldPartition/WorldPartition.cpp b/Engine/Source/Runtime/Engine/Private/WorldPartition/WorldPartition.cpp
--- a/Engine/Source/Runtime/Engine/Private/WorldPartition/WorldPartition.cpp
+++ b/Engine/Source/Runtime/Engine/Private/WorldPartition/WorldPartition.cpp
@@ -2184,16 +2184,6 @@ void UWorldPartition::DumpActorDescs(const FString& Path)
 
 TOptional<UWorldPartition::FStreamingGenerationErrorHandlerOverride> UWorldPartition::StreamingGenerationErrorHandlerOverride;
 
-void UWorldPartition::AppendAssetRegistryTags(TArray<FAssetRegistryTag>& OutTags) const
-{
-	FAssetRegistryTagsContextData Context(this, EAssetRegistryTagsCaller::Uncategorized);
-	AppendAssetRegistryTags(Context);
-	for (TPair<FName, FAssetRegistryTag>& Pair : Context.Tags)
-	{
-		OutTags.Add(MoveTemp(Pair.Value));
-	}
-}
-
 void UWorldPartition::AppendAssetRegistryTags(FAssetRegistryTagsContext Context) const
 {
 	static const FName NAME_LevelIsPartitioned("LevelIsPartitioned");
diff --git a/Engine/Source/Runtime/Engine/Private/WorldPartition/WorldPartitionHelpers.cpp b/Engine/Source/Runtime/Engine/Private/WorldPartition/WorldPartitionHelpers.cpp
--- a/Engine/Source/Runtime/Engine/Private/WorldPartition/WorldPartitionHelpers.cpp
+++ b/Engine/Source/Runtime/Engine/Private/WorldPartition/WorldPartitionHelpers.cpp
@@ -324,18 +324,6 @@ bool FWorldPartitionHelpers::ConvertRuntimePathToEditorPath(const FSoftObjectPat
 	return false;
 }
 
-bool FWorldPartitionHelpers::FixupRedirectedAssetPath(FSoftObjectPath& InOutSoftObjectPath)
-{
-	UAssetRegistryHelpers::FixupRedirectedAssetPath(InOutSoftObjectPath);
-	return true;
-}
-
-bool FWorldPartitionHelpers::FixupRedirectedAssetPath(FName& InOutAssetPath)
-{
-	UAssetRegistryHelpers::FixupRedirectedAssetPath(InOutAssetPath);
-	return true;
-}
-
 TMap<FGuid, AActor*> FWorldPartitionHelpers::GetLoadedActorsForLevel(const ULevel* InLevel)
 {
 	TMap<FGuid, AActor*> Result;
diff --git a/Engine/Source/Runtime/Engine/Private/WorldPartition/WorldPartitionLevelHelper.cpp b/Engine/Source/Runtime/Engine/Private/WorldPartition/WorldPartitionLevelHelper.cpp
--- a/Engine/Source/Runtime/Engine/Private/WorldPartition/WorldPartitionLevelHelper.cpp
+++ b/Engine/Source/Runtime/Engine/Private/WorldPartition/WorldPartitionLevelHelper.cpp
@@ -565,20 +565,6 @@ ULevel* FWorldPartitionLevelHelper::CreateEmptyLevelForRuntimeCell(const UWorldP
 	return NewLevel;
 }
 
-bool FWorldPartitionLevelHelper::LoadActors(UWorld* InOuterWorld, ULevel* InDestLevel, TArrayView<FWorldPartitionRuntimeCellObjectMapping> InActorPackages, FWorldPartitionLevelHelper::FPackageReferencer& InPackageReferencer, TFunction<void(bool)> InCompletionCallback, bool bInLoadAsync, FLinkerInstancingContext InInstancingContext)
-{
-	FLoadActorsParams Params = FLoadActorsParams()
-		.SetOuterWorld(InOuterWorld)
-		.SetDestLevel(InDestLevel)
-		.SetActorPackages(InActorPackages)
-		.SetPackageReferencer(&InPackageReferencer)
-		.SetCompletionCallback(InCompletionCallback)
-		.SetLoadAsync(bInLoadAsync)
-		.SetInstancingContext(InInstancingContext);
-
-	return LoadActors(MoveTemp(Params));
-}
-
 bool FWorldPartitionLevelHelper::LoadActorsWithPropertyOverridesInternal(FLoadActorsParams&& InParams)
 {
 	TMap<FString, FName> PropertyOverridesToLoad;
diff --git a/Engine/Source/Runtime/Engine/Public/WorldPartition/ErrorHandling/WorldPartitionStreamingGenerationErrorHandler.h b/Engine/Source/Runtime/Engine/Public/WorldPartition/ErrorHandling/WorldPartitionStreamingGenerationErrorHandler.h
--- a/Engine/Source/Runtime/Engine/Public/WorldPartition/ErrorHandling/WorldPartitionStreamingGenerationErrorHandler.h
+++ b/Engine/Source/Runtime/Engine/Public/WorldPartition/ErrorHandling/WorldPartitionStreamingGenerationErrorHandler.h
@@ -143,60 +143,5 @@ public:
 	// Helpers
 	static ENGINE_API FString GetActorName(const IWorldPartitionActorDescInstanceView& ActorDescView);
 	static ENGINE_API FString GetFullActorName(const IWorldPartitionActorDescInstanceView& ActorDescView);
-
-	PRAGMA_DISABLE_DEPRECATION_WARNINGS
-	
-	UE_DEPRECATED(5.4, "Use OnInvalidReferenceDataLayers with EDataLayerInvalidReason instead.")
-	virtual void OnInvalidReferenceDataLayers(const IWorldPartitionActorDescInstanceView& ActorDescView, const IWorldPartitionActorDescInstanceView& ReferenceActorDescView) final {}
-
-	UE_DEPRECATED(5.2, "OnInvalidReference is deprecated, use the version which takes an optional actor descriptor view.")
-	virtual void OnInvalidReference(const FWorldPartitionActorDescView& ActorDescView, const FGuid& ReferenceGuid) {}
-
-	UE_DEPRECATED(5.4, "Use IWorldPartitionActorDescInstanceView version instead")
-	virtual void OnInvalidRuntimeGrid(const FWorldPartitionActorDescView& ActorDescView, FName GridName) {}
-
-	UE_DEPRECATED(5.4, "Use IWorldPartitionActorDescInstanceView version instead")
-	virtual void OnInvalidReference(const FWorldPartitionActorDescView& ActorDescView, const FGuid& ReferenceGuid, FWorldPartitionActorDescView* ReferenceActorDescView) {}
-
-	UE_DEPRECATED(5.4, "Use IWorldPartitionActorDescInstanceView version instead")
-	virtual void OnInvalidReferenceGridPlacement(const FWorldPartitionActorDescView& ActorDescView, const FWorldPartitionActorDescView& ReferenceActorDescView) {}
-
-	UE_DEPRECATED(5.4, "Use IWorldPartitionActorDescInstanceView version instead")
-	virtual void OnInvalidReferenceDataLayers(const FWorldPartitionActorDescView& ActorDescView, const FWorldPartitionActorDescView& ReferenceActorDescView) {}
-
-	UE_DEPRECATED(5.4, "Use IWorldPartitionActorDescInstanceView version instead")
-	virtual void OnInvalidReferenceRuntimeGrid(const FWorldPartitionActorDescView& ActorDescView, const FWorldPartitionActorDescView& ReferenceActorDescView) {}
-
-	UE_DEPRECATED(5.4, "Use IWorldPartitionActorDescInstanceView version instead")
-	virtual void OnInvalidReferenceLevelScriptStreamed(const FWorldPartitionActorDescView& ActorDescView) {}
-	
-	UE_DEPRECATED(5.4, "Use IWorldPartitionActorDescInstanceView version instead")
-	virtual void OnInvalidReferenceLevelScriptDataLayers(const FWorldPartitionActorDescView& ActorDescView) {}
-	
-	UE_DEPRECATED(5.4, "Use IWorldPartitionActorDescInstanceView version instead")
-	virtual void OnActorNeedsResave(const FWorldPartitionActorDescView& ActorDescView) {}
-
-	UE_DEPRECATED(5.4, "Use IWorldPartitionActorDescInstanceView version instead")
-	virtual void OnLevelInstanceInvalidWorldAsset(const FWorldPartitionActorDescView& ActorDescView, FName WorldAsset, ELevelInstanceInvalidReason Reason) {}
-
-	UE_DEPRECATED(5.4, "Use IWorldPartitionActorDescInstanceView version instead")
-	virtual void OnInvalidActorFilterReference(const FWorldPartitionActorDescView& ActorDescView, const FWorldPartitionActorDescView& ReferenceActorDescView) {}
-	
-	UE_DEPRECATED(5.4, "Use IWorldPartitionActorDescInstanceView version instead")
-	virtual void OnInvalidHLODLayer(const FWorldPartitionActorDescView& ActorDescView) {}
-
-	UE_DEPRECATED(5.4, "Use OnInvalidWorldReference instead")
-	void OnInvalidReferenceLevelScriptStreamed(const IWorldPartitionActorDescInstanceView& ActorDescView) { OnInvalidWorldReference(ActorDescView, EWorldReferenceInvalidReason::ReferencedActorIsSpatiallyLoaded); }
-
-	UE_DEPRECATED(5.4, "Use OnInvalidWorldReference instead")
-	void OnInvalidReferenceLevelScriptDataLayers(const IWorldPartitionActorDescInstanceView& ActorDescView) { OnInvalidWorldReference(ActorDescView, EWorldReferenceInvalidReason::ReferencedActorHasDataLayers); }
-
-	UE_DEPRECATED(5.4, "Use IWorldPartitionActorDescInstanceView version instead")
-	static FString GetActorName(const FWorldPartitionActorDescView& ActorDescView) { return FString(); }
-
-	UE_DEPRECATED(5.4, "Use IWorldPartitionActorDescInstanceView version instead")
-	static FString GetFullActorName(const FWorldPartitionActorDescView& ActorDescView) { return FString(); }
-
-	PRAGMA_ENABLE_DEPRECATION_WARNINGS
 };
 #endif
diff --git a/Engine/Source/Runtime/Engine/Public/WorldPartition/WorldPartition.h b/Engine/Source/Runtime/Engine/Public/WorldPartition/WorldPartition.h
--- a/Engine/Source/Runtime/Engine/Public/WorldPartition/WorldPartition.h
+++ b/Engine/Source/Runtime/Engine/Public/WorldPartition/WorldPartition.h
@@ -173,11 +173,6 @@ public:
 public:
 	ENGINE_API void SetEnableStreaming(bool bInEnableStreaming);
 
-	UE_DEPRECATED(5.3, "CanBeUsedByLevelInstance is deprecated.")
-	bool CanBeUsedByLevelInstance() const { return true; }
-	UE_DEPRECATED(5.3, "SetCanBeUsedByLevelInstance is deprecated.")
-	void SetCanBeUsedByLevelInstance(bool bInCanBeUsedByLevelInstance) {}
-
 	ENGINE_API void OnEnableStreamingChanged();
 	ENGINE_API void OnEnableLoadingInEditorChanged();
 
@@ -249,13 +244,6 @@ public:
 	// Streaming generation
 	bool CanGenerateStreaming() const { return !StreamingPolicy; }
 
-	UE_DEPRECATED(5.3, "GenerateStreaming is deprecated, use GenerateStreaming with a param struct instead")
-	bool GenerateStreaming(TArray<FString>* OutPackagesToGenerate = nullptr) { return false; }
-
-	UE_DEPRECATED(5.3, "GenerateContainerStreaming is deprecated, use GenerateContainerStreaming with a param struct instead")
-	bool GenerateContainerStreaming(const UActorDescContainer* ActorDescContainer, TArray<FString>* OutPackagesToGenerate = nullptr) { return false; }
-
-	PRAGMA_DISABLE_DEPRECATION_WARNINGS
 	struct FGenerateStreamingParams
 	{
 		FGenerateStreamingParams() 
@@ -272,12 +260,6 @@ public:
 		FGenerateStreamingParams& SetErrorHandler(IStreamingGenerationErrorHandler* InErrorHandler) { ErrorHandler = InErrorHandler; return *this; }
 		FGenerateStreamingParams& SetOutputLogType(const FString& InOutputLogType) { OutputLogType = InOutputLogType; return *this; }
 		FGenerateStreamingParams& SetFilteredClasses(const TArray<TSubclassOf<AActor>>& InFilteredClasses) { FilteredClasses = InFilteredClasses; return *this; }
-
-		UE_DEPRECATED(5.4, "Use constructor receiving a ContainerInstanceCollection instead")
-		FGenerateStreamingParams& SetActorDescContainer(const UActorDescContainer* InActorDescContainer) { return *this; }
-
-		UE_DEPRECATED(5.4, "Use ContainerInstanceCollection instead")
-		FStreamingGenerationActorDescCollection ActorDescCollection;
 		
 	private:
 
@@ -288,7 +270,6 @@ public:
 
 		friend class UWorldPartition;
 	};
-	PRAGMA_ENABLE_DEPRECATION_WARNINGS
 
 	struct FGenerateStreamingContext
 	{
@@ -303,8 +284,6 @@ public:
 
 		 TOptional<FString> OutputLogFilename;
 
-		UE_DEPRECATED(5.4, "SetPackagesToGenerate is deprecated, use SetLevelPackagesToGenerate")
-		FGenerateStreamingContext& SetPackagesToGenerate(TArray<FString>* InPackagesToGenerate) { return SetLevelPackagesToGenerate(InPackagesToGenerate); }
 		FGenerateStreamingContext& SetLevelPackagesToGenerate(TArray<FString>* InLevelPackagesToGenerate) { PackagesToGenerate = InLevelPackagesToGenerate; return *this; }
 		FGenerateStreamingContext& SetGeneratedExternalStreamingObjects(TArray<URuntimeHashExternalStreamingObjectBase*>* InGeneratedExternalStreamingObjects) { GeneratedExternalStreamingObjects = InGeneratedExternalStreamingObjects; return *this; }
 	};
@@ -316,9 +295,6 @@ public:
 	ENGINE_API void FlushStreaming();
 	ENGINE_API URuntimeHashExternalStreamingObjectBase* FlushStreamingToExternalStreamingObject();
 
-	UE_DEPRECATED(5.5, "FlushStreamingToExternalStreamingObject(const FString& ) is deprecated, use FlushStreamingToExternalStreamingObject() instead")
-	URuntimeHashExternalStreamingObjectBase* FlushStreamingToExternalStreamingObject(const FString& ExternalStreamingObjectName) { return FlushStreamingToExternalStreamingObject(); }
-
 	// Event when world partition was enabled/disabled in the world
 	DECLARE_MULTICAST_DELEGATE_OneParam(FWorldPartitionChangedEvent, UWorld*);
 	static ENGINE_API FWorldPartitionChangedEvent WorldPartitionChangedEvent;
@@ -398,7 +374,6 @@ public:
 	ENGINE_API void CheckForErrors(IStreamingGenerationErrorHandler* ErrorHandler) const;
 
 	/* Struct of optional parameters passed to check for errors function. */
-	PRAGMA_DISABLE_DEPRECATION_WARNINGS
 	struct FCheckForErrorsParams
 	{
 		ENGINE_API FCheckForErrorsParams();
@@ -414,22 +389,7 @@ public:
 		FCheckForErrorsParams& SetActorDescContainerInstanceCollection(const FActorDescContainerInstanceCollection* InActorDescContainerInstanceCollection) { ActorDescContainerInstanceCollection = InActorDescContainerInstanceCollection; return *this; }
 		FCheckForErrorsParams& SetEnableStreaming(bool bInEnableStreaming) { bEnableStreaming = bInEnableStreaming; return *this; }
 		FCheckForErrorsParams& SetActorGuidsToContainerInstanceMap(const TMap<FGuid, const UActorDescContainerInstance*>& InActorGuidsToContainerInstanceMap) { ActorGuidsToContainerInstanceMap = InActorGuidsToContainerInstanceMap; return *this; }
-
-		UE_DEPRECATED(5.4, "Use ActorDescContainerInstanceCollection instead")
-		const FActorDescContainerCollection* ActorDescContainerCollection = nullptr;
-				
-		UE_DEPRECATED(5.4, "Use ActorGuidsToContainerInstanceMap instead")
-		TMap<FGuid, const UActorDescContainer*> ActorGuidsToContainerMap;
-
-		UE_DEPRECATED(5.4, "Use SetActorDescContainerInstanceCollection instead")
-		FCheckForErrorsParams& SetActorDescContainerCollection(const FActorDescContainerCollection* InActorDescContainerCollection) { return *this; }
-		UE_DEPRECATED(5.4, "Use SetActorGuidsToContainerInstanceMap instead")
-		FCheckForErrorsParams& SetActorGuidsToContainerMap(const TMap<FGuid, const UActorDescContainer*>& InActorGuidsToContainerMap) { return *this; }
 	};
-	PRAGMA_ENABLE_DEPRECATION_WARNINGS
-
-	UE_DEPRECATED(5.2, "CheckForErrors is deprecated, CheckForErrors with FCheckForErrorsParams should be used instead.")
-	static void CheckForErrors(IStreamingGenerationErrorHandler* ErrorHandler, const UActorDescContainer* ActorDescContainer, bool bEnableStreaming, bool bIsChangelistValidation) {}
 
 	static ENGINE_API void CheckForErrors(const FCheckForErrorsParams& Params);
 
@@ -437,8 +397,6 @@ public:
 	ENGINE_API static TOptional<FStreamingGenerationErrorHandlerOverride> StreamingGenerationErrorHandlerOverride;
 
 	ENGINE_API void AppendAssetRegistryTags(FAssetRegistryTagsContext Context) const;
-	UE_DEPRECATED(5.4, "Implement the version that takes FAssetRegistryTagsContext instead.")
-	ENGINE_API void AppendAssetRegistryTags(TArray<FAssetRegistryTag>& OutTags) const;
 
 	struct FContainerRegistrationParams
 	{
@@ -457,23 +415,9 @@ public:
 	DECLARE_DELEGATE_TwoParams(FActorDescContainerInstancePreInitializeDelegate, UActorDescContainerInstance::FInitializeParams&, UActorDescContainerInstance*);
 	FActorDescContainerInstancePreInitializeDelegate OnActorDescContainerInstancePreInitialize;
 
-	UE_DEPRECATED(5.4, "Use RegisterActorDescContainerInstance")
-	UActorDescContainer* RegisterActorDescContainer(const FContainerRegistrationParams& InRegistrationParameters) { return nullptr; }
-
-	UE_DEPRECATED(5.4, "Use UnregisterActorDescContainerInstance")
-	bool UnregisterActorDescContainer(UActorDescContainer* Container) { return false; }
-
 	ENGINE_API void UninitializeActorDescContainers();
 
 	DECLARE_MULTICAST_DELEGATE_OneParam(FActorDescContainerRegistrationDelegate, UActorDescContainer*);
-
-	UE_DEPRECATED(5.4, "Use OnActorDescContainerInstanceRegistered")
-	FActorDescContainerRegistrationDelegate OnActorDescContainerRegistered;
-	UE_DEPRECATED(5.4, "Use OnActorDescContainerInstanceUnregistered")
-	FActorDescContainerRegistrationDelegate OnActorDescContainerUnregistered;
-
-	UE_DEPRECATED(5.3, "Use RegisterActorDescContainer with FContainerRegistrationParams instead.")
-	UActorDescContainer* RegisterActorDescContainer(const FName& ContainerPackage) { return nullptr; }
 		
 	ENGINE_API UActorDescContainerInstance* RegisterActorDescContainerInstance(const UActorDescContainerInstance::FInitializeParams& InInitializationParams);
 	ENGINE_API bool UnregisterActorDescContainerInstance(UActorDescContainerInstance* InContainerInstance);
@@ -518,8 +462,6 @@ public:
 	ENGINE_API bool CanInitialize(UWorld* InWorld) const;
 	ENGINE_API void Initialize(UWorld* World, const FTransform& InTransform);
 	ENGINE_API bool IsInitialized() const;
-	UE_DEPRECATED(5.4, "UWorldPartition::Update is deprecated.")
-	ENGINE_API void Update();
 	ENGINE_API void Uninitialize();
 
 	ENGINE_API bool SupportsStreaming() const;
@@ -565,11 +507,6 @@ public:
 
 	inline EWorldPartitionDataLayersLogicOperator GetDataLayersLogicOperator() const { return DataLayersLogicOperator; }
 
-	UE_DEPRECATED(5.3, "UpdateStreamingState is deprecated, use UWorldPartitionSubsystem::UpdateStreamingState instead.")
-	void UpdateStreamingState() {}
-	UE_DEPRECATED(5.3, "CanAddLoadedLevelToWorld is deprecated, use CanAddCellToWorld instead.")
-	bool CanAddLoadedLevelToWorld(ULevel* InLevel) const { return true; }
-
 #if WITH_EDITORONLY_DATA
 	UPROPERTY(DuplicateTransient)
 	TObjectPtr<UWorldPartitionEditorHash> EditorHash;
@@ -602,9 +539,6 @@ private:
 	void OnCleanupLevel();
 
 public:
-	UE_DEPRECATED(5.4, "Use GetActorDescContainerInstance")
-	UActorDescContainer* GetActorDescContainer() const { return nullptr; }
-
 #if WITH_EDITOR
 	UActorDescContainerInstance* GetActorDescContainerInstance() const { return ActorDescContainerInstance; }
 
@@ -613,15 +547,6 @@ public:
 		check(!IsInitialized());
 		ContainerInstanceClass = InContainerInstanceClass;
 	}
-
-	UE_DEPRECATED(5.4, "Use ForEachActorDescContainerInstanceBreakable.")
-	void ForEachActorDescContainerBreakable(TFunctionRef<bool(UActorDescContainer*)> Func) {}
-	UE_DEPRECATED(5.4, "Use ForEachActorDescContainerInstanceBreakable.")
-	void ForEachActorDescContainerBreakable(TFunctionRef<bool(UActorDescContainer*)> Func) const {}
-	UE_DEPRECATED(5.4, "Use ForEachActorDescContainerInstance.")
-	void ForEachActorDescContainer(TFunctionRef<void(UActorDescContainer*)> Func) {}
-	UE_DEPRECATED(5.4, "Use ForEachActorDescContainerInstance.")
-	void ForEachActorDescContainer(TFunctionRef<void(UActorDescContainer*)> Func) const {}
 #endif			
 
 	UPROPERTY()
diff --git a/Engine/Source/Runtime/Engine/Public/WorldPartition/WorldPartitionActorLoaderInterface.h b/Engine/Source/Runtime/Engine/Public/WorldPartition/WorldPartitionActorLoaderInterface.h
--- a/Engine/Source/Runtime/Engine/Public/WorldPartition/WorldPartitionActorLoaderInterface.h
+++ b/Engine/Source/Runtime/Engine/Public/WorldPartition/WorldPartitionActorLoaderInterface.h
@@ -46,11 +46,6 @@ public:
 		virtual TOptional<FString> GetLabel() const { return TOptional<FString>(); }
 		virtual TOptional<FColor> GetColor() const { return TOptional<FColor>(); }
 
-		UE_DEPRECATED(5.4, "Use OnActorDescContainerInstanceInitialize")
-		void OnActorDescContainerInitialize(UActorDescContainer* Container) {}
-		UE_DEPRECATED(5.4, "Use OnActorDescContainerInstanceUninitialize")
-		void OnActorDescContainerUninitialize(UActorDescContainer* Container) {}
-
 	protected:
 		void OnActorDescContainerInstanceInitialize(UActorDescContainerInstance* ContainerInstance);
 		void OnActorDescContainerInstanceUninitialize(UActorDescContainerInstance* ContainerInstance);
diff --git a/Engine/Source/Runtime/Engine/Public/WorldPartition/WorldPartitionHandle.h b/Engine/Source/Runtime/Engine/Public/WorldPartition/WorldPartitionHandle.h
--- a/Engine/Source/Runtime/Engine/Public/WorldPartition/WorldPartitionHandle.h
+++ b/Engine/Source/Runtime/Engine/Public/WorldPartition/WorldPartitionHandle.h
@@ -94,11 +94,6 @@ public:
 		virtual void UnregisterActor(FWorldPartitionActorDescInstance* InActorDescInstance) override {}
 	};
 
-	UE_DEPRECATED(5.4, "Use FWorldPartitionActorDescInstance version instead")
-	static void LoadAndRegisterActor(FWorldPartitionActorDesc* ActorDesc) {}
-	UE_DEPRECATED(5.4, "Use FWorldPartitionActorDescInstance version instead")
-	static void UnloadAndUnregisterActor(FWorldPartitionActorDesc* ActorDesc) {}
-
 private:
 	static void LoadAndRegisterActor(FWorldPartitionActorDescInstance* ActorDescInstance);
 	static void UnloadAndUnregisterActor(FWorldPartitionActorDescInstance* ActorDescInstance);
@@ -166,35 +161,6 @@ public:
 			IncRefCount();
 		}
 	}
-
-	UE_DEPRECATED(5.4, "Use FWorldPartitionActorDescInstance version")
-	inline TWorldPartitionHandle(TUniquePtr<FWorldPartitionActorDesc>* InActorDesc)
-		: ContainerInstance(nullptr)
-		, ActorDescInstance(nullptr)
-	{
-	}
-
-	UE_DEPRECATED(5.4, "Use FActorDescContainerInstance version")
-	inline TWorldPartitionHandle(UActorDescContainer* InContainer, const FGuid& ActorGuid)
-		: ContainerInstance(nullptr)
-		, ActorDescInstance(nullptr)
-	{
-	}
-		
-	template<class U>
-	struct TActorDescContainerCollectionDeprecated
-	{
-		UE_DEPRECATED(5.4, "Use TActorDescContainerInstanceCollection version")
-		static void Deprecated() {}
-	};
-
-	template<class U>
-	inline TWorldPartitionHandle(TActorDescContainerCollection<U>* ContainerCollection, const FGuid& ActorGuid)
-		: ContainerInstance(nullptr)
-		, ActorDescInstance(nullptr)
-	{
-		TActorDescContainerCollectionDeprecated<U>::Deprecated();
-	}
 	
 	inline TWorldPartitionHandle(const TWorldPartitionHandle& Other)
 		: ContainerInstance(nullptr)
@@ -348,17 +314,6 @@ public:
 		return IsValid() ? Impl::GetActor(ActorDescInstance->Get()) : nullptr;
 	}
 
-	UE_DEPRECATED(5.4, "Use GetInstance instead")
-	inline FWorldPartitionActorDesc* Get() const
-	{
-		if (ContainerInstance.IsValid() && ActorDescInstance && ActorDescInstance->IsValid())
-		{
-			return Impl::GetActorDesc(ActorDescInstance->Get());
-		}
-
-		return nullptr;
-	}
-
 	inline void Reset()
 	{
 		if (IsValid())
@@ -415,15 +370,6 @@ public:
 
 struct FWorldPartitionImplBase
 {
-	UE_DEPRECATED(5.4, "Use FWorldPartitionActorDescInstance")
-	static TUniquePtr<FWorldPartitionActorDesc>* GetActorDesc(UActorDescContainer* Container, const FGuid& ActorGuid) { return nullptr;}
-	UE_DEPRECATED(5.4, "Use FWorldPartitionActorDescInstance")
-	static UActorDescContainer* GetActorDescContainer(TUniquePtr<FWorldPartitionActorDesc>* ActorDesc) { return nullptr; }
-	UE_DEPRECATED(5.4, "Use FWorldPartitionActorDescInstance")
-	static bool IsActorDescLoaded(FWorldPartitionActorDesc* ActorDesc) { return false; }
-	UE_DEPRECATED(5.4, "Use FWorldPartitionActorDescInstance")
-	static AActor* GetActorDescActor(FWorldPartitionActorDesc* ActorDesc) { return nullptr; }
-	
 	template<class U>
 	static UActorDescContainer* GetActorDescContainer(TActorDescContainerCollection<U>* ContainerCollection, const FGuid& ActorGuid)
 	{
@@ -448,11 +394,6 @@ struct FWorldPartitionImplBase
 
 struct FWorldPartitionHandleImpl : FWorldPartitionImplBase
 {
-	UE_DEPRECATED(5.4, "Use FWorldPartitionActorDescInstance")
-	static void IncRefCount(FWorldPartitionActorDesc* ActorDesc) {}
-	UE_DEPRECATED(5.4, "Use FWorldPartitionActorDescInstance")
-	static void DecRefCount(FWorldPartitionActorDesc* ActorDesc) {}
-
 	static ENGINE_API void IncRefCount(FWorldPartitionActorDescInstance* InActorDescInstance);
 	static ENGINE_API void DecRefCount(FWorldPartitionActorDescInstance* InActorDescInstance);
 
@@ -461,11 +402,6 @@ struct FWorldPartitionHandleImpl : FWorldPartitionImplBase
 
 struct FWorldPartitionReferenceImpl : FWorldPartitionImplBase
 {
-	UE_DEPRECATED(5.4, "Use FWorldPartitionActorDescInstance")
-	static void IncRefCount(FWorldPartitionActorDesc* ActorDesc) {}
-	UE_DEPRECATED(5.4, "Use FWorldPartitionActorDescInstance")
-	static void DecRefCount(FWorldPartitionActorDesc* ActorDesc) {}
-
 	static ENGINE_API void IncRefCount(FWorldPartitionActorDescInstance* InActorDescInstance);
 	static ENGINE_API void DecRefCount(FWorldPartitionActorDescInstance* InActorDescInstance);
 
diff --git a/Engine/Source/Runtime/Engine/Public/WorldPartition/WorldPartitionHelpers.h b/Engine/Source/Runtime/Engine/Public/WorldPartition/WorldPartitionHelpers.h
--- a/Engine/Source/Runtime/Engine/Public/WorldPartition/WorldPartitionHelpers.h
+++ b/Engine/Source/Runtime/Engine/Public/WorldPartition/WorldPartitionHelpers.h
@@ -77,34 +77,7 @@ public:
 private:
 	static ENGINE_API bool IsActorDescClassCompatibleWith(const FWorldPartitionActorDesc* ActorDesc, const UClass* Class);
 
-	template<class ActorClass>
-	struct TDeprecated
-	{
-		UE_DEPRECATED(5.4, "Use ForEachIntersectingActorDescInstance instead")
-		static void DeprecatedForEachIntersectingActorDesc() {}
-
-		UE_DEPRECATED(5.4, "Use ForEachActorDescInstance instead")
-		static void DeprecatedForEachActorDesc() {}
-	};
-
 public:
-	template <class ActorClass = AActor>
-	static void ForEachIntersectingActorDesc(UWorldPartition* WorldPartition, const FBox& Box, TFunctionRef<bool(const FWorldPartitionActorDesc*)> Func)
-	{
-		TDeprecated<ActorClass>::DeprecatedForEachIntersectingActorDesc();
-	}
-
-	template<class ActorClass = AActor>
-	static void ForEachActorDesc(UWorldPartition* WorldPartition, TFunctionRef<bool(const FWorldPartitionActorDesc*)> Func)
-	{
-		TDeprecated<ActorClass>::DeprecatedForEachActorDesc();
-	}
-
-	UE_DEPRECATED(5.4, "Use ForEachIntersectingActorDescInstance")
-	static void ForEachIntersectingActorDesc(UWorldPartition* WorldPartition, const FBox& Box, TSubclassOf<AActor> ActorClass, TFunctionRef<bool(const FWorldPartitionActorDesc*)> Func) {}
-	UE_DEPRECATED(5.4, "Use ForEachActorDescInstance")
-	static void ForEachActorDesc(UWorldPartition* WorldPartition, TSubclassOf<AActor> ActorClass, TFunctionRef<bool(const FWorldPartitionActorDesc*)> Func) {}
-
 
 	template <class ActorClass = AActor>
 	static void ForEachIntersectingActorDescInstance(UWorldPartition* WorldPartition, const FBox& Box, TFunctionRef<bool(const FWorldPartitionActorDescInstance*)> Func)
@@ -154,17 +127,10 @@ public:
 		/* Reference to all actors and actor references loaded by the foreach actordesc function */
 		TMap<FGuid, FWorldPartitionReference> ActorReferences;
 	};
-
-	UE_DEPRECATED(5.4, "Use ForEachActorWithLoading with FWorldPartitionActorDescInstance")
-	static void ForEachActorWithLoading(UWorldPartition* WorldPartition, TFunctionRef<bool(const FWorldPartitionActorDesc*)> Func, const FForEachActorWithLoadingParams& Params = FForEachActorWithLoadingParams()) {}
-
-	UE_DEPRECATED(5.4, "Use ForEachActorWithLoading with FWorldPartitionActorDescInstance")
-	static void ForEachActorWithLoading(UWorldPartition* WorldPartition, TFunctionRef<bool(const FWorldPartitionActorDesc*)> Func, const FForEachActorWithLoadingParams& Params, FForEachActorWithLoadingResult& Result) {}
 	
 	static ENGINE_API void ForEachActorWithLoading(UWorldPartition* WorldPartition, TFunctionRef<bool(const FWorldPartitionActorDescInstance*)> Func, const FForEachActorWithLoadingParams& Params = FForEachActorWithLoadingParams());
 	static ENGINE_API void ForEachActorWithLoading(UWorldPartition* WorldPartition, TFunctionRef<bool(const FWorldPartitionActorDescInstance*)> Func, const FForEachActorWithLoadingParams& Params, FForEachActorWithLoadingResult& Result);
 
-
 	static ENGINE_API bool HasExceededMaxMemory();
 	static ENGINE_API bool ShouldCollectGarbage();
 	static ENGINE_API void DoCollectGarbage();
@@ -175,12 +141,6 @@ public:
 	// Runtime/Editor conversions
 	static ENGINE_API bool ConvertRuntimePathToEditorPath(const FSoftObjectPath& InPath, FSoftObjectPath& OutPath);
 
-	UE_DEPRECATED(5.4, "Use UAssetRegistryHelpers::FixupRedirectedAssetPath instead.")
-	static ENGINE_API bool FixupRedirectedAssetPath(FName& InOutAssetPath);
-	
-	UE_DEPRECATED(5.4, "Use UAssetRegistryHelpers::FixupRedirectedAssetPath instead.")
-	static ENGINE_API bool FixupRedirectedAssetPath(FSoftObjectPath& InOutSoftObjectPath);
-
 	// Returns all loaded actors for a specific level, including unregistered and garbage ones
 	static ENGINE_API TMap<FGuid, AActor*> GetLoadedActorsForLevel(const ULevel* InLevel);
 
diff --git a/Engine/Source/Runtime/Engine/Public/WorldPartition/WorldPartitionLevelHelper.h b/Engine/Source/Runtime/Engine/Public/WorldPartition/WorldPartitionLevelHelper.h
--- a/Engine/Source/Runtime/Engine/Public/WorldPartition/WorldPartitionLevelHelper.h
+++ b/Engine/Source/Runtime/Engine/Public/WorldPartition/WorldPartitionLevelHelper.h
@@ -41,9 +41,6 @@ public:
 	static void MoveExternalActorsToLevel(const TArray<FWorldPartitionRuntimeCellObjectMapping>& InChildPackages, ULevel* InLevel, TArray<UPackage*>& OutModifiedPackages);
 	static void RemapLevelSoftObjectPaths(ULevel* InLevel, UWorldPartition* InWorldPartition);
 	
-	UE_DEPRECATED(5.4, "LoadActors is deprecated, LoadActors with FLoadActorsParams should be used instead.")
-	static bool LoadActors(UWorld* InOuterWorld, ULevel* InDestLevel, TArrayView<FWorldPartitionRuntimeCellObjectMapping> InActorPackages, FPackageReferencer& InPackageReferencer, TFunction<void(bool)> InCompletionCallback, bool bInLoadAsync, FLinkerInstancingContext InInstancingContext);
-
 	/* Struct of optional parameters passed to LoadActors function. */
 	struct FLoadActorsParams
 	{
diff --git a/Engine/Source/Runtime/Engine/Public/WorldPartition/WorldPartitionStreamingGeneration.h b/Engine/Source/Runtime/Engine/Public/WorldPartition/WorldPartitionStreamingGeneration.h
--- a/Engine/Source/Runtime/Engine/Public/WorldPartition/WorldPartitionStreamingGeneration.h
+++ b/Engine/Source/Runtime/Engine/Public/WorldPartition/WorldPartitionStreamingGeneration.h
@@ -181,16 +181,6 @@ public:
 
 	ENGINE_API FStreamingGenerationActorDescView* Emplace(const FGuid& InActorGuid, FStreamingGenerationActorDescView&& InActorDescView);
 	ENGINE_API FStreamingGenerationActorDescView* Emplace(const FWorldPartitionActorDescInstance* InActorDescInstance);
-
-	PRAGMA_DISABLE_DEPRECATION_WARNINGS
-
-	UE_DEPRECATED(5.4, "Use FStreamingGenerationActorDescView version instead")
-	FWorldPartitionActorDescView* Emplace(const FGuid& InActorGuid, const FWorldPartitionActorDescView& InActorDescView) { return nullptr; }
-
-	UE_DEPRECATED(5.4, "Use FWorldPartitionActorDescInstance version instead")
-	class FWorldPartitionActorDescView* Emplace(const FWorldPartitionActorDesc* InActorDesc) { return nullptr; }
-
-	PRAGMA_ENABLE_DEPRECATION_WARNINGS
 			
 	inline int32 Num() const
 	{
@@ -282,23 +272,6 @@ private:
 	friend struct FGenerateStreamingParams;
 };
 
-class UE_DEPRECATED(5.4, "Use FStreamingGenerationContainerInstanceCollection instead") FStreamingGenerationActorDescCollection : public TActorDescContainerCollection<TObjectPtr<const UActorDescContainer>>
-{
-public:
-	FStreamingGenerationActorDescCollection() = default;
-	FStreamingGenerationActorDescCollection(std::initializer_list<TObjectPtr<const UActorDescContainer>> ActorDescContainerArray) {}
-	FStreamingGenerationActorDescCollection(const TArray<const UActorDescContainer*>& ActorDescContainers) {}
-
-	UWorld* GetWorld() const { return nullptr; }
-	FGuid GetContentBundleGuid() const { return FGuid(); }
-	const UActorDescContainer* GetMainActorDescContainer() const { return nullptr; }
-	FName GetMainContainerPackageName() const { return NAME_None; }
-	TArrayView<const UActorDescContainer* const> GetExternalDataLayerContainers() { return TArrayView<const UActorDescContainer*>(); }
-	virtual void OnCollectionChanged() override {}
-};
-
-using FActorDescViewMap UE_DEPRECATED(5.4, "Use FStreamingGenerationActorDescViewMap instead") = FStreamingGenerationActorDescViewMap;
-
 /**
  * Experimental: an actor descriptor mutator used by external code to change properties on actor descriptors.
  */
diff --git a/Engine/Source/Runtime/Engine/Public/WorldPartition/WorldPartitionStreamingGenerationContext.h b/Engine/Source/Runtime/Engine/Public/WorldPartition/WorldPartitionStreamingGenerationContext.h
--- a/Engine/Source/Runtime/Engine/Public/WorldPartition/WorldPartitionStreamingGenerationContext.h
+++ b/Engine/Source/Runtime/Engine/Public/WorldPartition/WorldPartitionStreamingGenerationContext.h
@@ -33,22 +33,6 @@ public:
 		TArray<FGuid> Actors;
 	};
 
-	struct UE_DEPRECATED(5.4, "Use FActorSetContainerInstance") FActorSetContainer
-	{
-		FActorSetContainer()
-			: ActorDescViewMap(nullptr)
-			, ActorDescCollection(nullptr)
-		{}
-
-		// Non-copyable
-		FActorSetContainer(const FActorSetContainer&) = delete;
-		FActorSetContainer& operator=(const FActorSetContainer&) = delete;
-
-		const FStreamingGenerationActorDescViewMap* ActorDescViewMap;
-		const class FStreamingGenerationActorDescCollection* ActorDescCollection; // Only used by UWorldPartitionRuntimeSpatialHash::SetupHLODActors
-		TArray<TUniquePtr<FActorSet>> ActorSets;
-	};
-
 	/**
 	 * An actor set container represents the list of actor sets in an actor container, e.g. a level instance.
 	 */
@@ -81,13 +65,6 @@ public:
 		TArray<const UDataLayerInstance*> DataLayers;
 		FGuid ContentBundleID;
 
-		PRAGMA_DISABLE_DEPRECATION_WARNINGS
-
-		UE_DEPRECATED(5.4, "Use ActorSetContainerInstance instead")
-		const FActorSetContainer* ContainerInstance = nullptr;
-
-		PRAGMA_ENABLE_DEPRECATION_WARNINGS
-
 		const FActorSetContainerInstance* ActorSetContainerInstance = nullptr;
 
 		FActorContainerID ContainerID;
@@ -139,16 +116,6 @@ public:
 	virtual const FActorSetContainerInstance* GetActorSetContainerForContextBaseContainerInstance() const = 0;
 	virtual void ForEachActorSetInstance(TFunctionRef<void(const FActorSetInstance&)> Func) const = 0;
 	virtual void ForEachActorSetContainerInstance(TFunctionRef<void(const FActorSetContainerInstance&)> Func) const = 0;
-
-	PRAGMA_DISABLE_DEPRECATION_WARNINGS
-
-	UE_DEPRECATED(5.4, "Implement GetActorSetContainerForContextBaseContainerInstance instead")
-	virtual const FActorSetContainer* GetMainWorldContainer() const { return nullptr; };
-
-	UE_DEPRECATED(5.4, "Implement ForEachActorSetContainerInstance instead")
-	virtual void ForEachActorSetContainer(TFunctionRef<void(const FActorSetContainer&)> Func) const {};
-
-	PRAGMA_ENABLE_DEPRECATION_WARNINGS
 };
 
 class FStreamingGenerationContextProxy : public IStreamingGenerationContext
diff --git a/Engine/Source/Runtime/Engine/Public/WorldPartition/WorldPartitionStreamingPolicy.h b/Engine/Source/Runtime/Engine/Public/WorldPartition/WorldPartitionStreamingPolicy.h
--- a/Engine/Source/Runtime/Engine/Public/WorldPartition/WorldPartitionStreamingPolicy.h
+++ b/Engine/Source/Runtime/Engine/Public/WorldPartition/WorldPartitionStreamingPolicy.h
@@ -108,9 +108,6 @@ public:
 	virtual void OnCellShown(const UWorldPartitionRuntimeCell* InCell);
 	virtual void OnCellHidden(const UWorldPartitionRuntimeCell* InCell);
 
-	UE_DEPRECATED(5.3, "CanAddLoadedLevelToWorld is deprecated.")
-	virtual bool CanAddLoadedLevelToWorld(class ULevel* InLevel) const { return true; }
-
 #if WITH_EDITOR
 	virtual TSubclassOf<class UWorldPartitionRuntimeCell> GetRuntimeCellClass() const PURE_VIRTUAL(UWorldPartitionStreamingPolicy::GetRuntimeCellClass, return UWorldPartitionRuntimeCell::StaticClass(); );
 
@@ -119,8 +116,6 @@ public:
 	virtual void SetContainerResolver(const FWorldPartitionRuntimeContainerResolver& InContainerResolver) {}
 	virtual void RemapSoftObjectPath(FSoftObjectPath& ObjectPath) const {}
 
-	UE_DEPRECATED(5.4, "Use StoreStreamingContentToExternalStreamingObject instead.")
-	virtual bool StoreToExternalStreamingObject(URuntimeHashExternalStreamingObjectBase& OutExternalStreamingObject) { return StoreStreamingContentToExternalStreamingObject(OutExternalStreamingObject); }
 	virtual bool StoreStreamingContentToExternalStreamingObject(URuntimeHashExternalStreamingObjectBase& OutExternalStreamingObject) { return true; }
 	virtual bool ConvertContainerPathToEditorPath(const FActorContainerID& InContainerID, const FSoftObjectPath& InPath, FSoftObjectPath& OutPath) const { return false; }
 #endif
@@ -152,8 +147,6 @@ protected:
 	virtual void GetCellsToReprioritize(TArray<const UWorldPartitionRuntimeCell*>& OutToLoadCells, TArray<const UWorldPartitionRuntimeCell*>& OutToActivateCells);
 	virtual void UpdateStreamingSources(bool bCanOptimizeUpdate);
 
-	UE_DEPRECATED(5.6, "Use version that also sets if streaming should block on slow loading")
-	void UpdateStreamingPerformance(EWorldPartitionStreamingPerformance NewStreamingPerformance);
 	void UpdateStreamingPerformance(EWorldPartitionStreamingPerformance NewStreamingPerformance, bool bBlockOnSlowStreaming);
 	bool IsInBlockTillLevelStreamingCompleted(bool bIsCausedByBadStreamingPerformance = false) const;
 
diff --git a/Engine/Source/Runtime/Engine/Public/WorldPartition/WorldPartitionStreamingSource.h b/Engine/Source/Runtime/Engine/Public/WorldPartition/WorldPartitionStreamingSource.h
--- a/Engine/Source/Runtime/Engine/Public/WorldPartition/WorldPartitionStreamingSource.h
+++ b/Engine/Source/Runtime/Engine/Public/WorldPartition/WorldPartitionStreamingSource.h
@@ -260,13 +260,6 @@ struct FWorldPartitionStreamingQuerySource
 		, TargetBehavior(EStreamingSourceTargetBehavior::Include)
 	{}
 
-	PRAGMA_DISABLE_DEPRECATION_WARNINGS
-	FWorldPartitionStreamingQuerySource(const FWorldPartitionStreamingQuerySource& Other) = default;
-	FWorldPartitionStreamingQuerySource& operator=(const FWorldPartitionStreamingQuerySource& Other) = default;
-	FWorldPartitionStreamingQuerySource(FWorldPartitionStreamingQuerySource&& Other) = default;
-	FWorldPartitionStreamingQuerySource& operator=(FWorldPartitionStreamingQuerySource&& Other) = default;
-	PRAGMA_ENABLE_DEPRECATION_WARNINGS
-
 	/* Location to query. (not used if bSpatialQuery is false) */
 	UPROPERTY(BlueprintReadWrite, EditAnywhere, Category = "Query")
 	FVector Location;
@@ -298,9 +291,6 @@ struct FWorldPartitionStreamingQuerySource
 	EStreamingSourceTargetBehavior TargetBehavior;
 
 	TSet<FName> TargetGrids;
-
-	UE_DEPRECATED(5.4, "TargetHLODLayers is depredcated, use TargetGrids instead.")
-	TSet<FSoftObjectPath> TargetHLODLayers;
 	
 	TArray<FStreamingSourceShape> Shapes;
 
@@ -386,13 +376,6 @@ struct FWorldPartitionStreamingSource
 		, ExtraAngle(0)
 	{}
 
-	PRAGMA_DISABLE_DEPRECATION_WARNINGS
-	FWorldPartitionStreamingSource(const FWorldPartitionStreamingSource& Other) = default;
-	FWorldPartitionStreamingSource& operator=(const FWorldPartitionStreamingSource& Other) = default;
-	FWorldPartitionStreamingSource(FWorldPartitionStreamingSource&& Other) = default;
-	FWorldPartitionStreamingSource& operator=(FWorldPartitionStreamingSource&& Other) = default;
-	PRAGMA_ENABLE_DEPRECATION_WARNINGS
-
 	FColor GetDebugColor() const
 	{
 		if (!DebugColor.ToPackedBGRA())
@@ -433,15 +416,11 @@ struct FWorldPartitionStreamingSource
 	/** Color used for debugging. */
 	FColor DebugColor;
 
-	/** Defines how TargetGrids/TargetHLODLayers will be applied to this streaming source. */
+	/** Defines how TargetGrids will be applied to this streaming source. */
 	EStreamingSourceTargetBehavior TargetBehavior;
 
 	/** When set, this will change how this streaming source is applied to the provided runtime streaming grids based on the TargetBehavior. */
 	TSet<FName> TargetGrids;
-	
-	/** When set, this will change how this streaming source is applied to the provided HLODLayers based on the TargetBehavior. */
-	UE_DEPRECATED(5.4, "TargetHLODLayers is depredcated, use TargetGrids instead.")
-	TSet<FSoftObjectPath> TargetHLODLayers;
 
 	/** Source internal shapes. When none are provided, a sphere is automatically used. It's radius is equal to grid's loading range and center equals source's location. */
 	TArray<FStreamingSourceShape> Shapes;
diff --git a/Engine/Source/Runtime/Engine/Public/WorldPartition/WorldPartitionSubsystem.h b/Engine/Source/Runtime/Engine/Public/WorldPartition/WorldPartitionSubsystem.h
--- a/Engine/Source/Runtime/Engine/Public/WorldPartition/WorldPartitionSubsystem.h
+++ b/Engine/Source/Runtime/Engine/Public/WorldPartition/WorldPartitionSubsystem.h
@@ -135,15 +135,6 @@ public:
 	// When enabled, returns a linear loading range scaling ramp to boost initial PIE availability, at the cost of popping. When disabled, the values are 1.f
 	const FDynamicLoadingRangeScales& GetDynamicLoadingRangeScales() const { return LoadingRangeScales; }
 
-	UE_DEPRECATED(5.4, "Use UActorDescContainerSubsystem instead")
-	UActorDescContainer* RegisterContainer(FName PackageName) { return nullptr; }
-	UE_DEPRECATED(5.4, "Use UActorDescContainerSubsystem instead")
-	void UnregisterContainer(UActorDescContainer* Container) { }
-	UE_DEPRECATED(5.4, "Use UActorDescContainerSubsystem instead")
-	FBox GetContainerBounds(FName PackageName) const { return FBox(); }
-	UE_DEPRECATED(5.4, "Use UActorDescContainerSubsystem instead")
-	void UpdateContainerBounds(FName PackageName) { }
-
 	TSet<FWorldPartitionHandle> SelectedActorHandles;
 
 private:
