[WorldParitition] Clean deprecated code

#Synced-CL 50613390

diff --git a/Engine/Source/Runtime/Engine/Private/WorldPartition/RuntimeSpatialHash/RuntimeSpatialHashGridHelper.cpp b/Engine/Source/Runtime/Engine/Private/WorldPartition/RuntimeSpatialHash/RuntimeSpatialHashGridHelper.cpp
--- a/Engine/Source/Runtime/Engine/Private/WorldPartition/RuntimeSpatialHash/RuntimeSpatialHashGridHelper.cpp
+++ b/Engine/Source/Runtime/Engine/Private/WorldPartition/RuntimeSpatialHash/RuntimeSpatialHashGridHelper.cpp
@@ -35,10 +35,6 @@ static FAutoConsoleVariableRef CVarRuntimeSpatialHashPlacePartitionActorsUsingLo
 
 #endif
 
-FSquare2DGridHelper::FSquare2DGridHelper(const FBox& InWorldBounds, const FVector& InOrigin, int64 InCellSize)
-	: FSquare2DGridHelper(InWorldBounds, InOrigin, InCellSize, true)
-{}
-
 FSquare2DGridHelper::FSquare2DGridHelper(const FBox& InWorldBounds, const FVector& InOrigin, int64 InCellSize, bool bUseAlignedGridLevels)
 	: WorldBounds(InWorldBounds)
 	, Origin(InOrigin)
@@ -146,11 +142,6 @@ int32 FSquare2DGridHelper::ForEachIntersectingCells(const FSphericalSector& InSh
 }
 
 #if WITH_EDITOR
-// Deprecated
-FSquare2DGridHelper GetGridHelper(const FBox& WorldBounds, const FVector& GridOrigin, int64 GridCellSize)
-{
-	return GetGridHelper(WorldBounds, GridOrigin, GridCellSize, true);
-}
 
 FSquare2DGridHelper GetGridHelper(const FBox& WorldBounds, const FVector& GridOrigin, int64 GridCellSize, bool bUseAlignedGridLevels)
 {
@@ -158,13 +149,6 @@ FSquare2DGridHelper GetGridHelper(const FBox& WorldBounds, const FVector& GridOr
 	return FSquare2DGridHelper(WorldBounds, GridOrigin, GridCellSize, bUseAlignedGridLevels);
 }
 
-// Deprecated
-FSquare2DGridHelper GetPartitionedActors(const FBox& WorldBounds, const FSpatialHashRuntimeGrid& Grid, const TArray<const IStreamingGenerationContext::FActorSetInstance*>& ActorSetInstances)
-{
-	FSpatialHashSettings Settings;
-	return GetPartitionedActors(WorldBounds, Grid, ActorSetInstances, Settings);
-}
-
 FSquare2DGridHelper GetPartitionedActors(const FBox& WorldBounds, const FSpatialHashRuntimeGrid& Grid, const TArray<const IStreamingGenerationContext::FActorSetInstance*>& ActorSetInstances, const FSpatialHashSettings& Settings)
 {
 	TRACE_CPUPROFILER_EVENT_SCOPE(GetPartitionedActors);
diff --git a/Engine/Source/Runtime/Engine/Private/WorldPartition/WorldPartitionRuntimeCellData.cpp b/Engine/Source/Runtime/Engine/Private/WorldPartition/WorldPartitionRuntimeCellData.cpp
--- a/Engine/Source/Runtime/Engine/Private/WorldPartition/WorldPartitionRuntimeCellData.cpp
+++ b/Engine/Source/Runtime/Engine/Private/WorldPartition/WorldPartitionRuntimeCellData.cpp
@@ -53,10 +53,6 @@ void UWorldPartitionRuntimeCellData::ResetStreamingSourceInfo(const FWorldPartit
 }
 
 DECLARE_CYCLE_STAT_WITH_FLAGS(TEXT("Append Streaming Source Info"), STAT_WorldPartitionAppendStreamingSourceInfo, STATGROUP_WorldPartition, EStatFlags::Verbose);
-void UWorldPartitionRuntimeCellData::AppendStreamingSourceInfo(const FWorldPartitionStreamingSource& Source, const FSphericalSector& SourceShape, const FWorldPartitionStreamingContext& Context) const
-{
-	AppendStreamingSourceInfo(Source, SourceShape, Context, false);
-}
 
 void UWorldPartitionRuntimeCellData::AppendStreamingSourceInfo(const FWorldPartitionStreamingSource& Source, const FSphericalSector& SourceShape, const FWorldPartitionStreamingContext& Context, bool bBlockOnSlowLoading) const
 {
diff --git a/Engine/Source/Runtime/Engine/Private/WorldPartition/WorldPartitionRuntimeCellDataSpatialHash.cpp b/Engine/Source/Runtime/Engine/Private/WorldPartition/WorldPartitionRuntimeCellDataSpatialHash.cpp
--- a/Engine/Source/Runtime/Engine/Private/WorldPartition/WorldPartitionRuntimeCellDataSpatialHash.cpp
+++ b/Engine/Source/Runtime/Engine/Private/WorldPartition/WorldPartitionRuntimeCellDataSpatialHash.cpp
@@ -85,11 +85,6 @@ float UWorldPartitionRuntimeCellDataSpatialHash::ComputeSourceToCellAngleFactor(
 	return AngleFactor;
 }
 
-void UWorldPartitionRuntimeCellDataSpatialHash::AppendStreamingSourceInfo(const FWorldPartitionStreamingSource& Source, const FSphericalSector& SourceShape, const FWorldPartitionStreamingContext& Context) const
-{
-	Super::AppendStreamingSourceInfo(Source, SourceShape, Context, false);
-}
-
 void UWorldPartitionRuntimeCellDataSpatialHash::AppendStreamingSourceInfo(const FWorldPartitionStreamingSource& Source, const FSphericalSector& SourceShape, const FWorldPartitionStreamingContext& Context, bool bBlockOnSlowLoading) const
 {
 	Super::AppendStreamingSourceInfo(Source, SourceShape, Context, bBlockOnSlowLoading);
diff --git a/Engine/Source/Runtime/Engine/Private/WorldPartition/WorldPartitionRuntimeHash.cpp b/Engine/Source/Runtime/Engine/Private/WorldPartition/WorldPartitionRuntimeHash.cpp
--- a/Engine/Source/Runtime/Engine/Private/WorldPartition/WorldPartitionRuntimeHash.cpp
+++ b/Engine/Source/Runtime/Engine/Private/WorldPartition/WorldPartitionRuntimeHash.cpp
@@ -242,11 +242,6 @@ UWorldPartitionRuntimeCell* UWorldPartitionRuntimeHash::CreateRuntimeCell(UClass
 	return RuntimeCell;
 }
 
-EWorldPartitionStreamingPerformance UWorldPartitionRuntimeHash::GetStreamingPerformanceForCell(const UWorldPartitionRuntimeCell* Cell) const
-{
-	return EWorldPartitionStreamingPerformance::Good;
-}
-
 EWorldPartitionStreamingPerformance UWorldPartitionRuntimeHash::GetStreamingPerformanceForCell(const UWorldPartitionRuntimeCell* Cell, bool& bOutShouldBlock) const
 {
 	EWorldPartitionStreamingPerformance BlockingPerformance = EWorldPartitionStreamingPerformance::Good;
@@ -626,11 +621,6 @@ bool UWorldPartitionRuntimeHash::IsCellRelevantFor(bool bClientOnlyVisible) cons
 	return true;
 }
 
-EWorldPartitionStreamingPerformance UWorldPartitionRuntimeHash::GetStreamingPerformance(const TSet<const UWorldPartitionRuntimeCell*>& CellsToActivate) const
-{
-	return EWorldPartitionStreamingPerformance::Good;
-}
-
 EWorldPartitionStreamingPerformance UWorldPartitionRuntimeHash::GetStreamingPerformance(const TSet<const UWorldPartitionRuntimeCell*>& CellsToActivate, bool& bOutShouldBlock) const
 {
 	TRACE_CPUPROFILER_EVENT_SCOPE(UWorldPartitionRuntimeHash::GetStreamingPerformance)
diff --git a/Engine/Source/Runtime/Engine/Private/WorldPartition/WorldPartitionRuntimeSpatialHash.cpp b/Engine/Source/Runtime/Engine/Private/WorldPartition/WorldPartitionRuntimeSpatialHash.cpp
--- a/Engine/Source/Runtime/Engine/Private/WorldPartition/WorldPartitionRuntimeSpatialHash.cpp
+++ b/Engine/Source/Runtime/Engine/Private/WorldPartition/WorldPartitionRuntimeSpatialHash.cpp
@@ -1992,11 +1992,6 @@ const FSpatialHashStreamingGrid* UWorldPartitionRuntimeSpatialHash::GetStreaming
 	return nullptr;
 }
 
-EWorldPartitionStreamingPerformance UWorldPartitionRuntimeSpatialHash::GetStreamingPerformanceForCell(const UWorldPartitionRuntimeCell* Cell) const
-{
-	return EWorldPartitionStreamingPerformance::Good;
-}
-
 ENGINE_API EWorldPartitionStreamingPerformance UWorldPartitionRuntimeSpatialHash::GetStreamingPerformanceForCell(const UWorldPartitionRuntimeCell* Cell, bool& bOutShouldBlock) const
 {
 	const UWorldPartitionRuntimeCellDataSpatialHash* CellDataSpatialHash = CastChecked<UWorldPartitionRuntimeCellDataSpatialHash>(Cell->RuntimeCellData);
diff --git a/Engine/Source/Runtime/Engine/Public/WorldPartition/RuntimeSpatialHash/RuntimeSpatialHashGridHelper.h b/Engine/Source/Runtime/Engine/Public/WorldPartition/RuntimeSpatialHash/RuntimeSpatialHashGridHelper.h
--- a/Engine/Source/Runtime/Engine/Public/WorldPartition/RuntimeSpatialHash/RuntimeSpatialHashGridHelper.h
+++ b/Engine/Source/Runtime/Engine/Public/WorldPartition/RuntimeSpatialHash/RuntimeSpatialHashGridHelper.h
@@ -393,9 +393,6 @@ struct FSquare2DGridHelper
 		TMap<int64, int64> CellsMapping;
 	};
 
-	UE_DEPRECATED(5.4, "Use version with bUseAlignedGridLevels param")
-	ENGINE_API FSquare2DGridHelper(const FBox& InWorldBounds, const FVector& InOrigin, int64 InCellSize);
-
 	ENGINE_API FSquare2DGridHelper(const FBox& InWorldBounds, const FVector& InOrigin, int64 InCellSize, bool bUseAlignedGridLevels);
 
 	// Returns the lowest grid level
@@ -482,9 +479,4 @@ public:
 FSquare2DGridHelper GetGridHelper(const FBox& WorldBounds, const FVector& GridOrigin, int64 GridCellSize, bool bUseAlignedGridLevels);
 FSquare2DGridHelper GetPartitionedActors(const FBox& WorldBounds, const FSpatialHashRuntimeGrid& Grid, const TArray<const IStreamingGenerationContext::FActorSetInstance*>& ActorSetInstances, const FSpatialHashSettings& Settings);
 
-UE_DEPRECATED(5.4, "Use version with bUseAlignedGridLevels param")
-FSquare2DGridHelper GetGridHelper(const FBox& WorldBounds, const FVector& GridOrigin, int64 GridCellSize);
-
-UE_DEPRECATED(5.4, "Use version with Settings param")
-FSquare2DGridHelper GetPartitionedActors(const FBox& WorldBounds, const FSpatialHashRuntimeGrid& Grid, const TArray<const IStreamingGenerationContext::FActorSetInstance*>& ActorSetInstances);
 #endif // #if WITH_EDITOR
diff --git a/Engine/Source/Runtime/Engine/Public/WorldPartition/WorldPartitionEditorHash.h b/Engine/Source/Runtime/Engine/Public/WorldPartition/WorldPartitionEditorHash.h
--- a/Engine/Source/Runtime/Engine/Public/WorldPartition/WorldPartitionEditorHash.h
+++ b/Engine/Source/Runtime/Engine/Public/WorldPartition/WorldPartitionEditorHash.h
@@ -45,9 +45,6 @@ public:
 		FForEachIntersectingActorParams& SetMinimumBox(const FBox& InMinimumBox) { MinimumBox = InMinimumBox; return *this; };
 	};
 
-	UE_DEPRECATED(5.4, "Use ForEachIntersectingActor with FWorldPartitionActorDescInstance")
-	virtual int32 ForEachIntersectingActor(const FBox& Box, TFunctionRef<void(FWorldPartitionActorDesc*)> InOperation, const FForEachIntersectingActorParams& Params = FForEachIntersectingActorParams()) { return 0; }
-
 	virtual int32 ForEachIntersectingActor(const FBox& Box, TFunctionRef<void(FWorldPartitionActorDescInstance*)> InOperation, const FForEachIntersectingActorParams& Params = FForEachIntersectingActorParams()) PURE_VIRTUAL(UWorldPartitionEditorHash::ForEachIntersectingActor, return 0;);
 #endif
 };
diff --git a/Engine/Source/Runtime/Engine/Public/WorldPartition/WorldPartitionEditorSpatialHash.h b/Engine/Source/Runtime/Engine/Public/WorldPartition/WorldPartitionEditorSpatialHash.h
--- a/Engine/Source/Runtime/Engine/Public/WorldPartition/WorldPartitionEditorSpatialHash.h
+++ b/Engine/Source/Runtime/Engine/Public/WorldPartition/WorldPartitionEditorSpatialHash.h
@@ -198,9 +198,6 @@ public:
 	ENGINE_API virtual void HashActor(FWorldPartitionHandle& InActorHandle) override;
 	ENGINE_API virtual void UnhashActor(FWorldPartitionHandle& InActorHandle) override;
 
-	UE_DEPRECATED(5.4, "Use ForEachIntersectingActor with FWorldPartitionActorDescInstance")
-	virtual int32 ForEachIntersectingActor(const FBox& Box, TFunctionRef<void(FWorldPartitionActorDesc*)> InOperation, const FForEachIntersectingActorParams& Params = FForEachIntersectingActorParams()) override { return 0;  }
-
 	ENGINE_API virtual int32 ForEachIntersectingActor(const FBox& Box, TFunctionRef<void(FWorldPartitionActorDescInstance*)> InOperation, const FForEachIntersectingActorParams& Params = FForEachIntersectingActorParams()) override;
 	// UWorldPartitionEditorHash interface end
 #endif
diff --git a/Engine/Source/Runtime/Engine/Public/WorldPartition/WorldPartitionRuntimeCell.h b/Engine/Source/Runtime/Engine/Public/WorldPartition/WorldPartitionRuntimeCell.h
--- a/Engine/Source/Runtime/Engine/Public/WorldPartition/WorldPartitionRuntimeCell.h
+++ b/Engine/Source/Runtime/Engine/Public/WorldPartition/WorldPartitionRuntimeCell.h
@@ -289,19 +289,6 @@ class UWorldPartitionRuntimeCell : public UObject, public IWorldPartitionCell, p
 	ENGINE_API EDataLayerRuntimeState GetCellEffectiveWantedState(const FWorldPartitionStreamingContext& Context) const;
 	inline bool HasDataLayers() const { return !DataLayers.IsEmpty(); }
 
-	//~Begin Deprecation
-	UE_DEPRECATED(5.3, "IsLoading is deprecated.")
-	virtual bool IsLoading() const { return false; }
-	UE_DEPRECATED(5.4, "IsAddedToWorld is deprecated.")
-	virtual bool IsAddedToWorld() const { return false; }
-	UE_DEPRECATED(5.4, "CanAddToWorld is deprecated.")
-	virtual bool CanAddToWorld() const { return false; }
-	UE_DEPRECATED(5.5, "Use version that takes FWorldPartitionStreamingContext instead.")
-	EDataLayerRuntimeState GetCellEffectiveWantedState(int32 InDataLayersStateEpoch) const { return EDataLayerRuntimeState::Unloaded; }
-	UE_DEPRECATED(5.5, "Use version that takes FWorldPartitionStreamingContext instead.")
-	EDataLayerRuntimeState GetCellEffectiveWantedState() const { return EDataLayerRuntimeState::Unloaded; }
-	//~End Deprecation
-
 	void SetBlockOnSlowLoading(bool bInBlockOnSlowLoading) { bBlockOnSlowLoading = bInBlockOnSlowLoading; }
 	bool GetBlockOnSlowLoading() const { return bBlockOnSlowLoading; }
 #if WITH_EDITOR
@@ -316,25 +303,12 @@ class UWorldPartitionRuntimeCell : public UObject, public IWorldPartitionCell, p
 	//~End IWorldPartitionCell Interface
 
 	virtual void AddActorToCell(const FStreamingGenerationActorDescView& ActorDescView) PURE_VIRTUAL(UWorldPartitionRuntimeCell::AddActorToCell, );
-	
-	PRAGMA_DISABLE_DEPRECATION_WARNINGS
-
-	UE_DEPRECATED(5.4, "Implement FStreamingGenerationActorDescView version instead")
-	virtual void AddActorToCell(const class FWorldPartitionActorDescView& ActorDescView, const FActorContainerID& InContainerID, const FTransform& InContainerTransform, const UActorDescContainer* InContainer) {}
-
-	PRAGMA_ENABLE_DEPRECATION_WARNINGS
 
 	virtual void Fixup() PURE_VIRTUAL(UWorldPartitionRuntimeCell::Fixup, );
 	virtual int32 GetActorCount() const PURE_VIRTUAL(UWorldPartitionRuntimeCell::GetActorCount, return 0;);
 
 	// Cook methods
-	UE_DEPRECATED(5.5, "Use version with IWorldPartitionCookPackageContext")
-	virtual bool PrepareCellForCook(UPackage* InPackage) { return false; }
 	virtual bool PrepareCellForCook(const IWorldPartitionCookPackageContext& InCookContext, UPackage* InGeneratedPackage = nullptr) { return true; }
-	UE_DEPRECATED(5.4, "PopulateGeneratorPackageForCook is deprecated, it was replaced by OnPrepareGeneratorPackageForCook")
-	virtual bool PopulateGeneratorPackageForCook(TArray<UPackage*>& OutModifiedPackages) { return OnPrepareGeneratorPackageForCook(OutModifiedPackages); }
-	UE_DEPRECATED(5.4, "PopulateGeneratedPackageForCook is deprecated, it was replaced by OnPopulateGeneratedPackageForCook")
-	virtual bool PopulateGeneratedPackageForCook(UPackage* InPackage, TArray<UPackage*>& OutModifiedPackages) { return false; }
 
 	//~Begin IWorldPartitionCookPackageObject
 	virtual bool IsLevelPackage() const override { return true; }
@@ -457,11 +431,6 @@ public:
 	ENGINE_API virtual bool IsDebugShown() const;
 	//~End UWorldPartitionRuntimeCellData Proxy
 
-	//~Begin Deprecation
-	UE_DEPRECATED(5.5, "Use version that takes FWorldPartitionStreamingContext instead.")
-	inline void AppendStreamingSourceInfo(const FWorldPartitionStreamingSource& Source, const FSphericalSector& SourceShape) {}
-	//~End Deprecation
-
 	UPROPERTY()
 	TObjectPtr<UWorldPartitionRuntimeCellData> RuntimeCellData;
 
diff --git a/Engine/Source/Runtime/Engine/Public/WorldPartition/WorldPartitionRuntimeCellData.h b/Engine/Source/Runtime/Engine/Public/WorldPartition/WorldPartitionRuntimeCellData.h
--- a/Engine/Source/Runtime/Engine/Public/WorldPartition/WorldPartitionRuntimeCellData.h
+++ b/Engine/Source/Runtime/Engine/Public/WorldPartition/WorldPartitionRuntimeCellData.h
@@ -27,8 +27,6 @@ class UWorldPartitionRuntimeCellData : public UObject
 #endif
 
 	ENGINE_API virtual void ResetStreamingSourceInfo(const FWorldPartitionStreamingContext& Context) const;
-	UE_DEPRECATED(5.6, "Use version that takes a boolean to flag if the cell blocks on slow loading (bBlockOnSlowLoading).")
-	ENGINE_API virtual void AppendStreamingSourceInfo(const FWorldPartitionStreamingSource& Source, const FSphericalSector& SourceShape, const FWorldPartitionStreamingContext& Context) const;
 	ENGINE_API virtual void AppendStreamingSourceInfo(const FWorldPartitionStreamingSource& Source, const FSphericalSector& SourceShape, const FWorldPartitionStreamingContext& Context, bool bBlockOnSlowLoading) const;
 	virtual void MergeStreamingSourceInfo() const {}
 	ENGINE_API virtual int32 SortCompare(const UWorldPartitionRuntimeCellData* InOther) const;
@@ -45,13 +43,6 @@ class UWorldPartitionRuntimeCellData : public UObject
 	virtual bool IsDebugShown() const { return true; }
 	ENGINE_API virtual FString GetDebugName() const;
 
-	//~Begin Deprecation
-	UE_DEPRECATED(5.5, "Use version that takes FWorldPartitionStreamingContext instead.")
-	virtual void ResetStreamingSourceInfo() const {}
-	UE_DEPRECATED(5.5, "Use version that takes FWorldPartitionStreamingContext instead.")
-	virtual void AppendStreamingSourceInfo(const FWorldPartitionStreamingSource& Source, const FSphericalSector& SourceShape) const {}
-	//~End Deprecation
-
 	// Minimum affecting source priority
 	mutable uint8 CachedMinSourcePriority;
 
diff --git a/Engine/Source/Runtime/Engine/Public/WorldPartition/WorldPartitionRuntimeCellDataSpatialHash.h b/Engine/Source/Runtime/Engine/Public/WorldPartition/WorldPartitionRuntimeCellDataSpatialHash.h
--- a/Engine/Source/Runtime/Engine/Public/WorldPartition/WorldPartitionRuntimeCellDataSpatialHash.h
+++ b/Engine/Source/Runtime/Engine/Public/WorldPartition/WorldPartitionRuntimeCellDataSpatialHash.h
@@ -15,8 +15,6 @@ class UWorldPartitionRuntimeCellDataSpatialHash : public UWorldPartitionRuntimeC
 
 	//~Begin UWorldPartitionRuntimeCellData
 	ENGINE_API virtual void ResetStreamingSourceInfo(const FWorldPartitionStreamingContext& Context) const override;
-	UE_DEPRECATED(5.6, "Use version that takes a boolean to flag if the cell blocks on slow loading (bBlockOnSlowLoading).")
-	ENGINE_API virtual void AppendStreamingSourceInfo(const FWorldPartitionStreamingSource& Source, const FSphericalSector& SourceShape, const FWorldPartitionStreamingContext& Context) const override;
 	ENGINE_API virtual void AppendStreamingSourceInfo(const FWorldPartitionStreamingSource& Source, const FSphericalSector& SourceShape, const FWorldPartitionStreamingContext& Context, bool bBlockOnSlowLoading) const override;
 	ENGINE_API virtual void MergeStreamingSourceInfo() const override;
 	ENGINE_API virtual int32 SortCompare(const UWorldPartitionRuntimeCellData* InOther) const override;
diff --git a/Engine/Source/Runtime/Engine/Public/WorldPartition/WorldPartitionRuntimeHash.h b/Engine/Source/Runtime/Engine/Public/WorldPartition/WorldPartitionRuntimeHash.h
--- a/Engine/Source/Runtime/Engine/Public/WorldPartition/WorldPartitionRuntimeHash.h
+++ b/Engine/Source/Runtime/Engine/Public/WorldPartition/WorldPartitionRuntimeHash.h
@@ -87,9 +87,6 @@ public:
 	ENGINE_API const UObject* GetLevelMountPointContextObject() const;
 
 #if WITH_EDITOR
-	UE_DEPRECATED(5.4, "PopulateGeneratorPackageForCook is depreacted and was replaced by OnPopulateGeneratorPackageForCook")
-	ENGINE_API void PopulateGeneratorPackageForCook();
-
 	//~Begin IWorldPartitionCookPackageObject interface
 	virtual bool IsLevelPackage() const override { return false; }
 	virtual const UExternalDataLayerAsset* GetExternalDataLayerAsset() const override { return ExternalDataLayerAsset; }
@@ -214,27 +211,9 @@ class UWorldPartitionRuntimeHash : public UObject
 	virtual bool HasStreamingContent() const { return false; }
 	ENGINE_API URuntimeHashExternalStreamingObjectBase* StoreStreamingContentToExternalStreamingObject();
 
-	UE_DEPRECATED(5.5, "StoreStreamingContentToExternalStreamingObject(FName) is deprecated, use StoreStreamingContentToExternalStreamingObject() instead")
-	URuntimeHashExternalStreamingObjectBase* StoreStreamingContentToExternalStreamingObject(FName InStreamingObjectName) { return StoreStreamingContentToExternalStreamingObject(); }
-
 	ENGINE_API virtual void FlushStreamingContent();
 	virtual TSubclassOf<URuntimeHashExternalStreamingObjectBase> GetExternalStreamingObjectClass() const PURE_VIRTUAL(UWorldPartitionRuntimeHash::GetExternalStreamingObjectClass, return nullptr;);
 	ENGINE_API virtual void DumpStateLog(FHierarchicalLogArchive& Ar) const;
-
-	//~Begin Deprecation
-	UE_DEPRECATED(5.4, "Use StoreStreamingContentToExternalStreamingObject instead.")
-	virtual URuntimeHashExternalStreamingObjectBase* StoreToExternalStreamingObject(UObject* StreamingObjectOuter, FName StreamingObjectName) { return nullptr; }
-	UE_DEPRECATED(5.4, "Use FlushStreamingContent instead.")
-	virtual void FlushStreaming() { FlushStreamingContent(); }
-	UE_DEPRECATED(5.4, "PrepareGeneratorPackageForCook is deprecated.")
-	bool PrepareGeneratorPackageForCook(TArray<UPackage*>& OutModifiedPackages) { return false; }
-	UE_DEPRECATED(5.4, "PopulateGeneratorPackageForCook is deprecated.")
-	bool PopulateGeneratorPackageForCook(const TArray<FWorldPartitionCookPackage*>& PackagesToCook, TArray<UPackage*>& OutModifiedPackages) { return false; }
-	UE_DEPRECATED(5.4, "PopulateGeneratedPackageForCook is deprecated.")
-	bool PopulateGeneratedPackageForCook(const FWorldPartitionCookPackage& PackagesToCook, TArray<UPackage*>& OutModifiedPackages) { return false; }
-	UE_DEPRECATED(5.4, "GetCellForPackage is deprecated.")
-	UWorldPartitionRuntimeCell* GetCellForPackage(const FWorldPartitionCookPackage& PackageToCook) const { return nullptr; }
-	//~End Deprecation
 #endif
 
 	virtual void OnBeginPlay() {}
@@ -259,11 +238,6 @@ public:
 		int32 Num() const { return Cells.Num(); }
 		TSet<const UWorldPartitionRuntimeCell*>& GetCells() { return Cells; }
 
-		//~Begin Deprecation
-		UE_DEPRECATED(5.5, "Use version that takes FWorldPartitionStreamingContext instead.")
-		void AddCell(const UWorldPartitionRuntimeCell* Cell, const FWorldPartitionStreamingSource& Source, const FSphericalSector& SourceShape) {}
-		//~End Deprecation
-
 	private:
 		TSet<const UWorldPartitionRuntimeCell*> Cells;
 	};
@@ -277,8 +251,6 @@ public:
 
 	ENGINE_API bool IsCellRelevantFor(bool bClientOnlyVisible) const;
 
-	UE_DEPRECATED(5.6, "Use version that has also flags if streaming should block (bOutShouldBlock) instead.")
-	ENGINE_API EWorldPartitionStreamingPerformance GetStreamingPerformance(const TSet<const UWorldPartitionRuntimeCell*>& CellsToActivate) const;	
 	ENGINE_API EWorldPartitionStreamingPerformance GetStreamingPerformance(const TSet<const UWorldPartitionRuntimeCell*>& CellsToActivate, bool& bOutShouldBlock) const;	
 
 	ENGINE_API virtual bool IsExternalStreamingObjectInjected(URuntimeHashExternalStreamingObjectBase* InExternalStreamingObject) const;
@@ -297,8 +269,6 @@ public:
 	ENGINE_API UWorldPartitionRuntimeCell* CreateRuntimeCell(UClass* CellClass, UClass* CellDataClass, const FString& CellName, const FString& CellInstanceSuffix, UObject* InOuter = nullptr);
 
 protected:
-	UE_DEPRECATED(5.6, "Use version that has also flags if streaming should block (bOutShouldBlock) instead.")
-	virtual EWorldPartitionStreamingPerformance GetStreamingPerformanceForCell(const UWorldPartitionRuntimeCell* Cell) const;
 	virtual EWorldPartitionStreamingPerformance GetStreamingPerformanceForCell(const UWorldPartitionRuntimeCell* Cell, bool& bOutShouldBlock) const;
 
 #if WITH_EDITORONLY_DATA
diff --git a/Engine/Source/Runtime/Engine/Public/WorldPartition/WorldPartitionRuntimeSpatialHash.h b/Engine/Source/Runtime/Engine/Public/WorldPartition/WorldPartitionRuntimeSpatialHash.h
--- a/Engine/Source/Runtime/Engine/Public/WorldPartition/WorldPartitionRuntimeSpatialHash.h
+++ b/Engine/Source/Runtime/Engine/Public/WorldPartition/WorldPartitionRuntimeSpatialHash.h
@@ -192,13 +192,6 @@ struct FSpatialHashStreamingGrid
 	ENGINE_API const FSquare2DGridHelper& GetGridHelper() const;
 	ENGINE_API float GetLoadingRange() const;
 
-	//~Begin Deprecation
-	UE_DEPRECATED(5.5, "Use version that takes FWorldPartitionStreamingContext instead.")
-	void GetCells(const TArray<FWorldPartitionStreamingSource>& Sources, UWorldPartitionRuntimeHash::FStreamingSourceCells& OutActivateCells, UWorldPartitionRuntimeHash::FStreamingSourceCells& OutLoadCells, bool bEnableZCulling) const {}
-	UE_DEPRECATED(5.5, "Use version that takes FWorldPartitionStreamingContext instead.")
-	void GetNonSpatiallyLoadedCells(TSet<const UWorldPartitionRuntimeCell*>& OutActivateCells, TSet<const UWorldPartitionRuntimeCell*>& OutLoadCells) const {}
-	//~End Deprecation
-
 #if WITH_EDITOR
 	void DumpStateLog(FHierarchicalLogArchive& Ar) const;
 #endif
@@ -320,9 +313,6 @@ public:
 	ENGINE_API virtual void Serialize(FArchive& Ar) override;
 	ENGINE_API virtual void PostLoad() override;
 
-	UE_DEPRECATED(5.4, "ApplyCVars is deprecated")
-	void ApplyCVars() {}
-
 #if WITH_EDITOR
 	ENGINE_API virtual void PostEditChangeProperty(FPropertyChangedEvent& PropertyChangedEvent) override;
 
@@ -372,8 +362,6 @@ protected:
 	ENGINE_API void ForEachStreamingGrid(TFunctionRef<void(FSpatialHashStreamingGrid&)> Func);
 	ENGINE_API void ForEachStreamingGridBreakable(TFunctionRef<bool(const FSpatialHashStreamingGrid&)> Func) const;
 	
-	UE_DEPRECATED(5.6, "Use version that has also flags if streaming should block (bOutShouldBlock) instead.")
-	ENGINE_API virtual EWorldPartitionStreamingPerformance GetStreamingPerformanceForCell(const UWorldPartitionRuntimeCell* Cell) const override;
 	ENGINE_API virtual EWorldPartitionStreamingPerformance GetStreamingPerformanceForCell(const UWorldPartitionRuntimeCell* Cell, bool& bOutShouldBlock) const override;
 
 #if WITH_EDITOR
