Relax foliage warning for converted maps by suppressing duplicate InstancedFoliageActor warning for actors with partitioned names in non-partitioned levels.
Fixes false positives when maps are converted from partitioned to non-partitioned format.

#Synced-CL 50757097

diff --git a/Engine/Source/Runtime/Engine/Private/ActorPartition/PartitionActor.cpp b/Engine/Source/Runtime/Engine/Private/ActorPartition/PartitionActor.cpp
--- a/Engine/Source/Runtime/Engine/Private/ActorPartition/PartitionActor.cpp
+++ b/Engine/Source/Runtime/Engine/Private/ActorPartition/PartitionActor.cpp
@@ -2,6 +2,7 @@
 
 #include "ActorPartition/PartitionActor.h"
 #include "ActorPartition/ActorPartitionSubsystem.h"
+#include "String/ParseTokens.h"
 #include UE_INLINE_GENERATED_CPP_BY_NAME(PartitionActor)
 
 #if WITH_EDITOR
@@ -76,7 +77,7 @@ FString APartitionActor::GetActorName(UWorld* World, const UClass* Class, const
 	return GetActorName(World, ActorPartitionId, GridSize, CellCoordsX, CellCoordsY, CellCoordsZ);
 }
 
-FString APartitionActor::GetActorName(UWorld* World, const FActorPartitionIdentifier& ActorPartitionId, uint32 GridSize,  int32 CellCoordsX, int32 CellCoordsY, int32 CellCoordsZ)
+FString APartitionActor::GetActorName(UWorld* World, const FActorPartitionIdentifier& ActorPartitionId, uint32 GridSize, int32 CellCoordsX, int32 CellCoordsY, int32 CellCoordsZ)
 {
 	TStringBuilderWithBuffer<TCHAR, NAME_SIZE> ActorNameBuilder;
 
@@ -122,4 +123,54 @@ void APartitionActor::SetLabelForActor(APartitionActor* Actor, const FActorParti
 }
 #endif
 
+// Parse the actor name to confirm that it's a partitioned actor or not (see APartitionActor::GetActorName).
+bool APartitionActor::IsPartitioned(const APartitionActor* Actor)
+{
+#if WITH_EDITOR
+	int32 NumSeparators = 0;
+	bool bValidClass = false;
+	bool bValidGridSize = !Actor->ShouldIncludeGridSizeInName(Actor->GetWorld(), { Actor->GetClass(), Actor->GetGridGuid(), 0 });
+	bool bValidGridGuid = !Actor->GetGridGuid().IsValid();
+	int32 NumValidCoords = 0;
+
+	UE::String::ParseTokens(Actor->GetName(), TEXT("_"), 
+		[&Actor, &NumSeparators, &bValidClass, &bValidGridSize, &bValidGridGuid, &NumValidCoords](FStringView InToken)
+		{
+			switch (NumSeparators)
+			{
+			case 0:
+				bValidClass = Actor->GetClass()->GetName() == InToken;
+				break;
+
+			case 1:
+				if (!bValidGridGuid)
+				{
+					FGuid Guid;
+					bValidGridGuid = FGuid::ParseExact(InToken, EGuidFormats::Base36Encoded, Guid) && (Guid == Actor->GetGridGuid());
+					break;
+				}
+			[[fallthrough]];
+
+			case 2:
+				if (bValidGridGuid && !bValidGridSize)
+				{
+					bValidGridSize = FCString::Atoi(*FString(InToken)) == Actor->GetGridSize();
+					break;
+				}
+			[[fallthrough]];
+
+			default:
+				NumValidCoords++;
+				break;
+			}
+
+			NumSeparators++;
+		});
+
+	return bValidClass && bValidGridSize && bValidGridGuid && ((NumValidCoords == 3) || (NumValidCoords == 4));
+#else
+	return false;
+#endif // WITH_EDITOR
+}
+
 #undef LOCTEXT_NAMESPACE
diff --git a/Engine/Source/Runtime/Engine/Public/ActorPartition/PartitionActor.h b/Engine/Source/Runtime/Engine/Public/ActorPartition/PartitionActor.h
--- a/Engine/Source/Runtime/Engine/Public/ActorPartition/PartitionActor.h
+++ b/Engine/Source/Runtime/Engine/Public/ActorPartition/PartitionActor.h
@@ -49,6 +49,8 @@ public:
 	static ENGINE_API FString GetActorName(UWorld* World, const FActorPartitionIdentifier& ActorPartitionId, uint32 GridSize, int32 CellCoordsX, int32 CellCoordsY, int32 CellCoordsZ);
 	static ENGINE_API void SetLabelForActor(APartitionActor* PartitionActor, const FActorPartitionIdentifier& ActorPartitionId, uint32 GridSize, int32 CellCoordsX, int32 CellCoordsY, int32 CellCoordsZ);
 #endif
+
+	static ENGINE_API bool IsPartitioned(const APartitionActor* Actor);
 };
 
 DEFINE_ACTORDESC_TYPE(APartitionActor, FPartitionActorDesc);
diff --git a/Engine/Source/Runtime/Foliage/Private/InstancedFoliage.cpp b/Engine/Source/Runtime/Foliage/Private/InstancedFoliage.cpp
--- a/Engine/Source/Runtime/Foliage/Private/InstancedFoliage.cpp
+++ b/Engine/Source/Runtime/Foliage/Private/InstancedFoliage.cpp
@@ -4509,9 +4509,9 @@ void AInstancedFoliageActor::PostLoad()
 
 	ULevel* OwningLevel = GetLevel();
 	// We can't check the ActorPartitionSubsystem here because World is not initialized yet. So we fallback on the bIsPartitioned
-	// to know if multiple InstanceFoliageActors is valid or not.
+	// to know if multiple InstanceFoliageActors is valid or not. 
 	// For levels that are World Partition runtime cells. Having multiple IFAs is valid so skip validation.
-	if (OwningLevel && !OwningLevel->bIsPartitioned && !OwningLevel->IsWorldPartitionRuntimeCell())
+	if (OwningLevel && !OwningLevel->bIsPartitioned && !OwningLevel->IsWorldPartitionRuntimeCell() && !APartitionActor::IsPartitioned(this))
 	{
 		if (!OwningLevel->InstancedFoliageActor.IsValid())
 		{
