SpringMath: Added functions for dealing with stopping and starting times and distances

#Synced-CL 50709564

diff --git a/Engine/Source/Runtime/Engine/Classes/Kismet/BlueprintSpringMathLibrary.h b/Engine/Source/Runtime/Engine/Classes/Kismet/BlueprintSpringMathLibrary.h
--- a/Engine/Source/Runtime/Engine/Classes/Kismet/BlueprintSpringMathLibrary.h
+++ b/Engine/Source/Runtime/Engine/Classes/Kismet/BlueprintSpringMathLibrary.h
@@ -418,22 +418,112 @@ public:
 	UE_EXPERIMENTAL(5.8, "SpringMath is experimental")
 	/** Estimates an approximate stopping distance from a smoothing time for a character using a simple damped spring
 	 *
-	 * @param Velocity The velocity of the character at the start of the stop
+	 * @param InitialVelocity The velocity of the character at the start of the stop
+	 * @param InitialAcceleration The acceleration of the character at the start of the stop
 	 * @param SmoothingTime The time over which to smooth velocity. It takes roughly the smoothing time in order for the character to reach the target velocity.
 	 * @return An estimation of the total distance traveled during the stop
 	 */
 	UFUNCTION(BlueprintPure, meta = (Category = "Math|Spring|Experimental", ReturnDisplayName = "Stopping Distance"))
-	static ENGINE_API float SpringCharacterSmoothingTimeToStoppingDistance(const float Velocity, const float SmoothingTime = 0.2f);
+	static ENGINE_API float SpringCharacterStoppingDistance(const float InitialVelocity, const float InitialAcceleration, const float SmoothingTime = 0.2f);
 
 	UE_EXPERIMENTAL(5.8, "SpringMath is experimental")
-	/** Estimates an approximate smoothing time from a stopping distance for a character using by a simple damped spring
+	/** Estimates an approximate stopping time (how long until the velocity goes below the VelocityThreshold) from a smoothing time for a character using a simple damped spring. Assumes the initial acceleration is zero.
 	 *
-	 * @param Velocity The velocity of the character at the start of the stop
+	 * @param InitialVelocity The velocity of the character at the start of the stop
+	 * @param SmoothingTime The time over which to smooth velocity. It takes roughly the smoothing time in order for the character to reach the target velocity.
+	 * @param VelocityThreshold The velocity threshold at which the character is considered to be stopped in cm/s
+	 * @return An estimation of the time taken to come to a stop
+	 */
+	UFUNCTION(BlueprintPure, meta = (Category = "Math|Spring|Experimental", ReturnDisplayName = "Stopping Time"))
+	static ENGINE_API float SpringCharacterStoppingTime(const float InitialVelocity, const float SmoothingTime = 0.2f, const float VelocityThreshold = 1.0f);
+
+	UE_EXPERIMENTAL(5.8, "SpringMath is experimental")
+	/** Estimates an approximate smoothing time from a stopping distance for a character using by a simple damped spring. Assumes the initial acceleration is zero.
+	 *
+	 * @param InitialVelocity The velocity of the character at the start of the stop
 	 * @param StoppingDistance The total distance traveled during the stop
 	 * @return An estimation of the SmoothingTime of the spring
 	 */
 	UFUNCTION(BlueprintPure, meta = (Category = "Math|Spring|Experimental", ReturnDisplayName = "Smoothing Time"))
-	static ENGINE_API float SpringCharacterStoppingDistanceToSmoothingTime(const float Velocity, const float StoppingDistance = 100.0f);
+	static ENGINE_API float SpringCharacterSmoothingTimeFromStoppingDistance(const float InitialVelocity, const float StoppingDistance = 100.0f);
+
+	UE_EXPERIMENTAL(5.8, "SpringMath is experimental")
+	/** Estimates an approximate smoothing time from a stopping time for a character using by a simple damped spring. Assumes the initial acceleration is zero.
+	 *
+	 * @param InitialVelocity The velocity of the character at the start of the stop
+	 * @param StoppingTime The time taken to come to a stop
+	 * @param VelocityThreshold The velocity threshold at which the character is considered to be stopped in cm/s
+	 * @return An estimation of the SmoothingTime of the spring
+	 */
+	UFUNCTION(BlueprintPure, meta = (Category = "Math|Spring|Experimental", ReturnDisplayName = "Smoothing Time"))
+	static ENGINE_API float SpringCharacterSmoothingTimeFromStoppingTime(const float InitialVelocity, const float StoppingTime = 1.0f, const float VelocityThreshold = 10.0f);
+
+	UE_EXPERIMENTAL(5.8, "SpringMath is experimental")
+	/** Estimates an approximate starting distance (distance traveled until the velocity is within the VelocityThreshold of the target) from a smoothing time for a character using a simple damped spring
+	 *
+	 * @param TargetVelocity The target velocity of the character
+	 * @param SmoothingTime The time over which to smooth velocity. It takes roughly the smoothing time in order for the character to reach the target velocity.
+	 * @param VelocityThreshold The velocity threshold from the target at which the character is considered to have started in cm/s
+	 * @return An estimation of the total distance traveled during the start
+	 */
+	UFUNCTION(BlueprintPure, meta = (Category = "Math|Spring|Experimental", ReturnDisplayName = "Starting Distance"))
+	static ENGINE_API float SpringCharacterStartingDistance(const float TargetVelocity, const float SmoothingTime = 0.2f, const float VelocityThreshold = 10.0f);
+
+	UE_EXPERIMENTAL(5.8, "SpringMath is experimental")
+	/** Estimates an approximate starting time (time until the velocity is within the VelocityThreshold of the target) from a smoothing time for a character using a simple damped spring.
+	 *
+	 * @param TargetVelocity The target velocity of the character
+	 * @param SmoothingTime The time over which to smooth velocity. It takes roughly the smoothing time in order for the character to reach the target velocity.
+	 * @param VelocityThreshold The velocity threshold from the target at which the character is considered to have started in cm/s
+	 * @return An estimation of the time taken to come to a start
+	 */
+	UFUNCTION(BlueprintPure, meta = (Category = "Math|Spring|Experimental", ReturnDisplayName = "Starting Time"))
+	static ENGINE_API float SpringCharacterStartingTime(const float TargetVelocity, const float SmoothingTime = 0.2f, const float VelocityThreshold = 10.0f);
+
+	UE_EXPERIMENTAL(5.8, "SpringMath is experimental")
+	/** Estimates an approximate smoothing time from a starting distance for a character using by a simple damped spring.
+	 *
+	 * @param TargetVelocity The target velocity of the character
+	 * @param StartingDistance The total distance traveled during the start
+	 * @param VelocityThreshold The velocity threshold from the target at which the character is considered to have started in cm/s
+	 * @return An estimation of the SmoothingTime of the spring
+	 */
+	UFUNCTION(BlueprintPure, meta = (Category = "Math|Spring|Experimental", ReturnDisplayName = "Smoothing Time"))
+	static ENGINE_API float SpringCharacterSmoothingTimeFromStartingDistance(const float TargetVelocity, const float StartingDistance = 100.0f, const float VelocityThreshold = 10.0f);
+
+	UE_EXPERIMENTAL(5.8, "SpringMath is experimental")
+	/** Estimates an approximate smoothing time from a starting time for a character using by a simple damped spring.
+	 *
+	 * @param TargetVelocity The target velocity of the character
+	 * @param StartingTime The time taken to come to a start
+	 * @param VelocityThreshold The velocity threshold from the target at which the character is considered to have started in cm/s
+	 * @return An estimation of the SmoothingTime of the spring
+	 */
+	UFUNCTION(BlueprintPure, meta = (Category = "Math|Spring|Experimental", ReturnDisplayName = "Smoothing Time"))
+	static ENGINE_API float SpringCharacterSmoothingTimeFromStartingTime(const float TargetVelocity, const float StartingTime = 1.0f, const float VelocityThreshold = 10.0f);
+
+	UE_EXPERIMENTAL(5.8, "SpringMath is experimental")
+	/** Estimates an approximate turning time (time until the angle is within the AngleThreshold of the target) from a smoothing time for a character using a simple damped spring. Assumes the initial angular velocity is zero.
+	 *
+	 * @param TargetAngle The target angle of the character in degrees
+	 * @param SmoothingTime The time over which to smooth turning. It takes roughly the smoothing time in order for the character to reach the target angle.
+	 * @param AngleThreshold The angle threshold from the target at which the character is considered to have turned in degrees
+	 * @return An estimation of the time taken to turn
+	 */
+	UFUNCTION(BlueprintPure, meta = (Category = "Math|Spring|Experimental", ReturnDisplayName = "Turning Time"))
+	static ENGINE_API float SpringCharacterTurningTime(const float TargetAngle, const float SmoothingTime = 0.2f, const float AngleThreshold = 5.0f);
+
+	UE_EXPERIMENTAL(5.8, "SpringMath is experimental")
+	/** Estimates an approximate smoothing time from a turning time for a character using by a simple damped spring.
+	 *
+	 * @param TargetAngle The target angle of the character in degrees
+	 * @param TurningTime The time taken to come to turn
+	 * @param AngleThreshold The angle threshold from the target at which the character is considered to have turned in degrees
+	 * @return An estimation of the SmoothingTime of the spring
+	 */
+	UFUNCTION(BlueprintPure, meta = (Category = "Math|Spring|Experimental", ReturnDisplayName = "Smoothing Time"))
+	static ENGINE_API float SpringCharacterSmoothingTimeFromTurningTime(const float TargetAngle, const float TurningTime = 1.0f, const float AngleThreshold = 5.0f);
+
 
 	UE_EXPERIMENTAL(5.7, "SpringMath is experimental")
 	/** Update a position representing a character given a target velocity using a velocity spring.
@@ -453,6 +543,7 @@ public:
 	static ENGINE_API void VelocitySpringCharacterUpdate(UPARAM(ref) FVector& InOutPosition, UPARAM(ref) FVector& InOutVelocity, UPARAM(ref) FVector& InOutVelocityIntermediate, UPARAM(ref) FVector& InOutAcceleration, const FVector& TargetVelocity, float DeltaTime, float SmoothingTime = 0.2f, float MaxAcceleration = 1.0f);
 
 
+
 	// Conversion Methods //
 
 	UE_EXPERIMENTAL(5.7, "SpringMath is experimental")
diff --git a/Engine/Source/Runtime/Engine/Private/BlueprintSpringMathLibrary.cpp b/Engine/Source/Runtime/Engine/Private/BlueprintSpringMathLibrary.cpp
--- a/Engine/Source/Runtime/Engine/Private/BlueprintSpringMathLibrary.cpp
+++ b/Engine/Source/Runtime/Engine/Private/BlueprintSpringMathLibrary.cpp
@@ -237,14 +237,56 @@ void UBlueprintSpringMathLibrary::SpringCharacterUpdate(FVector& InOutPosition,
 	SpringMath::SpringCharacterUpdate(InOutPosition, InOutVelocity, InOutAcceleration, TargetVelocity, SmoothingTime, DeltaTime);
 }
 
-float UBlueprintSpringMathLibrary::SpringCharacterSmoothingTimeToStoppingDistance(const float Velocity, const float SmoothingTime)
+float UBlueprintSpringMathLibrary::SpringCharacterStoppingDistance(const float InitialVelocity, const float InitialAcceleration, const float SmoothingTime)
 {
-	return SpringMath::SpringCharacterSmoothingTimeToStoppingDistance(Velocity, SmoothingTime);
+	return SpringMath::SpringCharacterStoppingDistance(InitialVelocity, InitialAcceleration, SmoothingTime);
 }
 
-float UBlueprintSpringMathLibrary::SpringCharacterStoppingDistanceToSmoothingTime(const float Velocity, const float StoppingDistance)
+float UBlueprintSpringMathLibrary::SpringCharacterStoppingTime(const float InitialVelocity, const float SmoothingTime, const float VelocityThreshold)
 {
-	return SpringMath::SpringCharacterStoppingDistanceToSmoothingTime(Velocity, StoppingDistance);
+	return SpringMath::SpringCharacterStoppingTime(InitialVelocity, SmoothingTime, VelocityThreshold);
+}
+
+float UBlueprintSpringMathLibrary::SpringCharacterSmoothingTimeFromStoppingDistance(const float InitialVelocity, const float StoppingDistance)
+{
+	return SpringMath::SpringCharacterSmoothingTimeFromStoppingDistance(InitialVelocity, StoppingDistance);
+}
+
+float UBlueprintSpringMathLibrary::SpringCharacterSmoothingTimeFromStoppingTime(const float InitialVelocity, const float StoppingTime, const float VelocityThreshold)
+{
+	return SpringMath::SpringCharacterSmoothingTimeFromStoppingTime(InitialVelocity, StoppingTime, VelocityThreshold);
+}
+
+float UBlueprintSpringMathLibrary::SpringCharacterStartingDistance(const float TargetVelocity, const float SmoothingTime, const float VelocityThreshold)
+{
+	return SpringMath::SpringCharacterStartingDistance(TargetVelocity, SmoothingTime, VelocityThreshold);
+}
+
+float UBlueprintSpringMathLibrary::SpringCharacterStartingTime(const float TargetVelocity, const float SmoothingTime, const float VelocityThreshold)
+{
+	return SpringMath::SpringCharacterStartingTime(TargetVelocity, SmoothingTime, VelocityThreshold);
+}
+
+float UBlueprintSpringMathLibrary::SpringCharacterSmoothingTimeFromStartingDistance(const float TargetVelocity, const float StartingDistance, const float VelocityThreshold)
+{
+	return SpringMath::SpringCharacterSmoothingTimeFromStartingDistance(TargetVelocity, StartingDistance, VelocityThreshold);
+}
+
+float UBlueprintSpringMathLibrary::SpringCharacterSmoothingTimeFromStartingTime(const float TargetVelocity, const float StartingTime, const float VelocityThreshold)
+{
+	return SpringMath::SpringCharacterSmoothingTimeFromStartingTime(TargetVelocity, StartingTime, VelocityThreshold);
+}
+
+float UBlueprintSpringMathLibrary::SpringCharacterTurningTime(const float TargetAngle, const float SmoothingTime, const float AngleThreshold)
+{
+	// Logic is identical for turns vs starts when it comes to times, so we just forward to the following
+	return SpringMath::SpringCharacterStartingTime(TargetAngle, SmoothingTime, AngleThreshold);
+}
+
+float UBlueprintSpringMathLibrary::SpringCharacterSmoothingTimeFromTurningTime(const float TargetAngle, const float TurningTime, const float AngleThreshold)
+{
+	// Logic is identical for turns vs starts when it comes to times, so we just forward to the following
+	return SpringMath::SpringCharacterSmoothingTimeFromStartingTime(TargetAngle, TurningTime, AngleThreshold);
 }
 
 void UBlueprintSpringMathLibrary::VelocitySpringCharacterUpdate(FVector& InOutPosition, FVector& InOutVelocity, FVector& InOutVelocityIntermediate,
diff --git a/Engine/Source/Runtime/Engine/Public/Animation/SpringMath.h b/Engine/Source/Runtime/Engine/Public/Animation/SpringMath.h
--- a/Engine/Source/Runtime/Engine/Public/Animation/SpringMath.h
+++ b/Engine/Source/Runtime/Engine/Public/Animation/SpringMath.h
@@ -65,6 +65,58 @@ private:
 		return FMath::FindDeltaAngleRadians(B, A) + B;
 	}
 
+	/** Computes the lower (-1) branch of the LambertW function, which is used for solving equations of the form x*e^x */
+	static inline float LambertWLower(const float X)
+	{
+		if (FMath::IsNaN(X)) { return X; }
+		if (X == 0.0f) { return -INFINITY; }
+		if (X > 0.0f) { return INFINITY / INFINITY; }
+		
+		const float Redx = 0.625529587f * 0.588108778f + X;
+		
+		// Polynomial approximation for when close to -(exp(-1))
+		if (Redx <= 0.09765625f)
+		{
+			const float R = FMath::Sqrt(Redx);
+			float W = -3.30250000e+2f;
+			W = W * R + 3.53563141e+2f;
+			W = W * R - 1.91617889e+2f;
+			W = W * R + 4.94172478e+1f;
+			W = W * R - 1.23464909e+1f;
+			W = W * R - 1.38704872e+0f;
+			W = W * R - 1.99431837e+0f;
+			W = W * R - 1.81044364e+0f;
+			W = W * R - 2.33166337e+0f;
+			W = W * R - 1.00000000e+0f;
+			return W;
+		}
+		// From "New approximations to the principal real - valued branch of the Lambert W function"
+		else
+		{
+			const float C0 = +1.68090820e-1f;
+			const float C1 = -2.96497345e-3f;
+			const float C2 = -2.87322998e-2f;
+			const float C3 = +7.07275391e-1f;
+
+			const float S = FMath::Log2(-X) * -UE_LN2 - 1.0f;
+			float T = FMath::Sqrt(S);
+			float W = -1.0f - S - (1.0f / (FMath::Exp2(C2 * T) * (C1 * T) + C3 / T + C0));
+
+			if (X > -0x1.0p-116f)
+			{
+				const float E = FMath::Exp2(W / UE_LN2 + 32.0f);
+				W = -(W * E - 0x1.0p32f * X) / (W * E + E) + W;
+			}
+			else
+			{
+				T = W / (1.0f + W);
+				W = T * FMath::Loge(X / W) + T;
+			}
+
+			return W;
+		}
+	}
+
 public:
 	/** Convert a smoothing time to a half life
 	 * 
@@ -619,19 +671,122 @@ public:
 		}
 	}
 
+	/** Estimates the time at which a character driven by a critical spring will reach zero acceleration */
+	UE_EXPERIMENTAL(5.8, "SpringMath is experimental")
+	static inline float SpringCharacterStoppingInflectionTime(const float InitialVelocity, const float InitialAcceleration, const float SmoothingTime)
+	{
+		const float V = FMath::Abs(InitialVelocity);
+		const float Y = SmoothingTimeToDamping(SmoothingTime) / 2.0f;
+
+		return -V / (InitialAcceleration + V * Y);
+	}
+
 	/** Estimates an approximate stopping distance from a smoothing time for a character using a simple damped spring */
 	UE_EXPERIMENTAL(5.8, "SpringMath is experimental")
-	static inline float SpringCharacterSmoothingTimeToStoppingDistance(const float Velocity, const float SmoothingTime)
+	static inline float SpringCharacterStoppingDistance(const float InitialVelocity, const float InitialAcceleration, const float SmoothingTime)
 	{
+		const float V = FMath::Abs(InitialVelocity);
+		const float A = InitialAcceleration;
 		const float Y = SmoothingTimeToDamping(SmoothingTime) / 2.0f;
-		return 2.0f * FMath::Abs(Velocity) / FMath::Max(Y, UE_SMALL_NUMBER);
+
+		if (FMath::Abs(A) < UE_SMALL_NUMBER)
+		{
+			return 2.0f * V / FMath::Max(Y, UE_SMALL_NUMBER);
+		}
+
+		const float I0 = (V * Y * 2 + A) / FMath::Max(FMath::Square(Y), UE_SMALL_NUMBER);
+
+		PRAGMA_DISABLE_EXPERIMENTAL_WARNINGS
+		const float T = SpringCharacterStoppingInflectionTime(V, A, SmoothingTime);
+		PRAGMA_ENABLE_EXPERIMENTAL_WARNINGS
+
+		if (!FMath::IsFinite(T) || T <= 0.0f)
+		{
+			return FMath::Abs(I0);
+		}
+		else
+		{
+			const float IT = (FMath::Exp(-Y * T) / FMath::Max(FMath::Square(Y), UE_SMALL_NUMBER)) * (T * A * Y + V * Y * (T * Y + 2.0f) + A);
+			return FMath::Abs(I0 - IT) + FMath::Abs(IT);
+		}
 	}
 
-	/** Estimates an approximate spring smoothing time from a stopping distance and velocity for a character using by a simple damped spring */
+	/** Estimates an approximate stopping time (the time at which the velocity goes below the VelocityThreshold) from a smoothing time and velocity for a character using a simple damped spring */
 	UE_EXPERIMENTAL(5.8, "SpringMath is experimental")
-	static inline float SpringCharacterStoppingDistanceToSmoothingTime(const float Velocity, const float StoppingDistance)
+	static inline float SpringCharacterStoppingTime(const float InitialVelocity, const float SmoothingTime, const float VelocityThreshold = 10.0f)
 	{
-		return DampingToSmoothingTime(4.0f * FMath::Abs(Velocity) / FMath::Max(StoppingDistance, UE_SMALL_NUMBER));
+		if (FMath::Abs(VelocityThreshold) >= FMath::Abs(InitialVelocity)) { return 0.0f; }
+
+		const float Y = SmoothingTimeToDamping(SmoothingTime) / 2.0f;
+
+		const float VRatio = FMath::Clamp(FMath::Abs(VelocityThreshold) / FMath::Max(FMath::Abs(InitialVelocity), UE_SMALL_NUMBER), UE_SMALL_NUMBER, 1.0f - UE_SMALL_NUMBER);
+
+		return -(LambertWLower(-VRatio / UE_EULERS_NUMBER) + 1.0f) / FMath::Max(Y, UE_SMALL_NUMBER);
+	}
+
+	/** Estimates an approximate spring smoothing time from a stopping distance and velocity for a character using a simple damped spring. Assumes the initial acceleration is zero. */
+	UE_EXPERIMENTAL(5.8, "SpringMath is experimental")
+	static inline float SpringCharacterSmoothingTimeFromStoppingDistance(const float InitialVelocity, const float StoppingDistance)
+	{
+		return DampingToSmoothingTime(4.0f * FMath::Abs(InitialVelocity) / FMath::Max(StoppingDistance, UE_SMALL_NUMBER));
+	}
+
+	/** Estimates an approximate spring smoothing time from a stopping time and velocity for a character using a simple damped spring. Assumes the initial acceleration is zero. */
+	UE_EXPERIMENTAL(5.8, "SpringMath is experimental")
+	static inline float SpringCharacterSmoothingTimeFromStoppingTime(const float InitialVelocity, const float StoppingTime, const float VelocityThreshold = 10.0f)
+	{
+		if (FMath::Abs(VelocityThreshold) >= FMath::Abs(InitialVelocity)) { return 0.0f; }
+
+		const float VRatio = FMath::Clamp(FMath::Abs(VelocityThreshold) / FMath::Max(FMath::Abs(InitialVelocity), UE_SMALL_NUMBER), UE_SMALL_NUMBER, 1.0f - UE_SMALL_NUMBER);
+
+		return DampingToSmoothingTime(2.0f * -(LambertWLower(-VRatio / UE_EULERS_NUMBER) + 1.0f) / FMath::Max(StoppingTime, UE_SMALL_NUMBER));
+	}
+
+	/** Estimates an approximate starting time (the time at which the velocity gets within the VelocityThreshold of the TargetVelocity) from a smoothing time and target velocity for a character using a simple damped spring */
+	UE_EXPERIMENTAL(5.8, "SpringMath is experimental")
+	static inline float SpringCharacterStartingTime(const float TargetVelocity, const float SmoothingTime, const float VelocityThreshold = 10.0f)
+	{
+		// Starting time is the same as if the character was moving at the target velocity and coming to a stop
+		PRAGMA_DISABLE_EXPERIMENTAL_WARNINGS
+		return SpringCharacterStoppingTime(TargetVelocity, SmoothingTime, VelocityThreshold);
+		PRAGMA_ENABLE_EXPERIMENTAL_WARNINGS
+	}
+
+	/** Estimates an approximate starting distance (the distance traveled until the velocity gets within the VelocityThreshold of the TargetVelocity) from a smoothing time for a character using a simple damped spring */
+	UE_EXPERIMENTAL(5.8, "SpringMath is experimental")
+	static inline float SpringCharacterStartingDistance(const float TargetVelocity, const float SmoothingTime, const float VelocityThreshold = 10.0f)
+	{
+		const float V = FMath::Abs(TargetVelocity);
+		const float Y = SmoothingTimeToDamping(SmoothingTime) / 2.0f;
+		PRAGMA_DISABLE_EXPERIMENTAL_WARNINGS
+		const float T = SpringCharacterStartingTime(TargetVelocity, SmoothingTime, VelocityThreshold);
+		PRAGMA_ENABLE_EXPERIMENTAL_WARNINGS
+
+		return FMath::Exp(-Y * T) * ((2.0f * V) / FMath::Max(Y, UE_SMALL_NUMBER) + V * T) + (2.0f * V) / FMath::Max(Y, UE_SMALL_NUMBER) + V * T;
+	}
+
+	/** Estimates an approximate spring smoothing time from a starting time and target velocity for a character using a simple damped spring. */
+	UE_EXPERIMENTAL(5.8, "SpringMath is experimental")
+	static inline float SpringCharacterSmoothingTimeFromStartingTime(const float TargetVelocity, const float StartingTime, const float VelocityThreshold = 10.0f)
+	{
+		// Smoothing time is the same as if the character was moving at the target velocity and coming to a stop
+		PRAGMA_DISABLE_EXPERIMENTAL_WARNINGS
+		return SpringCharacterSmoothingTimeFromStoppingTime(TargetVelocity, StartingTime, VelocityThreshold);
+		PRAGMA_ENABLE_EXPERIMENTAL_WARNINGS
+	}
+
+	/** Estimates an approximate spring smoothing time from a starting distance and target velocity for a character using a simple damped spring. */
+	UE_EXPERIMENTAL(5.8, "SpringMath is experimental")
+		static inline float SpringCharacterSmoothingTimeFromStartingDistance(const float TargetVelocity, const float StartingDistance, const float VelocityThreshold = 10.0f)
+	{
+		if (FMath::Abs(VelocityThreshold) >= FMath::Abs(TargetVelocity)) { return 0.0f; }
+
+		const float V = FMath::Abs(TargetVelocity);
+		const float VRatio = FMath::Clamp(FMath::Abs(VelocityThreshold) / FMath::Max(FMath::Abs(TargetVelocity), UE_SMALL_NUMBER), UE_SMALL_NUMBER, 1.0f - UE_SMALL_NUMBER);
+
+		const float W = -(LambertWLower(-VRatio / UE_EULERS_NUMBER) + 1.0f);
+
+		return DampingToSmoothingTime((2.0f * FMath::Exp(-W) * V * (2.0f + W) + 2.0f * V * (W - 2.0f)) / FMath::Max(StartingDistance, UE_SMALL_NUMBER));
 	}
 
 	/** Update a position representing a character given a target velocity using a velocity spring.
