SpringMath: Fixed compiler error due to bad float behavior in LambertWLower function.

#Synced-CL 50709564

diff --git a/Engine/Source/Runtime/Engine/Public/Animation/SpringMath.h b/Engine/Source/Runtime/Engine/Public/Animation/SpringMath.h
--- a/Engine/Source/Runtime/Engine/Public/Animation/SpringMath.h
+++ b/Engine/Source/Runtime/Engine/Public/Animation/SpringMath.h
@@ -70,7 +70,7 @@ private:
 	{
 		if (FMath::IsNaN(X)) { return X; }
 		if (X == 0.0f) { return -INFINITY; }
-		if (X > 0.0f) { return INFINITY / INFINITY; }
+		if (X > 0.0f) { return NAN; }
 		
 		const float Redx = 0.625529587f * 0.588108778f + X;
 		
