RemoteObjects: Solution for Montages not playing after migration.  The Animation Instance was not updating due to the AnimUpdateRate (optimization) parameters not being migrated-over.  The solution is to reconstruct them from the data.

#Synced-CL 50628995

diff --git a/Engine/Source/Runtime/Engine/Classes/Components/SkinnedMeshComponent.h b/Engine/Source/Runtime/Engine/Classes/Components/SkinnedMeshComponent.h
--- a/Engine/Source/Runtime/Engine/Classes/Components/SkinnedMeshComponent.h
+++ b/Engine/Source/Runtime/Engine/Classes/Components/SkinnedMeshComponent.h
@@ -1270,6 +1270,9 @@ public:
 	ENGINE_API virtual void PostApplyToComponent() override;
 	ENGINE_API virtual void PostReinitProperties() override;
 	ENGINE_API virtual void PostInitProperties() override;
+#if UE_WITH_REMOTE_OBJECT_HANDLE
+	ENGINE_API virtual void PostMigrate(const struct FUObjectMigrationContext& MigrationContext) override;
+#endif
 	ENGINE_API virtual void GetResourceSizeEx(FResourceSizeEx& CumulativeResourceSize) override;
 	ENGINE_API virtual FString GetDetailedInfoInternal() const override;
 #if WITH_EDITOR
diff --git a/Engine/Source/Runtime/Engine/Private/Components/SkinnedMeshComponent.cpp b/Engine/Source/Runtime/Engine/Private/Components/SkinnedMeshComponent.cpp
--- a/Engine/Source/Runtime/Engine/Private/Components/SkinnedMeshComponent.cpp
+++ b/Engine/Source/Runtime/Engine/Private/Components/SkinnedMeshComponent.cpp
@@ -57,6 +57,10 @@
 #include "Rendering/SkeletalMeshHalfEdgeBufferAccessor.h"
 #include "SkinnedMeshComponentHelper.h"
 
+#if UE_WITH_REMOTE_OBJECT_HANDLE
+#include "UObject/UObjectMigrationContext.h"
+#endif
+
 #include UE_INLINE_GENERATED_CPP_BY_NAME(SkinnedMeshComponent)
 
 DEFINE_LOG_CATEGORY_STATIC(LogSkinnedMeshComp, Log, All);
@@ -721,6 +725,28 @@ void USkinnedMeshComponent::PostInitProperties()
 	Super::PostInitProperties();
 }
 
+#if UE_WITH_REMOTE_OBJECT_HANDLE
+void USkinnedMeshComponent::PostMigrate(const struct FUObjectMigrationContext& MigrationContext)
+{
+	Super::PostMigrate(MigrationContext);
+
+	if (MigrationContext.MigrationSide == EObjectMigrationSide::Send)
+	{
+		// When migrating-out, remove us from the local non-serialized cache
+		if (AnimUpdateRateParams)
+		{
+			FAnimUpdateRateManager::CleanupUpdateRateParametersRef(this);
+			AnimUpdateRateParams = nullptr;
+		}
+	}
+	else if (ensure(MigrationContext.MigrationSide == EObjectMigrationSide::Receive))
+	{
+		// When migrating-in, reconstruct this and add it to the local non-serialized cache
+		AnimUpdateRateParams = FAnimUpdateRateManager::GetUpdateRateParameters(this);
+	}
+}
+#endif
+
 void USkinnedMeshComponent::CollectPSOPrecacheData(const FPSOPrecacheParams& BasePrecachePSOParams, FMaterialInterfacePSOPrecacheParamsList& OutParams)
 {
 	if (GetSkinnedAsset() == nullptr ||
