Fixed esoteric SA warning.

#Synced-CL 50790641

diff --git a/Engine/Source/Runtime/Engine/Private/Animation/AnimSequenceTransformProviderCompiler.cpp b/Engine/Source/Runtime/Engine/Private/Animation/AnimSequenceTransformProviderCompiler.cpp
--- a/Engine/Source/Runtime/Engine/Private/Animation/AnimSequenceTransformProviderCompiler.cpp
+++ b/Engine/Source/Runtime/Engine/Private/Animation/AnimSequenceTransformProviderCompiler.cpp
@@ -878,11 +878,11 @@ FIoHash UAnimSequenceTransformProviderData::CreateDerivedDataKeyHash(const ITarg
 	for (int32 SequenceIndex = 0; SequenceIndex < Sequences.Num(); ++SequenceIndex)
 	{
 		Writer << SequenceIndex;
-		const FAnimSequenceTransformProviderSequence& Sequence = Sequences[SequenceIndex];
-		if (IsValid(Sequence.Sequence))
+		const FAnimSequenceTransformProviderSequence& SequenceData = Sequences[SequenceIndex];
+		if (IsValid(SequenceData.Sequence))
 		{
-			FAnimationUtils::EnsureAnimSequenceLoaded(*Sequence.Sequence);
-			FIoHash SequenceHash = Sequence.Sequence->GetDerivedDataKeyHash(TargetPlatform);
+			FAnimationUtils::EnsureAnimSequenceLoaded(*SequenceData.Sequence);
+			FIoHash SequenceHash = SequenceData.Sequence->GetDerivedDataKeyHash(TargetPlatform);
 			Writer << SequenceHash;
 		}
 	}
