Mark UKismetSystemLibrary::RaiseScriptError for export so that it can be used more generally, not just by Blueprints

#rnx

#Synced-CL 50825413

diff --git a/Engine/Source/Runtime/Engine/Classes/Kismet/KismetSystemLibrary.h b/Engine/Source/Runtime/Engine/Classes/Kismet/KismetSystemLibrary.h
--- a/Engine/Source/Runtime/Engine/Classes/Kismet/KismetSystemLibrary.h
+++ b/Engine/Source/Runtime/Engine/Classes/Kismet/KismetSystemLibrary.h
@@ -127,7 +127,7 @@ class UKismetSystemLibrary : public UBlueprintFunctionLibrary
 	 * @param ErrorMessage		The error message to display
 	 */
 	UFUNCTION(BlueprintCallable, Category = "Utilities|Debugging", meta=(Keywords = "raise error", DevelopmentOnly))
-	static void RaiseScriptError(const FString& ErrorMessage = FString("An error occurred"));
+	static ENGINE_API void RaiseScriptError(const FString& ErrorMessage = FString("An error occurred"));
 
 	// Return true if the object is usable : non-null and not pending kill
 	UFUNCTION(BlueprintPure, Category = "Utilities")
