New BP function on AGameUserSettings - IsGlobalIlluminationAllowed which can be used by a game's blueprints to set the PPVolume depending on whether GI is enabled.

#Synced-CL 50792079

diff --git a/Engine/Source/Runtime/Engine/Classes/GameFramework/GameUserSettings.h b/Engine/Source/Runtime/Engine/Classes/GameFramework/GameUserSettings.h
--- a/Engine/Source/Runtime/Engine/Classes/GameFramework/GameUserSettings.h
+++ b/Engine/Source/Runtime/Engine/Classes/GameFramework/GameUserSettings.h
@@ -397,6 +397,13 @@ public:
 	UFUNCTION(BlueprintCallable, Category = Settings, meta = (DisplayName = "Is HDR Enabled"))
 	ENGINE_API bool IsHDREnabled() const;
 
+	/**
+	 * @return True if Lumen Global Illumination is allowed by scalability and supported by the device. This does not take into account whether Lumen is enabled in the project settings, or a Post Process Volume.
+	 * This is useful for the game to setup exposure depending on whether interiors can be darkened by skylight shadowing or not.
+	 */
+	UFUNCTION(BlueprintCallable, Category = Settings, meta = (DisplayName = "Is Global Illumination Allowed", WorldContext = "WorldContextObject"))
+	ENGINE_API bool IsGlobalIlluminationAllowed(UObject* WorldContextObject) const;
+
 	/** Whether to use VSync or not. (public to allow UI to connect to it) */
 	UPROPERTY(config)
 	bool bUseVSync;
diff --git a/Engine/Source/Runtime/Engine/Private/GameUserSettings.cpp b/Engine/Source/Runtime/Engine/Private/GameUserSettings.cpp
--- a/Engine/Source/Runtime/Engine/Private/GameUserSettings.cpp
+++ b/Engine/Source/Runtime/Engine/Private/GameUserSettings.cpp
@@ -18,6 +18,7 @@
 #include "HDRHelper.h"
 #include "UnrealClient.h"
 #include "ComponentRecreateRenderStateContext.h"
+#include "RenderUtils.h"
 
 // For sandboxing migration of UserSettings from AppData\Local to AppData\LocalLow.
 #if PLATFORM_WINDOWS
@@ -1199,3 +1200,17 @@ bool UGameUserSettings::IsHDREnabled() const
 {
 	return bUseHDRDisplayOutput;
 }
+
+bool UGameUserSettings::IsGlobalIlluminationAllowed(UObject* WorldContextObject) const
+{
+	UWorld* World = GEngine->GetWorldFromContextObject(WorldContextObject, EGetWorldErrorMode::ReturnNull);
+
+	const EShaderPlatform ShaderPlatform = World ? World->Scene->GetShaderPlatform() : GMaxRHIShaderPlatform;
+
+	static const auto CVarLumenAllowed = IConsoleManager::Get().FindTConsoleVariableDataInt(TEXT("r.Lumen.DiffuseIndirect.Allow"));
+
+	const bool bLumenAllowed = CVarLumenAllowed->GetValueOnGameThread()
+		&& DoesPlatformSupportLumenGI(ShaderPlatform);
+
+	return bLumenAllowed;
+}
\ No newline at end of file
