Suppress static analysis warning V568

#Synced-CL 50298075

diff --git a/Engine/Source/Runtime/Engine/Private/Streaming/SimpleStreamableAssetManager.cpp b/Engine/Source/Runtime/Engine/Private/Streaming/SimpleStreamableAssetManager.cpp
--- a/Engine/Source/Runtime/Engine/Private/Streaming/SimpleStreamableAssetManager.cpp
+++ b/Engine/Source/Runtime/Engine/Private/Streaming/SimpleStreamableAssetManager.cpp
@@ -102,7 +102,7 @@ void FSimpleStreamableAssetManager::UnregisterAsset(UStreamableRenderAsset* InAs
 
 uint32 FSimpleStreamableAssetManager::GetAllocatedSize()
 {
-	constexpr uint32 StaticSize = sizeof(Instance) + sizeof(GUseSimpleStreamableAssetManager) + sizeof(CVarUseSimpleStreamableAssetManager);
+	constexpr uint32 StaticSize = sizeof(Instance) + sizeof(GUseSimpleStreamableAssetManager) + sizeof(CVarUseSimpleStreamableAssetManager); //-V568
 	if (IsEnabled())
 	{
 		check(Instance != nullptr);
