SpringMath: Fixed to remove use of INFINITY and NAN

#Synced-CL 50709564

diff --git a/Engine/Source/Runtime/Engine/Public/Animation/SpringMath.h b/Engine/Source/Runtime/Engine/Public/Animation/SpringMath.h
--- a/Engine/Source/Runtime/Engine/Public/Animation/SpringMath.h
+++ b/Engine/Source/Runtime/Engine/Public/Animation/SpringMath.h
@@ -68,9 +68,7 @@ private:
 	/** Computes the lower (-1) branch of the LambertW function, which is used for solving equations of the form x*e^x */
 	static inline float LambertWLower(const float X)
 	{
-		if (FMath::IsNaN(X)) { return X; }
-		if (X == 0.0f) { return -INFINITY; }
-		if (X > 0.0f) { return NAN; }
+		if (X >= 0.0f) { return -UE_MAX_FLT; }
 		
 		const float Redx = 0.625529587f * 0.588108778f + X;
 		
