Fix crash on HD5500 GPUs due to missing typed UAV load support check

#Synced-CL 50783237

diff --git a/Engine/Source/Runtime/Engine/Private/Rendering/SpecularProfile.cpp b/Engine/Source/Runtime/Engine/Private/Rendering/SpecularProfile.cpp
--- a/Engine/Source/Runtime/Engine/Private/Rendering/SpecularProfile.cpp
+++ b/Engine/Source/Runtime/Engine/Private/Rendering/SpecularProfile.cpp
@@ -361,8 +361,17 @@ IPooledRenderTarget* FSpecularProfileTextureManager::GetAtlasTexture(FRDGBuilder
 
 		RDG_EVENT_SCOPE(GraphBuilder, "SpecularProfile");
 
+		// Certain platform don't support read/write on the same texture array, even if this done safely on different slices
+		const bool bPlatformSupportReadWriteOnTextureArray = !IsOpenGLPlatform(ShaderPlatform);
+
 		// 1. Create atlas texture
-		const bool bR11G11B10Supported = GPixelFormats[PF_FloatR11G11B10].Supported && UE::PixelFormat::HasCapabilities(PF_FloatR11G11B10, EPixelFormatCapabilities::TextureFilterable);
+		EPixelFormatCapabilities SpecularProfileTextureCapabilities = EPixelFormatCapabilities::TextureFilterable;
+		if (bPlatformSupportReadWriteOnTextureArray)
+		{
+			SpecularProfileTextureCapabilities |= EPixelFormatCapabilities::TypedUAVLoad;
+		}
+
+		const bool bR11G11B10Supported = GPixelFormats[PF_FloatR11G11B10].Supported && UE::PixelFormat::HasCapabilities(PF_FloatR11G11B10, SpecularProfileTextureCapabilities);
 		const EPixelFormat Format = bR11G11B10Supported ? PF_FloatR11G11B10 : PF_FloatRGBA;
 
 		FPooledRenderTargetDesc Desc(FPooledRenderTargetDesc::Create2DArrayDesc(FIntPoint(Resolution, Resolution), Format, FClearValueBinding::None, TexCreate_None, TexCreate_UAV | TexCreate_ShaderResource, false, LayerCount * SUBSTRATE_SPECULAR_PROFILE_ENTRY_COUNT));
@@ -435,8 +444,6 @@ IPooledRenderTarget* FSpecularProfileTextureManager::GetAtlasTexture(FRDGBuilder
 		// Convolve all the profiles
 		if (LayerCount > 0)
 		{
-			// Certain platform don't support read/write on the same texture array, even if this done safely on different slices
-			const bool bPlatformSupportReadWriteOnTextureArray = !IsOpenGLPlatform(ShaderPlatform);
 			FRDGTextureRef SourceSpecularProfileTexture = nullptr;
 			if (!bPlatformSupportReadWriteOnTextureArray)
 			{
