[World Asset Streaming] Add CVar to control whether HLOD flag is included in cell hash during world asset streaming registration.
Setting wp.Runtime.HashSet.WorldAssetStreaming.IncludeHLODFlagInCellHash CVar to false allows reverting to the legacy cell hash behavior.

#Synced-CL 50763424

diff --git a/Engine/Source/Runtime/Engine/Private/WorldPartition/RuntimeHashSet/WorldPartitionRuntimeHashSet.cpp b/Engine/Source/Runtime/Engine/Private/WorldPartition/RuntimeHashSet/WorldPartitionRuntimeHashSet.cpp
--- a/Engine/Source/Runtime/Engine/Private/WorldPartition/RuntimeHashSet/WorldPartitionRuntimeHashSet.cpp
+++ b/Engine/Source/Runtime/Engine/Private/WorldPartition/RuntimeHashSet/WorldPartitionRuntimeHashSet.cpp
@@ -20,6 +20,12 @@
 
 #include UE_INLINE_GENERATED_CPP_BY_NAME(WorldPartitionRuntimeHashSet)
 
+static bool GRuntimeHashSetWorldAssetStreamingIncludeHLODFlagInCellHash = true;
+static FAutoConsoleVariableRef CVarRuntimeHashSetWorldAssetStreamingIncludeHLODFlagInCellHash(
+	TEXT("wp.Runtime.HashSet.WorldAssetStreaming.IncludeHLODFlagInCellHash"),
+	GRuntimeHashSetWorldAssetStreamingIncludeHLODFlagInCellHash,
+	TEXT("Used to choose whether HLOD flag should be included in cell hash generated during world asset streaming generation."));
+
 #if WITH_EDITOR
 TObjectPtr<URuntimePartition> FRuntimePartitionDesc::GetFirstSpatiallyLoadedHLODPartitionAncestor(int32 HLODSetupsIndex)
 {
@@ -837,8 +843,14 @@ FGuid UWorldPartitionRuntimeHashSet::RegisterWorldAssetStreaming(const UWorldPar
 		// Create Cell
 		FName TargetPartitionName = TargetPartition->Name;
 		FGuid InstanceGuid = InParams.Guid;
+
 		FArchiveMD5 ArMD5;
-		ArMD5 << TargetPartitionName << InstanceGuid << bIsHLODPass;
+		ArMD5 << TargetPartitionName << InstanceGuid; 
+		if (GRuntimeHashSetWorldAssetStreamingIncludeHLODFlagInCellHash)
+		{
+			ArMD5 << bIsHLODPass;
+		}
+
 		const FGuid CellGuid = ArMD5.GetGuidFromHash();
 		check(CellGuid.IsValid());
 		if (!bIsHLODPass)
