Fixed non-deterministic padding being serialized during cook process.
- TStaticTexCoordVertexData when number of UV's is odd.
- Fixed uninitialized padding in FMeshToMeshVertData.

#Synced-CL 50790641

diff --git a/Engine/Source/Runtime/Engine/Private/Rendering/StaticMeshVertexBuffer.cpp b/Engine/Source/Runtime/Engine/Private/Rendering/StaticMeshVertexBuffer.cpp
--- a/Engine/Source/Runtime/Engine/Private/Rendering/StaticMeshVertexBuffer.cpp
+++ b/Engine/Source/Runtime/Engine/Private/Rendering/StaticMeshVertexBuffer.cpp
@@ -42,6 +42,14 @@ public:
 	{
 		uint32 NumTexCoordData = GetNumTexCoordInVertexBuffer(NumTexCoords, NumVertices);
 		TStaticMeshVertexData<VertexDataType>::ResizeBuffer(NumTexCoordData);
+
+		// When NumTexCoords is odd, GetNumTexCoordInVertexBuffer adds a padding vertex
+		// for GPU half4/float4 alignment. Zero it so DDC serialization is deterministic.
+		VertexDataType* Data = reinterpret_cast<VertexDataType*>(this->GetDataPointer());
+		for (uint32 Index = NumVertices * NumTexCoords; Index < NumTexCoordData; Index++)
+		{
+			Data[Index] = VertexDataType();
+		}
 	}
 
 	void Empty(uint32 NumVertices) override
diff --git a/Engine/Source/Runtime/Engine/Public/SkeletalMeshTypes.h b/Engine/Source/Runtime/Engine/Public/SkeletalMeshTypes.h
--- a/Engine/Source/Runtime/Engine/Public/SkeletalMeshTypes.h
+++ b/Engine/Source/Runtime/Engine/Public/SkeletalMeshTypes.h
@@ -113,7 +113,7 @@ struct FMeshToMeshVertData
 	float	 Weight = 0.0f;
 
 	// Dummy for alignment
-	uint32	 Padding;
+	uint32	 Padding = 0;
 
 	friend ENGINE_API FArchive& operator<<(FArchive& Ar, FMeshToMeshVertData& V);
 	
