UE: Improve validation message about incompatible skeleton

#Synced-CL 50810826

diff --git a/Engine/Source/Runtime/Engine/Private/SkeletalMesh.cpp b/Engine/Source/Runtime/Engine/Private/SkeletalMesh.cpp
--- a/Engine/Source/Runtime/Engine/Private/SkeletalMesh.cpp
+++ b/Engine/Source/Runtime/Engine/Private/SkeletalMesh.cpp
@@ -1731,7 +1731,7 @@ EDataValidationResult USkeletalMesh::IsDataValid(FDataValidationContext& Context
 			//Validate if the skeleton is compatible with this skeletal mesh
 			if (!GetSkeleton()->IsCompatibleMesh(this))
 			{
-				Context.AddError(LOCTEXT("SkeletalMeshValidation_IncompatibleSkeleton", "This skeletal mesh asset has an incompatible Skeleton. Assign a compatible skeleton to this skeletal mesh asset."));
+				Context.AddError(FText::Format(LOCTEXT("SkeletalMeshValidation_IncompatibleSkeleton", "This skeletal mesh asset has an incompatible Skeleton '{0}'. Assign a compatible skeleton to this skeletal mesh asset."), FText::FromString(GetPathNameSafe(GetSkeleton()))));
 				ValidationResult = EDataValidationResult::Invalid;
 			}
 		}
