Mobile: Fixed incorrect normal decoding when sampling from virtual textures on platforms using ASTC.
VT use ETC encoding and pixel shaders assume ASTC LA encoding

#Synced-CL 50816790

diff --git a/Engine/Shaders/Private/Common.ush b/Engine/Shaders/Private/Common.ush
--- a/Engine/Shaders/Private/Common.ush
+++ b/Engine/Shaders/Private/Common.ush
@@ -1705,6 +1705,16 @@ MaterialFloat4 UnpackNormalMap( MaterialFloat4 TextureSample )
 	return MaterialFloat4( NormalXY.xy, NormalZ, 1.0f );
 }
 
+MaterialFloat4 UnpackNormalMapVT(MaterialFloat4 TextureSample)
+{
+	// UnpackNormalMap expects a specific channel layout, which may differ between regular and virtual textures on certain platforms (e.g., ASTC LA)
+#if LA_NORMALMAPS
+	return UnpackNormalMap(TextureSample.rggg);
+#else
+	return UnpackNormalMap(TextureSample);
+#endif
+}
+
 // Antialiased version of a binary comparison between ThresholdConst and a texture channel.
 float AntialiasedTextureMask( Texture2D Tex, SamplerState Sampler, float2 UV, float ThresholdConst, int Channel )
 {
diff --git a/Engine/Shaders/Private/MaterialTexture.ush b/Engine/Shaders/Private/MaterialTexture.ush
--- a/Engine/Shaders/Private/MaterialTexture.ush
+++ b/Engine/Shaders/Private/MaterialTexture.ush
@@ -224,8 +224,8 @@ float4 ApplyMaterialSamplerType(float4 Value, int SamplerType)
 	case SAMPLERTYPE_LinearGrayscale:
 	case SAMPLERTYPE_VirtualLinearGrayscale: return ProcessMaterialLinearGreyscaleTextureLookup(Value);
 
-	case SAMPLERTYPE_Normal:
-	case SAMPLERTYPE_VirtualNormal: return UnpackNormalMap(Value);
+	case SAMPLERTYPE_Normal:		return UnpackNormalMap(Value);
+	case SAMPLERTYPE_VirtualNormal: return UnpackNormalMapVT(Value);
 
 	case SAMPLERTYPE_Masks:
 	case SAMPLERTYPE_VirtualMasks:
diff --git a/Engine/Source/Runtime/Engine/Private/Materials/HLSLMaterialTranslator.cpp b/Engine/Source/Runtime/Engine/Private/Materials/HLSLMaterialTranslator.cpp
--- a/Engine/Source/Runtime/Engine/Private/Materials/HLSLMaterialTranslator.cpp
+++ b/Engine/Source/Runtime/Engine/Private/Materials/HLSLMaterialTranslator.cpp
@@ -6916,11 +6916,13 @@ static FString ApplySamplerType(const FString& InSampleCode, EMaterialSamplerTyp
 	case SAMPLERTYPE_VirtualLinearGrayscale:
 		DstSampleCode = FString::Printf(TEXT("ProcessMaterialLinearGreyscaleTextureLookup(%s)"), *InSampleCode);
 		break;
-
+		
+		// Normal maps need to be unpacked in the pixel shader.
 	case SAMPLERTYPE_Normal:
+		DstSampleCode = FString::Printf(TEXT("UnpackNormalMap(%s)"), *InSampleCode);
+		break;
 	case SAMPLERTYPE_VirtualNormal:
-		// Normal maps need to be unpacked in the pixel shader.
-		DstSampleCode = FString::Printf( TEXT("UnpackNormalMap(%s)"), *InSampleCode );
+		DstSampleCode = FString::Printf( TEXT("UnpackNormalMapVT(%s)"), *InSampleCode );
 		break;
 
 	case SAMPLERTYPE_Masks:
diff --git a/Engine/Source/Runtime/Engine/Private/Materials/MaterialIRToHLSLTranslator.cpp b/Engine/Source/Runtime/Engine/Private/Materials/MaterialIRToHLSLTranslator.cpp
--- a/Engine/Source/Runtime/Engine/Private/Materials/MaterialIRToHLSLTranslator.cpp
+++ b/Engine/Source/Runtime/Engine/Private/Materials/MaterialIRToHLSLTranslator.cpp
@@ -1859,12 +1859,14 @@ struct FEntryPointCodeGen
 			case SAMPLERTYPE_VirtualLinearGrayscale:
 				Printer << TEXTVIEW("ProcessMaterialLinearGreyscaleTextureLookup");
 				break;
-
-			case SAMPLERTYPE_Normal:
-			case SAMPLERTYPE_VirtualNormal:
+				
 				// Normal maps need to be unpacked in the pixel shader.
+			case SAMPLERTYPE_Normal:
 				Printer << TEXTVIEW("UnpackNormalMap");
 				break;
+			case SAMPLERTYPE_VirtualNormal:
+				Printer << TEXTVIEW("UnpackNormalMapVT");
+				break;
 
 			case SAMPLERTYPE_Masks:
 			case SAMPLERTYPE_VirtualMasks:
