Skeletal Mesh: Implement BP functions to selectively remove morph targets and skin weight profiles.

[RN] Added Blueprint functions to remove morph targets and skin weight profiles from a mesh

#Synced-CL 50547167

diff --git a/Engine/Source/Runtime/Engine/Classes/Engine/SkeletalMesh.h b/Engine/Source/Runtime/Engine/Classes/Engine/SkeletalMesh.h
--- a/Engine/Source/Runtime/Engine/Classes/Engine/SkeletalMesh.h
+++ b/Engine/Source/Runtime/Engine/Classes/Engine/SkeletalMesh.h
@@ -575,7 +575,7 @@ public:
 	struct FCommitMeshDescriptionParams
 	{
 		FCommitMeshDescriptionParams() {}
-
+		
 		/** Ensure that there is a placeholder on the skeletal mesh to hold the morph target delta data prior to
 		 *  the mesh rebuild. Set this to off if the morph target placeholder creation is not necessary. If the
 		 *  UMorphTarget already exists, then nothing is done.
@@ -619,11 +619,11 @@ public:
 	 *  \return \c true if the commit was successful. Even if there was no mesh description to store, the commit
 	 *    is still successful.
 	 */
-	ENGINE_API bool CommitMeshDescription(const int32 InLODIndex, const FCommitMeshDescriptionParams& InParams = {});
+	ENGINE_API bool CommitMeshDescription(const int32 InLODIndex, const FCommitMeshDescriptionParams& InParams = FCommitMeshDescriptionParams());
 
 	/** Mark the bulk data as dirty so that it can be transacted for undo/redo.
 	 *  \param InLODIndex The LOD index at which to mark the bulk data dirty.
-	 *  \param bInAlwaysMarkPackageDirty Always mark the owning package dirty, even if there was not transaction open to store into.
+	 *  \param bInAlwaysMarkPackageDirty Always mark the owning package dirty, even if there was not a transaction open to store into.
 	 *  \return \c true if the bulk data was stored in the transaction buffer.
 	 */
 	ENGINE_API bool ModifyMeshDescription(const int32 InLODIndex, const bool bInAlwaysMarkPackageDirty = true);
@@ -2769,9 +2769,23 @@ public:
 
 #if WITH_EDITOR
 	/** Remove the morph targets with the specified names */
-	ENGINE_API bool RemoveMorphTargets(TConstArrayView<FName> InMorphTargetNames);
+	ENGINE_API bool RemoveMorphTargets(TConstArrayView<FName> InMorphTargetNames, bool bInRebuildRenderMesh = false);
 
+	/**
+	 * Remove a single morph target from all LODs on this mesh.
+	 * @param InMorphTargetName The name of the morph target to remove.
+	 * @param bInRebuildRenderMesh Set to true if rebuilding the render mesh as a part of the operation is desired. 
+	 *    This is needed if the result of the morph removal on the mesh is to be visible afterward. 
+	 * @return true if the morph target was successfully removed. 
+	 */
+	UFUNCTION(BlueprintCallable, Category = Mesh, meta=(DevelopmentOnly))
+	bool RemoveMorphTarget(const FName InMorphTargetName, bool bInRebuildRenderMesh = false)
+	{
+		return RemoveMorphTargets({InMorphTargetName}, bInRebuildRenderMesh);
+	}
+	
 	/** Rename the morph target named InOldName to InNewName */
+	UFUNCTION(BlueprintCallable, Category = Mesh, meta=(DevelopmentOnly))
 	ENGINE_API bool RenameMorphTarget(FName InOldName, FName InNewName);
 #endif
 
@@ -3059,6 +3073,22 @@ public:
 		PRAGMA_ENABLE_DEPRECATION_WARNINGS
 	}
 	ENGINE_API void AddSkinWeightProfile(const FSkinWeightProfileInfo& Profile);
+	
+#if WITH_EDITOR
+	ENGINE_API bool RemoveSkinWeightProfiles(TConstArrayView<FName> InProfileName, bool bInRebuildRenderMesh = false);
+ 	
+	/** Remove the named skin weight profile from this skeletal mesh. 
+	 *  @param InProfileName The name of the profile to remove.
+	 *  @param bInRebuildRenderMesh Set to true if rebuilding the render mesh as a part of the operation is desired. 
+	 *    This is needed if the result of the skin weight profile removal on the mesh is to be visible afterward. 
+	 *  @return true if the named profile was found and removed. 
+	 */
+	UFUNCTION(BlueprintCallable, Category = Mesh, meta=(DevelopmentOnly))
+	bool RemoveSkinWeightProfile(const FName InProfileName, bool bInRebuildRenderMesh = false)
+	{
+		return RemoveSkinWeightProfiles({InProfileName}, bInRebuildRenderMesh);
+	}
+#endif
 
 	int32 GetNumSkinWeightProfiles() const 
 	{ 
diff --git a/Engine/Source/Runtime/Engine/Private/SkeletalMesh.cpp b/Engine/Source/Runtime/Engine/Private/SkeletalMesh.cpp
--- a/Engine/Source/Runtime/Engine/Private/SkeletalMesh.cpp
+++ b/Engine/Source/Runtime/Engine/Private/SkeletalMesh.cpp
@@ -95,7 +95,7 @@
 
 #include UE_INLINE_GENERATED_CPP_BY_NAME(SkeletalMesh)
 
-#define LOCTEXT_NAMESPACE "SkeltalMesh"
+#define LOCTEXT_NAMESPACE "SkeletalMesh"
 
 DEFINE_LOG_CATEGORY(LogSkeletalMesh);
 
@@ -4798,7 +4798,7 @@ UMorphTarget* USkeletalMesh::FindMorphTargetAndIndex(FName MorphTargetName, int3
 }
 
 #if WITH_EDITOR
-bool USkeletalMesh::RemoveMorphTargets(TConstArrayView<FName> InMorphTargetNames)
+bool USkeletalMesh::RemoveMorphTargets(TConstArrayView<FName> InMorphTargetNames, bool bInRebuildRenderMesh)
 {
 	if(InMorphTargetNames.Num() == 0)
 	{
@@ -4819,27 +4819,38 @@ bool USkeletalMesh::RemoveMorphTargets(TConstArrayView<FName> InMorphTargetNames
 			Modify();
 			MorphTarget->Modify();
 
-			constexpr int32 LODIndex = 0;
-
-			if (HasMeshDescription(LODIndex))
+			bool bMorphTargetRemovedFromMesh = false;
+			for (int32 LODIndex = 0; LODIndex < GetNumSourceModels(); LODIndex++)
 			{
-				//Remove the morph target from the raw import data
-				FMeshDescription* MeshDescription = GetMeshDescription(LODIndex);
-				FSkeletalMeshAttributes MeshAttributes(*MeshDescription);
-
-				if (MeshAttributes.GetMorphTargetNames().Contains(MorphTargetName))
-				{
-					ModifyMeshDescription(LODIndex);
-					MeshAttributes.UnregisterMorphTargetAttribute(MorphTargetName);
-					CommitMeshDescription(LODIndex);
-				}
-				else
+				if (HasMeshDescription(LODIndex))
 				{
-					// this means that MorphTargets and MeshDescription are not synchronized (which should not happen)
-					// if the DDC is not invalidated, the MorphTargets array will be reset to its previous value in the next build. 
-					InvalidateDeriveDataCacheGUID();
+					//Remove the morph target from the raw import data
+					FMeshDescription* MeshDescription = GetMeshDescription(LODIndex);
+					FSkeletalMeshAttributes MeshAttributes(*MeshDescription);
+
+					if (MeshAttributes.GetMorphTargetNames().Contains(MorphTargetName))
+					{
+						// Only update morph target info.
+						FCommitMeshDescriptionParams CommitParams;
+						CommitParams.bUpdateMorphTargets = true;
+						CommitParams.bUpdateSkinWeightProfiles = false;
+						CommitParams.bUpdateVertexAttributes = false;
+						CommitParams.bUpdateVertexColors = false;
+						
+						ModifyMeshDescription(LODIndex);
+						MeshAttributes.UnregisterMorphTargetAttribute(MorphTargetName);
+						CommitMeshDescription(LODIndex, CommitParams);
+						bMorphTargetRemovedFromMesh = true;
+					}
 				}
 			}
+			
+			if (!bMorphTargetRemovedFromMesh)
+			{
+				// this means that MorphTargets and MeshDescription are not synchronized (which should not happen)
+				// if the DDC is not invalidated, the MorphTargets array will be reset to its previous value in the next build. 
+				InvalidateDeriveDataCacheGUID();
+			}
 
 			UnregisterMorphTarget(MorphTarget);
 
@@ -4856,6 +4867,12 @@ bool USkeletalMesh::RemoveMorphTargets(TConstArrayView<FName> InMorphTargetNames
 			LODInfoEntry.ImportedMorphTargetSourceFilename.Remove(MorphTargetName.ToString());
 		}
 	}
+	
+	if (bRemoved && bInRebuildRenderMesh)
+	{
+		PostEditChange();
+		FSkinnedAssetCompilingManager::Get().FinishCompilation({ this });		
+	}
 
 	return bRemoved;
 }
@@ -6494,6 +6511,67 @@ void USkeletalMesh::AddSkinWeightProfile(const FSkinWeightProfileInfo& Profile)
 	PRAGMA_ENABLE_DEPRECATION_WARNINGS
 }
 
+#if WITH_EDITOR
+
+bool USkeletalMesh::RemoveSkinWeightProfiles(TConstArrayView<FName> InProfileNames, bool bInRebuildRenderMesh)
+{
+	if (Algo::NoneOf(GetSkinWeightProfiles(), [InProfileNames](const FSkinWeightProfileInfo& Profile) { return InProfileNames.Contains(Profile.Name); }))
+	{
+		return false;
+	}
+	
+	FScopedTransaction Transaction(LOCTEXT("RemoveSkinWeightProfiles", "Remove Skin Weight Profiles"));
+	
+	Modify();
+
+	// Only update skin weight profile info.
+	FCommitMeshDescriptionParams CommitParams;
+	CommitParams.bUpdateMorphTargets = false;
+	CommitParams.bUpdateVertexAttributes = false;
+	CommitParams.bUpdateVertexColors = false;
+	
+	bool bSkinWeightProfileRemovedFromMesh = false;
+	for (int32 LODIndex = 0; LODIndex < GetNumSourceModels(); LODIndex++)
+	{
+		if (HasMeshDescription(LODIndex))
+		{
+			// Remove the skin weight profile from the editable mesh. The render mesh will be generated from it.
+			FMeshDescription* MeshDescription = GetMeshDescription(LODIndex);
+			FSkeletalMeshAttributes MeshAttributes(*MeshDescription);
+
+			bool bNeedsCommit = false;
+			for (const FName ProfileName: InProfileNames)
+			{
+				if (MeshAttributes.GetSkinWeightProfileNames().Contains(ProfileName))
+				{
+					if (!bNeedsCommit)
+					{
+						ModifyMeshDescription(LODIndex);
+					}
+					MeshAttributes.UnregisterSkinWeightAttribute(ProfileName);
+					bNeedsCommit = true;
+				}
+			}
+			
+			if (bNeedsCommit)
+			{
+				CommitMeshDescription(LODIndex, CommitParams);
+				bSkinWeightProfileRemovedFromMesh = true;
+			}
+		}
+	}
+	
+	if (bSkinWeightProfileRemovedFromMesh && bInRebuildRenderMesh)
+	{
+		PostEditChange();
+		FSkinnedAssetCompilingManager::Get().FinishCompilation({ this });		
+	}
+	
+	return bSkinWeightProfileRemovedFromMesh;
+}
+
+#endif
+
 
 void USkeletalMesh::ReleaseSkinWeightProfileResources()
 {
