Add groom bound visualization.

#Synced-CL 50856398

diff --git a/Engine/Plugins/Runtime/HairStrands/Source/HairStrandsCore/Private/GroomComponent.cpp b/Engine/Plugins/Runtime/HairStrands/Source/HairStrandsCore/Private/GroomComponent.cpp
--- a/Engine/Plugins/Runtime/HairStrands/Source/HairStrandsCore/Private/GroomComponent.cpp
+++ b/Engine/Plugins/Runtime/HairStrands/Source/HairStrandsCore/Private/GroomComponent.cpp
@@ -901,7 +901,16 @@ public:
 
 				#if !(UE_BUILD_SHIPPING || UE_BUILD_TEST)
 					// Render bounds
-					RenderBounds(Collector.GetPDI(ViewIndex), ViewFamily.EngineShowFlags, GetBounds(), IsSelected());
+					if (ViewMode == EGroomViewMode::Bounds)
+					{
+						FEngineShowFlags LocalShowFlags = ViewFamily.EngineShowFlags;
+						LocalShowFlags.Bounds = true;
+						RenderBounds(Collector.GetPDI(ViewIndex), LocalShowFlags, GetBounds(), true);
+					}
+					else
+					{
+						RenderBounds(Collector.GetPDI(ViewIndex), ViewFamily.EngineShowFlags, GetBounds(), IsSelected());
+					}
 				#endif
 				}
 			}
diff --git a/Engine/Source/Editor/CommonMenuExtensions/Private/GroomVisualizationMenuCommands.cpp b/Engine/Source/Editor/CommonMenuExtensions/Private/GroomVisualizationMenuCommands.cpp
--- a/Engine/Source/Editor/CommonMenuExtensions/Private/GroomVisualizationMenuCommands.cpp
+++ b/Engine/Source/Editor/CommonMenuExtensions/Private/GroomVisualizationMenuCommands.cpp
@@ -73,6 +73,7 @@ void FGroomVisualizationMenuCommands::InternalBuildVisualisationSubMenu(FMenuBui
 			}
 			Commands.AddCommandTypeToMenu(Menu, EGroomViewMode::Group);
 			Commands.AddCommandTypeToMenu(Menu, EGroomViewMode::LODColoration);
+			Commands.AddCommandTypeToMenu(Menu, EGroomViewMode::Bounds);
 			Menu.EndSection();
 		}
 
diff --git a/Engine/Source/Runtime/Engine/Private/GroomVisualizationData.cpp b/Engine/Source/Runtime/Engine/Private/GroomVisualizationData.cpp
--- a/Engine/Source/Runtime/Engine/Private/GroomVisualizationData.cpp
+++ b/Engine/Source/Runtime/Engine/Private/GroomVisualizationData.cpp
@@ -84,6 +84,7 @@ EGroomViewMode GetGroomViewMode(const FSceneView& View)
 			case 38: return EGroomViewMode::AO;
 			case 39: return EGroomViewMode::ClumpID;
 			case 40: return EGroomViewMode::Memory;
+			case 41: return EGroomViewMode::Bounds;
 			default: break;
 			}
 		}
@@ -142,6 +143,7 @@ const TCHAR* GetGroomViewModeName(EGroomViewMode In)
 	case EGroomViewMode::LODColoration:				return TEXT("LODColoration");
 	case EGroomViewMode::CardGuides:				return TEXT("CardGuides");
 	case EGroomViewMode::Memory:					return TEXT("Memory");
+	case EGroomViewMode::Bounds:					return TEXT("Bounds");
 	}
 	return TEXT("None");
 }
@@ -232,6 +234,7 @@ void FGroomVisualizationData::Initialize()
 		AddVisualizationMode(ModeMap, true, EGroomViewMode::AO,							LOCTEXT("AO", "Per-CV AO"),									LOCTEXT("AODesc", "Control point AO  (optional)"));
 		AddVisualizationMode(ModeMap, true, EGroomViewMode::ClumpID,					LOCTEXT("ClumpID", "Clump ID"),								LOCTEXT("ClumpIDDesc", "Curve clumpID  (optional)"));
 		AddVisualizationMode(ModeMap, true, EGroomViewMode::Memory,						LOCTEXT("Memory", "Memory"),								LOCTEXT("MemoryDesc", "Memory"));
+		AddVisualizationMode(ModeMap, true, EGroomViewMode::Bounds,						LOCTEXT("Bounds", "Bounds"),								LOCTEXT("BoundsDesc", "Bounds"));
 
 		ConsoleDocumentationVisualizationMode = ConfigureConsoleCommand(ModeMap);
 	}
diff --git a/Engine/Source/Runtime/Engine/Public/GroomVisualizationData.h b/Engine/Source/Runtime/Engine/Public/GroomVisualizationData.h
--- a/Engine/Source/Runtime/Engine/Public/GroomVisualizationData.h
+++ b/Engine/Source/Runtime/Engine/Public/GroomVisualizationData.h
@@ -46,6 +46,7 @@ enum class EGroomViewMode : uint8
 	CardGuides,
 	Tile,
 	Memory,
+	Bounds,
 	Count
 };
 
