[Input] Upon device connection, populate data in the UInputDeviceSubsystem about the given device if a valid FInputDeviceScope is available. This makes it a bit easier for gameplay code to query the input device subsystem when connection events happen instead of needing to wait for the user to actually press an input for it to be populated.

This is behind the "input.bShouldUpdatedDeviceSubsystemOnConnection" CVar, which is default on, in case this change in metadata causes an unexpected change in flow for a project.

#Synced-CL 50735006

diff --git a/Engine/Source/Runtime/Engine/Classes/GameFramework/InputDeviceSubsystem.h b/Engine/Source/Runtime/Engine/Classes/GameFramework/InputDeviceSubsystem.h
--- a/Engine/Source/Runtime/Engine/Classes/GameFramework/InputDeviceSubsystem.h
+++ b/Engine/Source/Runtime/Engine/Classes/GameFramework/InputDeviceSubsystem.h
@@ -217,6 +217,8 @@ public:
 	
 protected:
 
+	void HandleDeviceConnectionChange(EInputDeviceConnectionState NewState, FPlatformUserId User, FInputDeviceId DeviceId);
+	
 	/** Set the most recently used hardware device from the input processor */
 	ENGINE_API void SetMostRecentlyUsedHardwareDevice(const FInputDeviceId InDeviceId, const FHardwareDeviceIdentifier& InHardwareId);
 
diff --git a/Engine/Source/Runtime/Engine/Private/GameFramework/InputDeviceSubsystem.cpp b/Engine/Source/Runtime/Engine/Private/GameFramework/InputDeviceSubsystem.cpp
--- a/Engine/Source/Runtime/Engine/Private/GameFramework/InputDeviceSubsystem.cpp
+++ b/Engine/Source/Runtime/Engine/Private/GameFramework/InputDeviceSubsystem.cpp
@@ -4,6 +4,7 @@
 #include "GameFramework/InputDeviceProperties.h"
 #include "Framework/Application/IInputProcessor.h"
 #include "Framework/Application/SlateApplication.h"			// For RegisterInputPreProcessor
+#include "HAL/IConsoleManager.h"
 #include "Misc/App.h"	// For FApp::GetDeltaTime()
 
 #include UE_INLINE_GENERATED_CPP_BY_NAME(InputDeviceSubsystem)
@@ -16,6 +17,14 @@
 
 DEFINE_LOG_CATEGORY(LogInputDeviceProperties);
 
+namespace UE
+{
+	static bool bShouldUpdatedDeviceSubsystemOnConnection = true;
+	static FAutoConsoleVariableRef CVarShouldUpdatedDeviceSubsystemOnConnection(TEXT("input.bShouldUpdatedDeviceSubsystemOnConnection"),
+		bShouldUpdatedDeviceSubsystemOnConnection,
+		TEXT("If true, then upon receiving a device connection event, the UInputDeviceSubsystem will attempt to update the latest known metadata about a given device."));
+}
+
 ////////////////////////////////////////////////////////
 // FInputDeviceSubsystemProcessor
 
@@ -192,6 +201,8 @@ void UInputDeviceSubsystem::Initialize(FSubsystemCollectionBase& Collection)
 	InputPreprocessor = MakeShared<FInputDeviceSubsystemProcessor>();
 	FSlateApplication::Get().RegisterInputPreProcessor(InputPreprocessor, EInputPreProcessorType::Engine);
 
+	IPlatformInputDeviceMapper::Get().GetOnInputDeviceConnectionChange().AddUObject(this, &UInputDeviceSubsystem::HandleDeviceConnectionChange);
+
 #if WITH_EDITOR
 	FEditorDelegates::PreBeginPIE.AddUObject(this, &UInputDeviceSubsystem::OnPrePIEStarted);
 	FEditorDelegates::PausePIE.AddUObject(this, &UInputDeviceSubsystem::OnPIEPaused);
@@ -209,6 +220,8 @@ void UInputDeviceSubsystem::Deinitialize()
 		FSlateApplication::Get().UnregisterInputPreProcessor(InputPreprocessor);
 	}
 	InputPreprocessor.Reset();
+
+	IPlatformInputDeviceMapper::Get().GetOnInputDeviceConnectionChange().RemoveAll(this);
 }
 
 bool UInputDeviceSubsystem::ShouldCreateSubsystem(UObject* Outer) const
@@ -544,6 +557,34 @@ FHardwareDeviceIdentifier UInputDeviceSubsystem::GetInputDeviceHardwareIdentifie
 	return FHardwareDeviceIdentifier::Invalid;
 }
 
+void UInputDeviceSubsystem::HandleDeviceConnectionChange(EInputDeviceConnectionState NewState, FPlatformUserId User, FInputDeviceId DeviceId)
+{
+	if (!UE::bShouldUpdatedDeviceSubsystemOnConnection)
+	{
+		return;
+	}
+	
+	// When a new device is connected, we can attempt to populate the subsystem with its connection info.
+	// This way it is easier to build menus or respond to specific devices on their connection states
+	if (NewState != EInputDeviceConnectionState::Connected)
+	{
+		return;
+	}
+
+	// Not all systems will have a device scope available right away upon connection yet, and we need that to proceed 
+	const FInputDeviceScope* Scope = FInputDeviceScope::GetCurrent();
+	if (!Scope)
+	{
+		return;
+	}
+	
+	const FName DeviceScopeName(*Scope->HardwareDeviceIdentifier);
+	if (const FHardwareDeviceIdentifier* Hardware = UInputPlatformSettings::Get()->GetHardwareDeviceForClassName(DeviceScopeName))
+	{
+		SetMostRecentlyUsedHardwareDevice(DeviceId, *Hardware);
+	}
+}
+
 void UInputDeviceSubsystem::SetMostRecentlyUsedHardwareDevice(const FInputDeviceId InDeviceId, const FHardwareDeviceIdentifier& InHardwareId)
 {	
 	// Store or update the most recent hardware information about this FInputDeviceId.
