Fix for crash when recompiling AnimBP with Blueprint open

#Synced-CL 50821121

diff --git a/Engine/Source/Runtime/Engine/Private/Animation/AnimInstance.cpp b/Engine/Source/Runtime/Engine/Private/Animation/AnimInstance.cpp
--- a/Engine/Source/Runtime/Engine/Private/Animation/AnimInstance.cpp
+++ b/Engine/Source/Runtime/Engine/Private/Animation/AnimInstance.cpp
@@ -4938,31 +4938,33 @@ void UAnimInstance::HandleObjectsReinstanced(const TMap<UObject*, UObject*>& Old
 			USkeletalMeshComponent* MeshComponent = GetSkelMeshComponent();
 			if(MeshComponent && MeshComponent->GetSkeletalMeshAsset())
 			{
-				TRACE_OBJECT_LIFETIME_BEGIN(this);
-				// Minimally reinit proxy (i.e. dont call per-node initialization) unless we are in an editor preview world (i.e. we are in the anim BP editor)
-				UWorld* World = GetWorld();
-				if(World && World->WorldType == EWorldType::EditorPreview)
+				if (UWorld* World = GetWorld())
 				{
-					InitializeAnimation(false);
-				}
-				else
-				{
-					RecalcRequiredBones();
+					TRACE_OBJECT_LIFETIME_BEGIN(this);
+					// Minimally reinit proxy (i.e. dont call per-node initialization) unless we are in an editor preview world (i.e. we are in the anim BP editor)
+					if(World->WorldType == EWorldType::EditorPreview)
+					{
+						InitializeAnimation(false);
+					}
+					else
+					{
+						RecalcRequiredBones();
 
-					FAnimInstanceProxy& Proxy = GetProxyOnGameThread<FAnimInstanceProxy>();
-					Proxy.Initialize(this);
-					Proxy.InitializeCachedClassData();
-					Proxy.InitializeRootNode_WithRoot(Proxy.RootNode);
-				}
+						FAnimInstanceProxy& Proxy = GetProxyOnGameThread<FAnimInstanceProxy>();
+						Proxy.Initialize(this);
+						Proxy.InitializeCachedClassData();
+						Proxy.InitializeRootNode_WithRoot(Proxy.RootNode);
+					}
 
-				MeshComponent->ClearMotionVector();
+					MeshComponent->ClearMotionVector();
 
-				// Reset interactors which might have been used in the blueprint
-				for (FClothingSimulationInstance& ClothingSimulationInstance : MeshComponent->GetClothingSimulationInstances())
-				{
-					if (UClothingSimulationInteractor* const ClothingSimulationInteractor = ClothingSimulationInstance.GetClothingSimulationInteractor())
+					// Reset interactors which might have been used in the blueprint
+					for (FClothingSimulationInstance& ClothingSimulationInstance : MeshComponent->GetClothingSimulationInstances())
 					{
-						ClothingSimulationInteractor->ClothConfigUpdated();
+						if (UClothingSimulationInteractor* const ClothingSimulationInteractor = ClothingSimulationInstance.GetClothingSimulationInteractor())
+						{
+							ClothingSimulationInteractor->ClothConfigUpdated();
+						}
 					}
 				}
 			}
