Use specular profile typed UAV load bypass when typed loads on multi-channel formats aren't supported

#Synced-CL 50861786

diff --git a/Engine/Source/Runtime/Engine/Private/Rendering/SpecularProfile.cpp b/Engine/Source/Runtime/Engine/Private/Rendering/SpecularProfile.cpp
--- a/Engine/Source/Runtime/Engine/Private/Rendering/SpecularProfile.cpp
+++ b/Engine/Source/Runtime/Engine/Private/Rendering/SpecularProfile.cpp
@@ -361,18 +361,19 @@ IPooledRenderTarget* FSpecularProfileTextureManager::GetAtlasTexture(FRDGBuilder
 
 		RDG_EVENT_SCOPE(GraphBuilder, "SpecularProfile");
 
-		// Certain platform don't support read/write on the same texture array, even if this done safely on different slices
-		const bool bPlatformSupportReadWriteOnTextureArray = !IsOpenGLPlatform(ShaderPlatform);
-
 		// 1. Create atlas texture
-		EPixelFormatCapabilities SpecularProfileTextureCapabilities = EPixelFormatCapabilities::TextureFilterable;
-		if (bPlatformSupportReadWriteOnTextureArray)
-		{
-			SpecularProfileTextureCapabilities |= EPixelFormatCapabilities::TypedUAVLoad;
-		}
+		// Check if we can do typed UAV loads on pixel formats with optional UAV load support
+		const bool bR11G11B10Supported = GPixelFormats[PF_FloatR11G11B10].Supported && UE::PixelFormat::HasCapabilities(PF_FloatR11G11B10, EPixelFormatCapabilities::TextureFilterable);
+		const bool bR11G11B10TypedLoadSupported = bR11G11B10Supported && UE::PixelFormat::HasCapabilities(PF_FloatR11G11B10, EPixelFormatCapabilities::TypedUAVLoad);
+		const bool bRGBA16TypedLoadSupported = UE::PixelFormat::HasCapabilities(PF_FloatRGBA, EPixelFormatCapabilities::TypedUAVLoad);
+
+		// Certain platforms don't support read/write on the same texture array, even if this done safely on different slices.
+		// We fall back to an SRV load during convolution if we can't support this, or if neither pixel format supports typed UAV loads.
+		const bool bConvolveSupportsTypedUAVLoad = !IsOpenGLPlatform(ShaderPlatform) && (bR11G11B10TypedLoadSupported || bRGBA16TypedLoadSupported);
 
-		const bool bR11G11B10Supported = GPixelFormats[PF_FloatR11G11B10].Supported && UE::PixelFormat::HasCapabilities(PF_FloatR11G11B10, SpecularProfileTextureCapabilities);
-		const EPixelFormat Format = bR11G11B10Supported ? PF_FloatR11G11B10 : PF_FloatRGBA;
+		const EPixelFormat Format = bConvolveSupportsTypedUAVLoad ?
+			(bR11G11B10TypedLoadSupported ? PF_FloatR11G11B10 : PF_FloatRGBA) :
+			(bR11G11B10Supported ? PF_FloatR11G11B10 : PF_FloatRGBA);
 
 		FPooledRenderTargetDesc Desc(FPooledRenderTargetDesc::Create2DArrayDesc(FIntPoint(Resolution, Resolution), Format, FClearValueBinding::None, TexCreate_None, TexCreate_UAV | TexCreate_ShaderResource, false, LayerCount * SUBSTRATE_SPECULAR_PROFILE_ENTRY_COUNT));
 		GRenderTargetPool.FindFreeElement(RHICmdList, Desc, GSpecularProfileTextureAtlas, TEXT("SpecularProfileTexture"));
@@ -445,7 +446,7 @@ IPooledRenderTarget* FSpecularProfileTextureManager::GetAtlasTexture(FRDGBuilder
 		if (LayerCount > 0)
 		{
 			FRDGTextureRef SourceSpecularProfileTexture = nullptr;
-			if (!bPlatformSupportReadWriteOnTextureArray)
+			if (!bConvolveSupportsTypedUAVLoad)
 			{
 				const FIntPoint Extent = SpecularProfileTexture->Desc.Extent;
 				SourceSpecularProfileTexture = GraphBuilder.CreateTexture(SpecularProfileTexture->Desc, TEXT("ConvolvedSpecularProfileTexture"));
