[Backout] - CL50648596

Cancelling a timer in a transaction leaves the post-moved FTimerData in the Timers array. This can cause issues if the timer is cancelled twice transactionally, or cancelled transactionally when the timer callback is being run.

These are complex edge cases to fix, so I'm backing out until I can find fixes.

#rnx
Original CL Desc
-----------------------------------------------------------------
[FTimerManager + AutoRTFM] Rework to make this safe to use from both the open and closed.

We have logic that uses FTimerManager from the open and closed of a transaction. This is extremely unsafe, and will very likely cause memory corruption on abort.

This CL makes it safe to use FTimerManager from both the open and closed of a transaction, so long as the timers as distinct.

All modifications to state have been moved into helper methods that perform the modification in the open and revert the change on abort.
This will enqueue multiple on-abort handlers (more than necessary) for single method calls, but makes this change a simple set of substitutions, reducing the chance of breaking this fundamental type.

#Synced-CL 50855404

diff --git a/Engine/Source/Runtime/Engine/Private/Tests/AutoRTFM/AutoRTFMTimerManagerTests.cpp b/Engine/Source/Runtime/Engine/Private/Tests/AutoRTFM/AutoRTFMTimerManagerTests.cpp
deleted file mode 100644
diff --git a/Engine/Source/Runtime/Engine/Private/TimerManager.cpp b/Engine/Source/Runtime/Engine/Private/TimerManager.cpp
--- a/Engine/Source/Runtime/Engine/Private/TimerManager.cpp
+++ b/Engine/Source/Runtime/Engine/Private/TimerManager.cpp
@@ -5,12 +5,9 @@
 =============================================================================*/
 
 #include "TimerManager.h"
-#include "AutoRTFM.h"
 #include "Containers/StringConv.h"
 #include "Engine/GameInstance.h"
-#include "Engine/TimerHandle.h"
 #include "Logging/LogScopedCategoryAndVerbosityOverride.h"
-#include "Misc/AssertionMacros.h"
 #include "Misc/CoreDelegates.h"
 #include "Misc/StringOutputDevice.h"
 #include "ProfilingDebugging/CsvProfiler.h"
@@ -330,7 +327,6 @@ FTimerManager::FTimerManager(UGameInstance* GameInstance)
 
 FTimerManager::~FTimerManager()
 {
-	checkf(!AutoRTFM::IsTransactional(), TEXT("FTimerManager must not be destructed while inside a transaction"));
 	if (IsRunningDedicatedServer())
 	{
 		FCoreDelegates::OnHandleSystemError.RemoveAll(this);
@@ -680,7 +676,7 @@ void FTimerManager::InternalSetTimer(FTimerHandle& InOutHandle, FTimerUnifiedDel
 			NewTimerData.ExpireTime = InternalTime + FirstDelay;
 			NewTimerData.Status = ETimerStatus::Active;
 			NewTimerHandle = AddTimer(MoveTemp(NewTimerData));
-			PushToActiveTimerHeap(NewTimerHandle);
+			ActiveTimerHeap.HeapPush(NewTimerHandle, FTimerHeapOrder(Timers));
 		}
 		else
 		{
@@ -688,7 +684,7 @@ void FTimerManager::InternalSetTimer(FTimerHandle& InOutHandle, FTimerUnifiedDel
 			NewTimerData.ExpireTime = FirstDelay;
 			NewTimerData.Status = ETimerStatus::Pending;
 			NewTimerHandle = AddTimer(MoveTemp(NewTimerData));
-			AddToPendingTimerSet(NewTimerHandle);
+			PendingTimerSet.Add(NewTimerHandle);
 		}
 
 		InOutHandle = NewTimerHandle;
@@ -728,14 +724,14 @@ FTimerHandle FTimerManager::InternalSetTimerForNextTick(FTimerUnifiedDelegate&&
 	{
 		NewTimerData.Status = ETimerStatus::Active;
 		NewTimerHandle = AddTimer(MoveTemp(NewTimerData));
-		PushToActiveTimerHeap(NewTimerHandle);
+		ActiveTimerHeap.HeapPush(NewTimerHandle, FTimerHeapOrder(Timers));
 	}
 	else
 	{
 		NewTimerData.Status = ETimerStatus::Pending;
 		NewTimerData.ExpireTime = 0.0;	// Pending timers' ExpireTime means time remaining. It will be converted upon activation.
 		NewTimerHandle = AddTimer(MoveTemp(NewTimerData));
-		AddToPendingTimerSet(NewTimerHandle);
+		PendingTimerSet.Add(NewTimerHandle);
 	}
 
 	return NewTimerHandle;
@@ -752,12 +748,15 @@ void FTimerManager::InternalClearTimer(FTimerHandle InHandle)
 	switch (Data.Status)
 	{
 		case ETimerStatus::Pending:
-			RemoveFromPendingTimerSet(InHandle);
-			RemoveTimer(InHandle);
+			{
+				int32 NumRemoved = PendingTimerSet.Remove(InHandle);
+				check(NumRemoved == 1);
+				RemoveTimer(InHandle);
+			}
 			break;
 
 		case ETimerStatus::Active:
-			SetTimerStatus(Data, ETimerStatus::ActivePendingRemoval);
+			Data.Status = ETimerStatus::ActivePendingRemoval;
 			break;
 
 		case ETimerStatus::ActivePendingRemoval:
@@ -765,8 +764,11 @@ void FTimerManager::InternalClearTimer(FTimerHandle InHandle)
 			break;
 
 		case ETimerStatus::Paused:
-			RemoveFromPausedTimerSet(InHandle);
-			RemoveTimer(InHandle);
+			{
+				int32 NumRemoved = PausedTimerSet.Remove(InHandle);
+				check(NumRemoved == 1);
+				RemoveTimer(InHandle);
+			}
 			break;
 
 		case ETimerStatus::Executing:
@@ -774,7 +776,7 @@ void FTimerManager::InternalClearTimer(FTimerHandle InHandle)
 
 			// Edge case. We're currently handling this timer when it got cleared.  Clear it to prevent it firing again
 			// in case it was scheduled to fire multiple times.
-			InvalidateCurrentlyExecutingTimer();
+			CurrentlyExecutingTimer.Invalidate();
 			RemoveTimer(InHandle);
 			break;
 
@@ -874,17 +876,24 @@ void FTimerManager::PauseTimer(FTimerHandle InHandle)
 			break;
 
 		case ETimerStatus::Active:
-			RemoveFromActiveTimerHeap(InHandle);
+			{
+				int32 IndexIndex = ActiveTimerHeap.Find(InHandle);
+				check(IndexIndex != INDEX_NONE);
+				ActiveTimerHeap.HeapRemoveAt(IndexIndex, FTimerHeapOrder(Timers), EAllowShrinking::No);
+			}
 			break;
 
 		case ETimerStatus::Pending:
-			RemoveFromPendingTimerSet(InHandle);
+			{
+				int32 NumRemoved = PendingTimerSet.Remove(InHandle);
+				check(NumRemoved == 1);
+			}
 			break;
 
 		case ETimerStatus::Executing:
 			check(CurrentlyExecutingTimer == InHandle);
 
-			InvalidateCurrentlyExecutingTimer();
+			CurrentlyExecutingTimer.Invalidate();
 			break;
 
 		default:
@@ -899,20 +908,15 @@ void FTimerManager::PauseTimer(FTimerHandle InHandle)
 	else
 	{
 		// Add to Paused list
-		AddToPausedTimerSet(InHandle);
+		PausedTimerSet.Add(InHandle);
 
-		if (PreviousStatus == ETimerStatus::Pending)
-		{
-			// Set new status
-			SetTimerStatus(*TimerToPause, ETimerStatus::Paused);
-		}
-		else
+		// Set new status
+		TimerToPause->Status = ETimerStatus::Paused;
+
+		// Store time remaining in ExpireTime while paused. Don't do this if the timer is in the pending list.
+		if (PreviousStatus != ETimerStatus::Pending)
 		{
-			// Store time remaining in ExpireTime while paused.
-			// Don't do this if the timer is in the pending list.
-			const double NewExpireTime = TimerToPause->ExpireTime - InternalTime;
-			// Set new status and expire time
-			SetTimerStatusAndExpireTime(*TimerToPause, ETimerStatus::Paused, NewExpireTime);
+			TimerToPause->ExpireTime -= InternalTime;
 		}
 	}
 }
@@ -932,18 +936,18 @@ void FTimerManager::UnPauseTimer(FTimerHandle InHandle)
 	if( HasBeenTickedThisFrame() )
 	{
 		// Convert from time remaining back to a valid ExpireTime
-		const double NewExpireTime = TimerToUnPause->ExpireTime + InternalTime;
-		SetTimerStatusAndExpireTime(*TimerToUnPause, ETimerStatus::Active, NewExpireTime);
-		PushToActiveTimerHeap(InHandle);
+		TimerToUnPause->ExpireTime += InternalTime;
+		TimerToUnPause->Status = ETimerStatus::Active;
+		ActiveTimerHeap.HeapPush(InHandle, FTimerHeapOrder(Timers));
 	}
 	else
 	{
-		SetTimerStatus(*TimerToUnPause, ETimerStatus::Pending);
-		AddToPendingTimerSet(InHandle);
+		TimerToUnPause->Status = ETimerStatus::Pending;
+		PendingTimerSet.Add(InHandle);
 	}
 
 	// remove from paused list
-	RemoveFromPausedTimerSet(InHandle);
+	PausedTimerSet.Remove(InHandle);
 }
 
 FTimerData::FTimerData()
@@ -964,8 +968,6 @@ DECLARE_DWORD_COUNTER_STAT(TEXT("TimerManager Heap Size"),STAT_NumHeapEntries,ST
 
 void FTimerManager::Tick(float DeltaTime)
 {
-	AutoRTFM::UnreachableIfClosed("FTimerManager::Tick() cannot be called from the closed");
-
 	SCOPED_NAMED_EVENT(FTimerManager_Tick, FColor::Orange);
 	CSV_SCOPED_TIMING_STAT_EXCLUSIVE(TimerManager);
 
@@ -1263,86 +1265,46 @@ void FTimerManager::ListTimers() const
 
 FTimerHandle FTimerManager::AddTimer(FTimerData&& TimerData)
 {
-	const ETimerStatus Status = TimerData.Status;
-
-	FTimerHandle Handle = TimerData.Handle;
-
-	UE_AUTORTFM_OPEN
-	{
-		const void* TimerIndicesByObjectKey = TimerData.TimerDelegate.GetBoundObject();
-		TimerData.TimerIndicesByObjectKey = TimerIndicesByObjectKey;
-
-		if (Handle.IsValid())
-		{
-			check(Timers[Handle.GetIndex()].Handle == Handle);
-			Timers.EmplaceAt(Handle.GetIndex(), MoveTemp(TimerData));
-		}
-		else
-		{
-			int32 Index = Timers.Add(MoveTemp(TimerData));
-			Handle = GenerateHandle(Index);
-			Timers[Handle.GetIndex()].Handle = Handle;
-		}
+	const void* TimerIndicesByObjectKey = TimerData.TimerDelegate.GetBoundObject();
+	TimerData.TimerIndicesByObjectKey = TimerIndicesByObjectKey;
 
-		if (TimerIndicesByObjectKey)
-		{
-			TSet<FTimerHandle>& HandleSet = ObjectToTimers.FindOrAdd(TimerIndicesByObjectKey);
+	int32 NewIndex = Timers.Add(MoveTemp(TimerData));
 
-			bool bAlreadyExists = false;
-			HandleSet.Add(Handle, &bAlreadyExists);
-			checkf(!bAlreadyExists, TEXT("A timer with this handle and object has already been added! (%s)"), *GetFTimerDataSafely(TimerData));
-		}
-	};
+	FTimerHandle Result = GenerateHandle(NewIndex);
+	Timers[NewIndex].Handle = Result;
 
-	UE_AUTORTFM_ONABORT(this, Handle, Status)
+	if (TimerIndicesByObjectKey)
 	{
-		RemoveTimer(Handle);
-	};
+		TSet<FTimerHandle>& HandleSet = ObjectToTimers.FindOrAdd(TimerIndicesByObjectKey);
 
-	return Handle;
+		bool bAlreadyExists = false;
+		HandleSet.Add(Result, &bAlreadyExists);
+		checkf(!bAlreadyExists, TEXT("A timer with this handle and object has already been added! (%s)"), *GetFTimerDataSafely(TimerData));
+	}
+
+	return Result;
 }
 
 void FTimerManager::RemoveTimer(FTimerHandle Handle)
 {
-	FTimerData Data = MoveTemp(GetTimer(Handle));
-	check(Data.Handle == Handle);
+	const FTimerData& Data = GetTimer(Handle);
 
-	UE_AUTORTFM_OPEN
+	// Remove TimerIndicesByObject entry if necessary
+	if (const void* TimerIndicesByObjectKey = Data.TimerIndicesByObjectKey)
 	{
-		// Remove TimerIndicesByObject entry if necessary
-		if (const void* TimerIndicesByObjectKey = Data.TimerIndicesByObjectKey)
-		{
-			TSet<FTimerHandle>* TimersForObject = ObjectToTimers.Find(TimerIndicesByObjectKey);
-			checkf(TimersForObject, TEXT("Removed timer was bound to an object which is not tracked by ObjectToTimers! (%s)"), *GetFTimerDataSafely(Data));
+		TSet<FTimerHandle>* TimersForObject = ObjectToTimers.Find(TimerIndicesByObjectKey);
+		checkf(TimersForObject, TEXT("Removed timer was bound to an object which is not tracked by ObjectToTimers! (%s)"), *GetFTimerDataSafely(Data));
 
-			int32 NumRemoved = TimersForObject->Remove(Handle);
-			checkf(NumRemoved == 1, TEXT("Removed timer was bound to an object which is not tracked by ObjectToTimers! (%s)"), *GetFTimerDataSafely(Data));
+		int32 NumRemoved = TimersForObject->Remove(Handle);
+		checkf(NumRemoved == 1, TEXT("Removed timer was bound to an object which is not tracked by ObjectToTimers! (%s)"), *GetFTimerDataSafely(Data));
 
-			if (TimersForObject->Num() == 0)
-			{
-				ObjectToTimers.Remove(TimerIndicesByObjectKey);
-			}
+		if (TimersForObject->Num() == 0)
+		{
+			ObjectToTimers.Remove(TimerIndicesByObjectKey);
 		}
-	};
+	}
 
-	// When in the closed, postpone removal of the timer until the transaction
-	// commits. We do this to ensure that the array index is not reused within
-	// the transaction.
-	// This is done immediately when not transactional.
-	UE_AUTORTFM_ONCOMMIT(this, Handle)
-	{
-		Timers.RemoveAt(Handle.GetIndex());
-	};
-
-	// We cannot move the FTimerData into the on-abort handler as
-	// AutoRTFM::TTask requires copy-construction due to complexities around
-	// closed -> open interop, so as a workaround we need to wrap it with a
-	// shared pointer. Not ideal, but only affects the code when used under
-	// transactions.
-	UE_AUTORTFM_ONABORT(this, Data = MakeShared<FTimerData>(MoveTemp(Data)))
-	{
-		AddTimer(MoveTemp(Data.Get()));
-	};
+	Timers.RemoveAt(Handle.GetIndex());
 }
 
 bool FTimerManager::WillRemoveTimerAssert(FTimerHandle Handle) const
@@ -1381,120 +1343,6 @@ FTimerHandle FTimerManager::GenerateHandle(int32 Index)
 	return Result;
 }
 
-void FTimerManager::PushToActiveTimerHeap(FTimerHandle Handle)
-{
-	UE_AUTORTFM_ONABORT(this, Handle)
-	{
-		RemoveFromActiveTimerHeap(Handle);
-	};
-	UE_AUTORTFM_OPEN
-	{
-		ActiveTimerHeap.HeapPush(Handle, FTimerHeapOrder(Timers));
-	};
-}
-
-void FTimerManager::RemoveFromActiveTimerHeap(FTimerHandle Handle)
-{
-	UE_AUTORTFM_ONABORT(this, Handle)
-	{
-		PushToActiveTimerHeap(Handle);
-	};
-	UE_AUTORTFM_OPEN
-	{
-		int32 IndexIndex = ActiveTimerHeap.Find(Handle);
-		check(IndexIndex != INDEX_NONE);
-		ActiveTimerHeap.HeapRemoveAt(IndexIndex, FTimerHeapOrder(Timers), EAllowShrinking::No);
-	};
-}
-
-void FTimerManager::AddToPendingTimerSet(FTimerHandle Handle)
-{
-	UE_AUTORTFM_ONABORT(this, Handle)
-	{
-		RemoveFromPendingTimerSet(Handle);
-	};
-	UE_AUTORTFM_OPEN
-	{
-		PendingTimerSet.Add(Handle);
-	};
-}
-
-void FTimerManager::RemoveFromPendingTimerSet(FTimerHandle Handle)
-{
-	UE_AUTORTFM_ONABORT(this, Handle)
-	{
-		AddToPendingTimerSet(Handle);
-	};
-	UE_AUTORTFM_OPEN
-	{
-		int32 NumRemoved = PendingTimerSet.Remove(Handle);
-		check(NumRemoved == 1);
-	};
-}
-
-void FTimerManager::AddToPausedTimerSet(FTimerHandle Handle)
-{
-	UE_AUTORTFM_ONABORT(this, Handle)
-	{
-		RemoveFromPausedTimerSet(Handle);
-	};
-	UE_AUTORTFM_OPEN
-	{
-		PausedTimerSet.Add(Handle);
-	};
-}
-
-void FTimerManager::RemoveFromPausedTimerSet(FTimerHandle Handle)
-{
-	UE_AUTORTFM_ONABORT(this, Handle)
-	{
-		AddToPausedTimerSet(Handle);
-	};
-	UE_AUTORTFM_OPEN
-	{
-		int32 NumRemoved = PausedTimerSet.Remove(Handle);
-		check(NumRemoved == 1);
-	};
-}
-
-void FTimerManager::SetTimerStatus(FTimerData& Timer, ETimerStatus NewStatus)
-{
-	UE_AUTORTFM_ONABORT(this, Handle = Timer.Handle, OldStatus = Timer.Status)
-	{
-		GetTimer(Handle).Status = OldStatus;
-	};
-	UE_AUTORTFM_OPEN
-	{
-		Timer.Status = NewStatus;
-	};
-}
-
-void FTimerManager::SetTimerStatusAndExpireTime(FTimerData& Timer, ETimerStatus NewStatus, double NewExpireTime)
-{
-	UE_AUTORTFM_ONABORT(this, Handle = Timer.Handle, OldStatus = Timer.Status, OldExpireTime = Timer.ExpireTime)
-	{
-		FTimerData& Timer = GetTimer(Handle);
-		Timer.Status = OldStatus;
-		Timer.ExpireTime = OldExpireTime;
-	};
-	UE_AUTORTFM_OPEN
-	{
-		Timer.Status = NewStatus;
-		Timer.ExpireTime = NewExpireTime;
-	};
-}
-
-void FTimerManager::InvalidateCurrentlyExecutingTimer()
-{
-	UE_AUTORTFM_ONABORT(this, Handle = CurrentlyExecutingTimer)
-	{
-		CurrentlyExecutingTimer = Handle;
-	};
-	UE_AUTORTFM_OPEN
-	{
-		CurrentlyExecutingTimer.Invalidate();
-	};
-}
 
 // Handler for ListTimers console command
 static void OnListTimers(UWorld* World)
diff --git a/Engine/Source/Runtime/Engine/Public/TimerManager.h b/Engine/Source/Runtime/Engine/Public/TimerManager.h
--- a/Engine/Source/Runtime/Engine/Public/TimerManager.h
+++ b/Engine/Source/Runtime/Engine/Public/TimerManager.h
@@ -6,7 +6,6 @@
 
 #pragma once
 
-#include "AutoRTFM.h"
 #include "CoreMinimal.h"
 #include "Stats/Stats.h"
 #include "UObject/Object.h"
@@ -129,10 +128,6 @@ struct FTimerManagerTimerParameters
 
 /** 
  * Class to globally manage timers.
- *
- * FTimerManager supports transactional usage from both the open and closed, so
- * long as you do not mix open and closed calls to the same timer. FTimerManager
- * must not be destructed within a transaction.
  */
 class FTimerManager : public FNoncopyable
 {
@@ -417,7 +412,7 @@ private:
 // This should be private, but needs to be public for testing.
 public:
 	/** Generates a handle for a timer at a given index */
-	AUTORTFM_OPEN ENGINE_API FTimerHandle GenerateHandle(int32 Index);
+	ENGINE_API FTimerHandle GenerateHandle(int32 Index);
 
 // These should be private, but need to be protected so IMPLEMENT_GET_PROTECTED_FUNC works for testing.
 protected:
@@ -445,29 +440,12 @@ private:
 	}
 	ENGINE_API FTimerData& GetTimer(FTimerHandle const& InHandle);
 
-	/**
-	 * Adds a timer to the Timers list, also updating the TimerIndicesByObject map.
-	 * If TimerData.Handle is valid then the Timers array must already contain a
-	 * timer with this handle, which will be replaced with TimerData.
-	 * Returns the TimerData.Handle if valid, otherwise a new generated handle.
-	 */
+	/** Adds a timer from the Timers list, also updating the TimerIndicesByObject map.  Returns the insertion index. */
 	ENGINE_API FTimerHandle AddTimer(FTimerData&& TimerData);
-	
 	/** Removes a timer from the Timers list at the given index, also cleaning up the TimerIndicesByObject map */
 	ENGINE_API void RemoveTimer(FTimerHandle Handle);
 	ENGINE_API bool WillRemoveTimerAssert(FTimerHandle Handle) const;
 
-	/* State mutation functions with transactional rollback handling */
-	ENGINE_API void PushToActiveTimerHeap(FTimerHandle Handle);
-	ENGINE_API void RemoveFromActiveTimerHeap(FTimerHandle Handle);
-	ENGINE_API void AddToPendingTimerSet(FTimerHandle Handle);
-	ENGINE_API void RemoveFromPendingTimerSet(FTimerHandle Handle);
-	ENGINE_API void AddToPausedTimerSet(FTimerHandle Handle);
-	ENGINE_API void RemoveFromPausedTimerSet(FTimerHandle Handle);
-	ENGINE_API void SetTimerStatus(FTimerData& Timer, ETimerStatus NewStatus);
-	ENGINE_API void SetTimerStatusAndExpireTime(FTimerData& Timer, ETimerStatus NewStatus, double NewExpireTime);
-	ENGINE_API void InvalidateCurrentlyExecutingTimer();
-
 	/** The array of timers - all other arrays will index into this */
 	TSparseArray<FTimerData> Timers;
 	/** Heap of actively running timers. */
