Skeletal Mesh: Move ref skeleton rebuild to PostLoad

#rnx

#Synced-CL 50875397

diff --git a/Engine/Source/Runtime/Engine/Private/SkeletalMesh.cpp b/Engine/Source/Runtime/Engine/Private/SkeletalMesh.cpp
--- a/Engine/Source/Runtime/Engine/Private/SkeletalMesh.cpp
+++ b/Engine/Source/Runtime/Engine/Private/SkeletalMesh.cpp
@@ -2022,12 +2022,6 @@ void USkeletalMesh::Serialize( FArchive& Ar )
 
 	Ar << GetRefSkeleton();
 
-	if (Ar.IsLoading())
-	{
-		const bool bRebuildNameMap = false;
-		GetRefSkeleton().RebuildRefSkeleton(GetSkeleton(), bRebuildNameMap);
-	}
-
 #if WITH_EDITORONLY_DATA
 	// Serialize the source model (if we want editor data)
 	if (!StripFlags.IsEditorDataStripped())
@@ -4093,8 +4087,17 @@ void USkeletalMesh::DeclareConstructClasses(TArray<FTopLevelAssetPath>& OutConst
 void USkeletalMesh::PostLoad()
 {
 	LLM_SCOPE(ELLMTag::SkeletalMesh);
+	
+	// Rebuild the ref skeleton to match the skeleton
+	if (GetSkeleton())
+	{
+		GetSkeleton()->ConditionalPostLoad();
+	
+		constexpr bool bRebuildNameMap = false;
+		GetRefSkeleton().RebuildRefSkeleton(GetSkeleton(), bRebuildNameMap);
+	}
+	
 	Super::PostLoad();
-
 }
 
 void USkeletalMesh::ExecutePostLoadInternal(FSkinnedAssetPostLoadContext& Context)
