Addresses mirrored blend samples not working properly with sync groups

- Added a new signature of CollectMarkersInRange() that handles mirroring.
- Ensured any time a passed marker was recorded we ensure we record the mirrored version if that is the case.
- Removed deprecated signature of GetMarkerSyncPositionfromMarkerIndicies().

#Synced-CL 50831251

diff --git a/Engine/Source/Runtime/Engine/Classes/Animation/BlendSpace.h b/Engine/Source/Runtime/Engine/Classes/Animation/BlendSpace.h
--- a/Engine/Source/Runtime/Engine/Classes/Animation/BlendSpace.h
+++ b/Engine/Source/Runtime/Engine/Classes/Animation/BlendSpace.h
@@ -758,6 +758,9 @@ private:
 	/** Gets the appropriate mirror data table for tick record's blend sample. */
 	const UMirrorDataTable* GetEffectiveMirrorDataTableForSample(const FAnimTickRecord& InInstance, const FBlendSample& InBlendSample) const;
 	
+	/** Gets the appropriate mirror data table for a blend sample taking into account any upstream mirror requests */
+	const UMirrorDataTable* GetEffectiveMirrorDataTableForSample(const FBlendSample& InBlendSample, const UMirrorDataTable* InUpStreamMirrorDataTable = nullptr) const;
+	
 public:
 #if WITH_EDITORONLY_DATA
 	UE_DEPRECATED(5.1, "This property is deprecated. Please use/see bContainsRotationOffsetMeshSpaceSamples instead")
diff --git a/Engine/Source/Runtime/Engine/Private/Animation/AnimMontage.cpp b/Engine/Source/Runtime/Engine/Private/Animation/AnimMontage.cpp
--- a/Engine/Source/Runtime/Engine/Private/Animation/AnimMontage.cpp
+++ b/Engine/Source/Runtime/Engine/Private/Animation/AnimMontage.cpp
@@ -2606,7 +2606,7 @@ void FAnimMontageInstance::Advance(float DeltaTime, struct FRootMotionMovementPa
 				{
 					if (bDidUseMarkerSyncThisTick)
 					{
-						Montage->MarkerData.CollectMarkersInRange(PreviousSubStepPosition, Position, MarkersPassedThisTick, SubStepDeltaMove);
+						Montage->MarkerData.CollectMarkersInRange(PreviousSubStepPosition, Position, MarkersPassedThisTick, SubStepDeltaMove, nullptr);
 					}
 
 					// Extract Root Motion for this time slice, and accumulate it.
diff --git a/Engine/Source/Runtime/Engine/Private/Animation/AnimSequence.cpp b/Engine/Source/Runtime/Engine/Private/Animation/AnimSequence.cpp
--- a/Engine/Source/Runtime/Engine/Private/Animation/AnimSequence.cpp
+++ b/Engine/Source/Runtime/Engine/Private/Animation/AnimSequence.cpp
@@ -291,6 +291,19 @@ FString GetAnimSequenceSpecificCacheKeySuffix(const UAnimSequence& Seq, bool bPe
 }
 #endif
 
+FName MarkerOrMirroredName(FName MarkerToFind, const UMirrorDataTable* MirrorTable)
+{
+	if (MirrorTable)
+	{
+		const FName* MirroredName = MirrorTable->SyncToMirrorSyncMap.Find(MarkerToFind);
+		if (MirroredName)
+		{
+			return *MirroredName;
+		}
+	}
+	return MarkerToFind;
+}
+
 /////////////////////////////////////////////////////
 // FRawAnimSequenceTrackNativeDeprecated
 
@@ -3399,7 +3412,7 @@ void UAnimSequence::AdvanceMarkerPhaseAsLeader(bool bLooping, float MoveDelta, c
 
 				// Record that we just passed a marker.
 				const int32 PassedMarker = MarkersPassed.Add(FPassedMarker());
-				MarkersPassed[PassedMarker].PassedMarkerName = NextSyncMarker.MarkerName;
+				MarkersPassed[PassedMarker].PassedMarkerName = MarkerOrMirroredName(NextSyncMarker.MarkerName, MirrorTable);
 				MarkersPassed[PassedMarker].DeltaTimeWhenPassed = CurrentMoveDelta;
 				
 				// Compute our new next marker.
@@ -3486,7 +3499,7 @@ void UAnimSequence::AdvanceMarkerPhaseAsLeader(bool bLooping, float MoveDelta, c
 
 				// Record that we just passed a marker.
 				const int32 PassedMarker = MarkersPassed.Add(FPassedMarker());
-				MarkersPassed[PassedMarker].PassedMarkerName = PrevSyncMarker.MarkerName;
+				MarkersPassed[PassedMarker].PassedMarkerName = MarkerOrMirroredName(PrevSyncMarker.MarkerName, MirrorTable);
 				MarkersPassed[PassedMarker].DeltaTimeWhenPassed = CurrentMoveDelta;
 
 				// Compute our new previous marker.
@@ -3539,19 +3552,6 @@ void UAnimSequence::AdvanceMarkerPhaseAsLeader(bool bLooping, float MoveDelta, c
 	check(CurrentTime >= 0.f && CurrentTime <= GetPlayLength());
 }
 
-FName MarkerOrMirroredName(FName MarkerToFind, const UMirrorDataTable* MirrorTable)
-{
-	if (MirrorTable)
-	{
-		const FName* MirroredName = MirrorTable->SyncToMirrorSyncMap.Find(MarkerToFind);
-		if (MirroredName)
-		{
-			return *MirroredName;
-		}
-	}
-	return MarkerToFind;
-}
-
 void AdvanceMarkerForwards(int32& Marker, FName MarkerToFind, bool bLooping, const TArray<FAnimSyncMarker>& AuthoredSyncMarkers, const UMirrorDataTable* MirrorTable)
 {
 	int32 MaxIterations = AuthoredSyncMarkers.Num();
diff --git a/Engine/Source/Runtime/Engine/Private/Animation/AnimTypes.cpp b/Engine/Source/Runtime/Engine/Private/Animation/AnimTypes.cpp
--- a/Engine/Source/Runtime/Engine/Private/Animation/AnimTypes.cpp
+++ b/Engine/Source/Runtime/Engine/Private/Animation/AnimTypes.cpp
@@ -228,11 +228,6 @@ void FMarkerSyncData::GetMarkerIndicesForTime(float CurrentTime, bool bLooping,
 	}
 }
 
-FMarkerSyncAnimPosition FMarkerSyncData::GetMarkerSyncPositionfromMarkerIndicies(int32 PrevMarker, int32 NextMarker, float CurrentTime, float SequenceLength) const
-{
-	return GetMarkerSyncPositionFromMarkerIndicies(PrevMarker, NextMarker, CurrentTime, SequenceLength, nullptr);
-}
-
 FMarkerSyncAnimPosition FMarkerSyncData::GetMarkerSyncPositionFromMarkerIndicies(int32 PrevMarker, int32 NextMarker, float CurrentTime, float SequenceLength, const UMirrorDataTable* MirrorTable) const
 {
 	FMarkerSyncAnimPosition SyncPosition;
@@ -322,13 +317,20 @@ void FMarkerSyncData::CollectUniqueNames()
 
 void FMarkerSyncData::CollectMarkersInRange(float PrevPosition, float NewPosition, 	TArray<FPassedMarker>& OutMarkersPassedThisTick, float TotalDeltaMove)
 {
-	for (const auto& Marker : AuthoredSyncMarkers)
+	CollectMarkersInRange(PrevPosition, NewPosition, OutMarkersPassedThisTick, TotalDeltaMove, nullptr);
+}
+
+extern FName MarkerOrMirroredName(FName MarkerToFind, const UMirrorDataTable* MirrorTable);
+
+void FMarkerSyncData::CollectMarkersInRange(float PrevPosition, float NewPosition, TArray<FPassedMarker>& OutMarkersPassedThisTick, float TotalDeltaMove, const UMirrorDataTable* MirrorDataTable)
+{
+	for (const FAnimSyncMarker& Marker : AuthoredSyncMarkers)
 	{
 		if (Marker.Time >= PrevPosition && Marker.Time < NewPosition)
 		{
 			float TimeToMarker = Marker.Time - PrevPosition;
 			int32 PassedMarker = OutMarkersPassedThisTick.Add(FPassedMarker());
-			OutMarkersPassedThisTick[PassedMarker].PassedMarkerName = Marker.MarkerName;
+			OutMarkersPassedThisTick[PassedMarker].PassedMarkerName = MarkerOrMirroredName(Marker.MarkerName, MirrorDataTable);
 			OutMarkersPassedThisTick[PassedMarker].DeltaTimeWhenPassed = TotalDeltaMove - TimeToMarker;
 		}
 	}
diff --git a/Engine/Source/Runtime/Engine/Private/Animation/BlendSpace.cpp b/Engine/Source/Runtime/Engine/Private/Animation/BlendSpace.cpp
--- a/Engine/Source/Runtime/Engine/Private/Animation/BlendSpace.cpp
+++ b/Engine/Source/Runtime/Engine/Private/Animation/BlendSpace.cpp
@@ -302,13 +302,18 @@ bool UBlendSpace::UpdateBlendSamples_Internal(
 
 const UMirrorDataTable* UBlendSpace::GetEffectiveMirrorDataTableForSample(const FAnimTickRecord& InInstance, const FBlendSample& InBlendSample) const
 {
-	const bool bInstanceIsMirrored = InInstance.MirrorDataTable != nullptr;
+	return GetEffectiveMirrorDataTableForSample(InBlendSample, InInstance.MirrorDataTable);
+}
+
+const UMirrorDataTable* UBlendSpace::GetEffectiveMirrorDataTableForSample(const FBlendSample& InBlendSample, const UMirrorDataTable* InUpStreamMirrorDataTable) const
+{
+	const bool bInstanceIsMirrored = InUpStreamMirrorDataTable != nullptr;
 	const bool bSampleIsMirrored = InBlendSample.bMirror;
 
 	// We have an upstream mirror request (via graph message), and our sample is unmirrored, use graph message's mirror data table as usual.
 	if (bInstanceIsMirrored && !bSampleIsMirrored)
 	{
-		return InInstance.MirrorDataTable;
+		return InUpStreamMirrorDataTable;
 	}
 
 	// No upstream mirroring request (via graph message) therefore we are good to use the blendspace's mirror data table since our sample is mirrored.
@@ -321,6 +326,7 @@ const UMirrorDataTable* UBlendSpace::GetEffectiveMirrorDataTableForSample(const
 	return nullptr;
 }
 
+
 bool UBlendSpace::UpdateBlendSamples(const FVector& InBlendSpacePosition, float InDeltaTime, TArray<FBlendSampleData>& InOutSampleDataCache, int32& InOutCachedTriangulationIndex) const
 {
 	TArray<FBlendSampleData>& OldSampleDataList = FBlendSpaceScratchData::Get().OldSampleDataList;
@@ -391,7 +397,7 @@ void UBlendSpace::ResetBlendSamples(TArray<FBlendSampleData>& InOutSampleDataCac
 			FMarkerSyncData SyncData;
 			
 			SyncData.AuthoredSyncMarkers = LeaderSample.Animation->AuthoredSyncMarkers;
-			SyncData.CollectMarkersInRange(0, LeaderSampleData.Time, Context.MarkerTickContext.MarkersPassedThisTick, LeaderSampleData.Time);
+			SyncData.CollectMarkersInRange(0, LeaderSampleData.Time, Context.MarkerTickContext.MarkersPassedThisTick, LeaderSampleData.Time, nullptr);
 		}
 		
 		// Reset follower samples.
@@ -681,7 +687,7 @@ void UBlendSpace::TickAssetPlayer(FAnimTickRecord& Instance, struct FAnimNotifyQ
 						// Keep track of the markers we passed.
 						FMarkerSyncData SyncData;
 						SyncData.AuthoredSyncMarkers = LeaderSample.Animation->AuthoredSyncMarkers;
-						SyncData.CollectMarkersInRange(LeaderSampleData.PreviousTime, LeaderSampleData.Time, Context.MarkerTickContext.MarkersPassedThisTick, NewDelta);
+						SyncData.CollectMarkersInRange(LeaderSampleData.PreviousTime, LeaderSampleData.Time, Context.MarkerTickContext.MarkersPassedThisTick, NewDelta, LeaderMirrorDataTable);
 					}
 				}
 				else
diff --git a/Engine/Source/Runtime/Engine/Public/Animation/AnimTypes.h b/Engine/Source/Runtime/Engine/Public/Animation/AnimTypes.h
--- a/Engine/Source/Runtime/Engine/Public/Animation/AnimTypes.h
+++ b/Engine/Source/Runtime/Engine/Public/Animation/AnimTypes.h
@@ -645,11 +645,13 @@ struct FMarkerSyncData
 
 	void GetMarkerIndicesForTime(float CurrentTime, bool bLooping, const TArray<FName>& ValidMarkerNames, FMarkerPair& OutPrevMarker, FMarkerPair& OutNextMarker, float SequenceLength) const;
 	
-	UE_DEPRECATED(5.0, "Use other GetMarkerSyncPositionFromMarkerIndicies signature")
-	FMarkerSyncAnimPosition GetMarkerSyncPositionfromMarkerIndicies(int32 PrevMarker, int32 NextMarker, float CurrentTime, float SequenceLength) const;
 	FMarkerSyncAnimPosition GetMarkerSyncPositionFromMarkerIndicies(int32 PrevMarker, int32 NextMarker, float CurrentTime, float SequenceLength, const UMirrorDataTable* MirrorTable = nullptr) const;
+	
 	void CollectUniqueNames();
+	
+	UE_DEPRECATED(5.8, "Use flavor of CollectMarkersInRange() that takes in an UMirrorDataTable")
 	void CollectMarkersInRange(float PrevPosition, float NewPosition, TArray<FPassedMarker>& OutMarkersPassedThisTick, float TotalDeltaMove);
+	void CollectMarkersInRange(float PrevPosition, float NewPosition, TArray<FPassedMarker>& OutMarkersPassedThisTick, float TotalDeltaMove, const UMirrorDataTable* MirrorDataTable);
 };
 
 // Shortcut for the allocator used by animation nodes.
