Add a Clear() method to USkeletalMesh, to help prepare a USkeletalMesh for re-use w/ entirely new mesh/skeleton/etc data

#Synced-CL 50894192

diff --git a/Engine/Source/Runtime/Engine/Classes/Engine/SkeletalMesh.h b/Engine/Source/Runtime/Engine/Classes/Engine/SkeletalMesh.h
--- a/Engine/Source/Runtime/Engine/Classes/Engine/SkeletalMesh.h
+++ b/Engine/Source/Runtime/Engine/Classes/Engine/SkeletalMesh.h
@@ -646,6 +646,10 @@ public:
 	 */
 	ENGINE_API void ClearMeshDescriptionAndBulkData(const int32 InLODIndex);
 
+
+	/** Generally clear data (source models, skeleton, physics, materials, etc) off the skeletal mesh */
+	ENGINE_API void Clear();
+
 	//////////////////////////////////////////////////////////////////////////
 	// USkeletalMeshEditorData public skeletalmesh API
 	// We do not want skeletal mesh client to use directly the asset(function GetMeshEditorData)
diff --git a/Engine/Source/Runtime/Engine/Private/SkeletalMesh.cpp b/Engine/Source/Runtime/Engine/Private/SkeletalMesh.cpp
--- a/Engine/Source/Runtime/Engine/Private/SkeletalMesh.cpp
+++ b/Engine/Source/Runtime/Engine/Private/SkeletalMesh.cpp
@@ -3175,6 +3175,43 @@ void USkeletalMesh::ClearMeshDescriptionAndBulkData(const int32 InLODIndex)
 	return GetSourceModel(InLODIndex).ClearAllMeshData();
 }
 
+void USkeletalMesh::Clear()
+{
+	if (!GetMorphTargets().IsEmpty())
+	{
+		UnregisterAllMorphTarget();
+	}
+
+#if WITH_EDITORONLY_DATA
+	if (FSkeletalMeshModel* ImportedResource = GetImportedModel())
+	{
+		ImportedResource->EmptyOriginalReductionSourceMeshData();
+		ImportedResource->LODModels.Empty();
+		ImportedResource->InlineReductionCacheDatas.Empty();
+	}
+#endif
+
+	SetNumSourceModels(0);
+	GetMaterials().Empty();
+	GetRefSkeleton().Empty();
+	SetSkeleton(nullptr);
+	SetPhysicsAsset(nullptr);
+	GetRefBasesInvMatrix().Empty();
+	GetCachedComposedRefPoseMatrices().Empty();
+	GetMeshOnlySocketList().Empty();
+	GetNodeMappingData().Empty();
+	GetMeshClothingAssets().Empty();
+	SetHasActiveClothingAssets(false);
+	GetSkinWeightProfiles().Empty();
+
+	ReleaseResources();
+#if WITH_EDITOR
+	InvalidateDeriveDataCacheGUID();
+#else
+	RebuildSocketMap();
+#endif
+}
+
 
 void USkeletalMesh::LoadLODImportedData(const int32 LODIndex, FSkeletalMeshImportData& OutMesh) const
 {
