Remove `AUTORTFM_DISABLE` from `FAttachmentCounter::Decrement` and switch it to a runtime assert.

We haven't done the work yet to work through the actual renderthread only code to disable for transactions.

#Synced-CL 50903364

diff --git a/Engine/Source/Runtime/Engine/Public/PrimitiveSceneInfoData.h b/Engine/Source/Runtime/Engine/Public/PrimitiveSceneInfoData.h
--- a/Engine/Source/Runtime/Engine/Public/PrimitiveSceneInfoData.h
+++ b/Engine/Source/Runtime/Engine/Public/PrimitiveSceneInfoData.h
@@ -39,9 +39,10 @@ struct FAttachmentCounter final
 		};
 	}
 
-	AUTORTFM_DISABLE void Decrement()
+	void Decrement()
 	{
 		// TODO: should we `check(IsInRenderingThread());`?
+		AutoRTFM::UnreachableIfClosed("Should only be decremented on the render thread in the open");
 		Payload->Payload--;
 	}
 
