[ShaderCompiler]
- Make sure to update remaining shader count when entering function so that we properly emit a broadcast when all shaders are done, even if all processed inside the same tick

diff --git a/Engine/Source/Runtime/Engine/Private/ShaderCompiler/ShaderCompiler.cpp b/Engine/Source/Runtime/Engine/Private/ShaderCompiler/ShaderCompiler.cpp
--- a/Engine/Source/Runtime/Engine/Private/ShaderCompiler/ShaderCompiler.cpp
+++ b/Engine/Source/Runtime/Engine/Private/ShaderCompiler/ShaderCompiler.cpp
@@ -2823,6 +2823,10 @@ void FShaderCompilingManager::FinishCompilation(const TCHAR* MaterialName, const
 
 void FShaderCompilingManager::FinishAllCompilation()
 {
+	// No-op if nothing changed since last call, but required at the beginning for LastNumRemainingAssets to
+	// be updated so that the subsequent call can notice that where was something to compile and the amount is back to 0.
+	UpdateNumRemainingAssets();
+
 #if WITH_EDITOR
 	// This is here for backward compatibility since textures are most probably expected to be ready too.
 	FTextureCompilingManager::Get().FinishAllCompilation();
@@ -2850,6 +2854,8 @@ void FShaderCompilingManager::FinishAllCompilation()
 
 	const double EndTime = FPlatformTime::Seconds();
 
+	UpdateNumRemainingAssets();
+
 	UE_LOG(LogShaders, Verbose, TEXT("FinishAllCompilation %.3fs"), (float)(EndTime - StartTime));
 }
 
@@ -2888,6 +2894,10 @@ void FShaderCompilingManager::ProcessAsyncResults(float TimeSlice, bool bBlockOn
 
 	const double StartTime = FPlatformTime::Seconds();
 
+	// No-op if nothing changed since last call, but required at the beginning for LastNumRemainingAssets to
+	// be updated so that the subsequent call can notice that where was something to compile and the amount is back to 0.
+	UpdateNumRemainingAssets();
+
 	// Some controllers need to be manually ticked if the engine loop is not initialized or blocked
 	// to do things like tick the HTTPModule.
 	// Otherwise the results from the controller will never be processed.
