[WorldPartition] Misc deprecation cleanup

#rnx
#Synced-CL 50906061

diff --git a/Engine/Source/Runtime/Engine/Public/WorldPartition/ContentBundle/WorldDataLayerReference.h b/Engine/Source/Runtime/Engine/Public/WorldPartition/ContentBundle/WorldDataLayerReference.h
--- a/Engine/Source/Runtime/Engine/Public/WorldPartition/ContentBundle/WorldDataLayerReference.h
+++ b/Engine/Source/Runtime/Engine/Public/WorldPartition/ContentBundle/WorldDataLayerReference.h
@@ -14,8 +14,6 @@ class FWorldDataLayersReference
 public:
 	ENGINE_API FWorldDataLayersReference();
 
-	UE_DEPRECATED(5.4, "Use UActorDescContainerInstance version instead")
-	FWorldDataLayersReference(UActorDescContainer* Container, FName WorldDataLayerName) {}
 	ENGINE_API FWorldDataLayersReference(UActorDescContainerInstance* ContainerInstance, FName WorldDataLayerName); // Retrieve Actor
 	ENGINE_API FWorldDataLayersReference(const FActorSpawnParameters& SpawnParameters); // Retrieve and Create Actor if Not Found, FActorSpawnParameters::Name must be set
 	ENGINE_API FWorldDataLayersReference(FWorldDataLayersReference&& Other);
diff --git a/Engine/Source/Runtime/Engine/Public/WorldPartition/Cook/WorldPartitionCookPackageContextInterface.h b/Engine/Source/Runtime/Engine/Public/WorldPartition/Cook/WorldPartitionCookPackageContextInterface.h
--- a/Engine/Source/Runtime/Engine/Public/WorldPartition/Cook/WorldPartitionCookPackageContextInterface.h
+++ b/Engine/Source/Runtime/Engine/Public/WorldPartition/Cook/WorldPartitionCookPackageContextInterface.h
@@ -25,11 +25,6 @@ public:
 	virtual void RegisterPackageCookPackageGenerator(IWorldPartitionCookPackageGenerator* Generator) = 0;
 	virtual void UnregisterPackageCookPackageGenerator(IWorldPartitionCookPackageGenerator* Generator) = 0;
 
-	UE_DEPRECATED(5.5, "Use AddPackageToGenerate instead.")
-	virtual const FWorldPartitionCookPackage* AddLevelStreamingPackageToGenerate(IWorldPartitionCookPackageGenerator* Generator, const FString& Root, const FString& RelativePath) { return nullptr; }
-	UE_DEPRECATED(5.5, "Use AddPackageToGenerate instead.")
-	virtual const FWorldPartitionCookPackage* AddGenericPackageToGenerate(IWorldPartitionCookPackageGenerator* Generator, const FString& Root, const FString& RelativePath) { return nullptr; }
-
 	virtual const FWorldPartitionCookPackage* AddPackageToGenerate(IWorldPartitionCookPackageGenerator* Generator, IWorldPartitionCookPackageObject* InCookPackageObject, const FString& Root, const FString& RelativePath) = 0;
 	virtual FString GetGeneratedPackagePath(IWorldPartitionCookPackageObject* InCookPackageObject) const = 0;
 
