Hide RayTracing, Navigation, LOD, HLOD categories  from SplineComponent : the spline component inherits from them from primitive component but is not renderable or navigable itself

#Synced-CL 50661191

#Synced-CL 50661191

diff --git a/Engine/Source/Runtime/Engine/Classes/Components/SplineComponent.h b/Engine/Source/Runtime/Engine/Classes/Components/SplineComponent.h
--- a/Engine/Source/Runtime/Engine/Classes/Components/SplineComponent.h
+++ b/Engine/Source/Runtime/Engine/Classes/Components/SplineComponent.h
@@ -210,7 +210,7 @@ struct FSplinePoint
  *	A spline component is a spline shape which can be used for other purposes (e.g. animating objects). It contains debug rendering capabilities.
  *	@see https://docs.unrealengine.com/latest/INT/Resources/ContentExamples/Blueprint_Splines
  */
-UCLASS(ClassGroup=Utility, ShowCategories = (Mobility), HideCategories = (Physics, Collision, Lighting, Rendering, Mobile), meta=(BlueprintSpawnableComponent), MinimalAPI)
+UCLASS(ClassGroup=Utility, ShowCategories = (Mobility), HideCategories = (Activation, Physics, Collision, HLOD, LOD, TextureStreaming, Navigation, RayTracing, Lighting, Rendering, Mobile), meta=(BlueprintSpawnableComponent), MinimalAPI)
 class USplineComponent : public UPrimitiveComponent
 {
 	GENERATED_UCLASS_BODY()
