[HWRT] Prevent RayTracingGeometryGroup leak when FStaticMeshRenderData is initialized twice.
- FStaticMeshRenderData::InitResources(...) can end up being called twice in some cases (eg: during Undo-Redo flows). FRenderResource::InitResource(...) guards against double initialization, FRayTracingGeometry::InitResource(...) also has logic to prevent double registration, so we were missing similar guard for the RayTracingGeometryGroup.

#Synced-CL 50872947

diff --git a/Engine/Source/Runtime/Engine/Private/Components/StaticMeshComponent.cpp b/Engine/Source/Runtime/Engine/Private/Components/StaticMeshComponent.cpp
--- a/Engine/Source/Runtime/Engine/Private/Components/StaticMeshComponent.cpp
+++ b/Engine/Source/Runtime/Engine/Private/Components/StaticMeshComponent.cpp
@@ -1912,6 +1912,7 @@ void UStaticMeshComponent::PostEditUndo()
 
 	// If the StaticMesh was also involved in this transaction, it may need reinitialization first
 	// In this case, the StaticMesh will have PostEditUndo called later in this transaction, which is too late to register this component
+	// TODO: should also check !GetStaticMesh()->IsCompiling() here? (UE-365996)
 	if (GetStaticMesh() && !GetStaticMesh()->AreRenderingResourcesInitialized())
 	{
 		// We need to recreate the render state of any components using the static mesh before modifying its rendering resources
diff --git a/Engine/Source/Runtime/Engine/Private/SkeletalMeshRenderData.cpp b/Engine/Source/Runtime/Engine/Private/SkeletalMeshRenderData.cpp
--- a/Engine/Source/Runtime/Engine/Private/SkeletalMeshRenderData.cpp
+++ b/Engine/Source/Runtime/Engine/Private/SkeletalMeshRenderData.cpp
@@ -636,6 +636,7 @@ void FSkeletalMeshRenderData::InitResources(bool bNeedsVertexColors, USkinnedAss
 			ENQUEUE_RENDER_COMMAND(SkeletalMeshInitRayTracingGeometryGroup)(UE::RenderCommandPipe::SkeletalMesh,
 				[this, NumLODs = LODRenderData.Num()]
 				{
+					check(RayTracingGeometryGroupHandle.IsNull());
 					RayTracingGeometryGroupHandle = GRayTracingGeometryManager->RegisterRayTracingGeometryGroup(NumLODs, CurrentFirstLODIdx);
 				});
 		}
diff --git a/Engine/Source/Runtime/Engine/Private/StaticMesh.cpp b/Engine/Source/Runtime/Engine/Private/StaticMesh.cpp
--- a/Engine/Source/Runtime/Engine/Private/StaticMesh.cpp
+++ b/Engine/Source/Runtime/Engine/Private/StaticMesh.cpp
@@ -2768,6 +2768,10 @@ void FStaticMeshRenderData::Serialize(FArchive& Ar, UStaticMesh* Owner, bool bCo
 
 void FStaticMeshRenderData::InitResources(ERHIFeatureLevel::Type InFeatureLevel, UStaticMesh* Owner)
 {
+	// TODO: This check is disabled because in some cases we end up calling InitResources(...) on already initialized RenderData
+	// due to an issue with PostEditUndo + async builds (UE-365996)
+	//check(!bIsInitialized);
+
 #if WITH_EDITOR
 	// Init the section info only for uncooked editor.
 	// Cooked packages don't need this and don't want any LOD screen size changes that it applies.
@@ -2809,7 +2813,10 @@ void FStaticMeshRenderData::InitResources(ERHIFeatureLevel::Type InFeatureLevel,
 
 				check(!RayTracingProxy->bUsingRenderingLODs || RayTracingLODs.Num() == LODResources.Num());
 
-				RayTracingGeometryGroupHandle = GRayTracingGeometryManager->RegisterRayTracingGeometryGroup(RayTracingLODs.Num(), RayTracingProxy->bUsingRenderingLODs ? CurrentFirstLODIdx : 0);
+				if (RayTracingGeometryGroupHandle.IsNull()) // workaround for FStaticMeshRenderData::InitResources(...) being called on already initialized RenderData due to bug with PostEditUndo + async build (UE-365996)
+				{
+					RayTracingGeometryGroupHandle = GRayTracingGeometryManager->RegisterRayTracingGeometryGroup(RayTracingLODs.Num(), RayTracingProxy->bUsingRenderingLODs ? CurrentFirstLODIdx : 0);
+				}
 
 				for (int32 LODIndex = 0; LODIndex < RayTracingLODs.Num(); ++LODIndex)
 				{
