Added a info log when a material is being translated with the new material translator for better debugging.

#rnx

#Synced-CL 50775131

diff --git a/Engine/Source/Runtime/Engine/Private/Materials/MaterialShared.cpp b/Engine/Source/Runtime/Engine/Private/Materials/MaterialShared.cpp
--- a/Engine/Source/Runtime/Engine/Private/Materials/MaterialShared.cpp
+++ b/Engine/Source/Runtime/Engine/Private/Materials/MaterialShared.cpp
@@ -3562,6 +3562,8 @@ bool FMaterial::Translate_New(const FMaterialShaderMapId& InShaderMapId,
 	FMaterialCompilationOutput& OutCompilationOutput,
 	TRefCountPtr<FSharedShaderCompilerEnvironment>& OutMaterialEnvironment)
 {
+	UE_LOG(LogMaterial, Display, TEXT("Translating material '%s' with the new experimental material translator."), *GetMaterialInterface()->GetFullName());
+
 	// Clear existing Material Errors.
 	CompileErrors.Empty();
 	ErrorExpressions.Empty();
