[WorldPartition] Fix crash when world is destroyed during async actor load

#Synced-CL 50826935

diff --git a/Engine/Source/Runtime/Engine/Private/WorldPartition/WorldPartitionLevelHelper.cpp b/Engine/Source/Runtime/Engine/Private/WorldPartition/WorldPartitionLevelHelper.cpp
--- a/Engine/Source/Runtime/Engine/Private/WorldPartition/WorldPartitionLevelHelper.cpp
+++ b/Engine/Source/Runtime/Engine/Private/WorldPartition/WorldPartitionLevelHelper.cpp
@@ -852,12 +852,26 @@ bool FWorldPartitionLevelHelper::LoadActorsInternal(FLoadActorsParams&& InParams
 			}
 		}
 
-		FLoadPackageAsyncDelegate CompletionCallback = FLoadPackageAsyncDelegate::CreateLambda([LoadProgress, PackageObjectMapping, LoadedOverrides = InLoadedPropertyOverrides, PackageReferencer = InParams.PackageReferencer, OuterWorld = InParams.OuterWorld, DestLevel = InParams.DestLevel, CompletionCallback = InParams.CompletionCallback, bSilenceLoadFailures = InParams.bSilenceLoadFailures, GetEffectiveEditorOnlyParentTransform](const FName& LoadedPackageName, UPackage* LoadedPackage, EAsyncLoadingResult::Type Result)
+		FLoadPackageAsyncDelegate CompletionCallback = FLoadPackageAsyncDelegate::CreateLambda([LoadProgress, PackageObjectMapping, LoadedOverrides = InLoadedPropertyOverrides, PackageReferencer = InParams.PackageReferencer, WeakOuterWorld = MakeWeakObjectPtr(InParams.OuterWorld), DestLevel = InParams.DestLevel, CompletionCallback = InParams.CompletionCallback, bSilenceLoadFailures = InParams.bSilenceLoadFailures, GetEffectiveEditorOnlyParentTransform](const FName& LoadedPackageName, UPackage* LoadedPackage, EAsyncLoadingResult::Type Result)
 		{
 			const FName ActorName = *FPaths::GetExtension(PackageObjectMapping->Path.ToString());
 			check(LoadProgress->NumPendingLoadRequests);
 			LoadProgress->NumPendingLoadRequests--;
 
+			// OuterWorld may have been destroyed while the async load was in progress. Treat this as a failed load request.
+			UWorld* OuterWorld = WeakOuterWorld.Get();
+			if (!OuterWorld)
+			{
+				UE_LOG(LogWorldPartition, Verbose, TEXT("LoadActors completion callback invoked when OuterWorld is no longer valid. Treating this request as failed. Actor: %s."), *ActorName.ToString());
+
+				LoadProgress->NumFailedLoadedRequests++;
+				if (!LoadProgress->NumPendingLoadRequests)
+				{
+					CompletionCallback(!LoadProgress->NumFailedLoadedRequests);
+				}
+				return;
+			}
+
 			AActor* Actor = nullptr;
 
 			if (LoadedPackage)
