World Partition - HLOD: Allow wp.Runtime.HLOD.ForceDisableShadows to be toggled anytime
Idea adapted from pr 14233

#Synced-CL 50899516

diff --git a/Engine/Source/Runtime/Engine/Private/WorldPartition/HLOD/HLODActor.cpp b/Engine/Source/Runtime/Engine/Private/WorldPartition/HLOD/HLODActor.cpp
--- a/Engine/Source/Runtime/Engine/Private/WorldPartition/HLOD/HLODActor.cpp
+++ b/Engine/Source/Runtime/Engine/Private/WorldPartition/HLOD/HLODActor.cpp
@@ -376,16 +376,65 @@ void AWorldPartitionHLOD::PreSave(FObjectPreSaveContext ObjectSaveContext)
 }
 #endif
 
+void AWorldPartitionHLOD::PostInitProperties()
+{
+	Super::PostInitProperties();
+
+	if (HasAnyFlags(RF_ClassDefaultObject) && ExactCast<AWorldPartitionHLOD>(this))
+	{
+		CVarWorldPartitionHLODForceDisableShadows->OnChangedDelegate().RemoveAll(this);
+		CVarWorldPartitionHLODForceDisableShadows->OnChangedDelegate().AddUObject(this, &AWorldPartitionHLOD::OnForceDisableShadowsCVarChanged);
+	}
+}
+
+void AWorldPartitionHLOD::BeginDestroy()
+{
+	if (HasAnyFlags(RF_ClassDefaultObject) && ExactCast<AWorldPartitionHLOD>(this))
+	{
+		CVarWorldPartitionHLODForceDisableShadows->OnChangedDelegate().RemoveAll(this);
+	}
+
+	Super::BeginDestroy();
+}
+
+void AWorldPartitionHLOD::OnForceDisableShadowsCVarChanged(IConsoleVariable* CVar)
+{
+	const bool bShouldCastShadows = (CVar != nullptr) ? !CVar->GetBool() : true;
+
+	for (const FWorldContext& WorldContext : GEngine->GetWorldContexts())
+	{
+		UWorld* World = WorldContext.World();
+		if (!World || !World->IsGameWorld())
+		{
+			continue;
+		}
+
+		for (TActorIterator<AWorldPartitionHLOD> It(World); It; ++It)
+		{
+			AWorldPartitionHLOD * HLODActor = *It;
+			if (IsValid(HLODActor))
+			{
+				HLODActor->ApplyCastShadowToComponents(bShouldCastShadows);
+			}
+		}
+	}
+}
+
+void AWorldPartitionHLOD::ApplyCastShadowToComponents(bool bInShouldCastShadows)
+{
+	ForEachComponent<UPrimitiveComponent>(false, [bInShouldCastShadows](UPrimitiveComponent* PrimitiveComponent)
+	{
+		PrimitiveComponent->SetCastShadow(bInShouldCastShadows);
+	});
+}
+
 void AWorldPartitionHLOD::PreRegisterAllComponents()
 {
 	Super::PreRegisterAllComponents();
 
 	if (GWorldPartitionHLODForceDisableShadows && GetWorld() && GetWorld()->IsGameWorld())
 	{
-		ForEachComponent<UPrimitiveComponent>(false, [](UPrimitiveComponent* PrimitiveComponent)
-		{
-			PrimitiveComponent->SetCastShadow(false);
-		});
+		ApplyCastShadowToComponents(false);
 	}
 
 #if WITH_EDITOR
diff --git a/Engine/Source/Runtime/Engine/Public/WorldPartition/HLOD/HLODActor.h b/Engine/Source/Runtime/Engine/Public/WorldPartition/HLOD/HLODActor.h
--- a/Engine/Source/Runtime/Engine/Public/WorldPartition/HLOD/HLODActor.h
+++ b/Engine/Source/Runtime/Engine/Public/WorldPartition/HLOD/HLODActor.h
@@ -159,6 +159,7 @@ private:
 
 protected:
 	//~ Begin UObject Interface.
+	ENGINE_API virtual void BeginDestroy() override;
 	ENGINE_API virtual void Serialize(FArchive& Ar) override;
 	ENGINE_API virtual bool IsEditorOnly() const override;
 	ENGINE_API virtual bool NeedsLoadForServer() const override;
@@ -172,6 +173,7 @@ protected:
 	//~ End UObject Interface.
 
 	//~ Begin AActor Interface.
+	ENGINE_API virtual void PostInitProperties() override;
 	ENGINE_API virtual void PreRegisterAllComponents() override;
 	ENGINE_API virtual void PostRegisterAllComponents() override;
 	ENGINE_API virtual void PostUnregisterAllComponents() override;
@@ -207,6 +209,9 @@ private:
 	FString GenerateHLODBuildReportHeaderString() const;
 #endif
 
+	void OnForceDisableShadowsCVarChanged(IConsoleVariable* CVar);
+	void ApplyCastShadowToComponents(bool bInShouldCastShadows);
+
 #if WITH_EDITORONLY_DATA
 	UPROPERTY()
 	TObjectPtr<UWorldPartitionHLODSourceActors> SourceActors;
