Implementation of a decode cache for sound waves

* Hooked up to new platform project settings
* New cvar to disable/enable it and override it from project settings (au.DecodeCache.Enabled -1/0/1)
* New cvar to set an auto-duration threshold. if set to non-negative, will override settings.
(au.DecodeCache.DurationThreshold -1)

High level overview:
* Singleton (global) decode cache for caching sound file decoding to prevent spending CPU resources on decoding the same audio over and over.
* Cache is stored using the TLru container in UE
* Lru Key is based on similar key for stream cache with the addition of an index quantized according to configured chunk size and playback position
* Cache is thread safe and intended to be accessed by async tasks (including async decode tasks and async metasound tasks).
* If chunk is not in the cache, it will do the decode synchronous to the caller (but out of the lock). When the decode is done, it'll lock and store the results. In the case of multiple queries, it'll decode twice and store once.
* Chunks are ref'd by thread-safe shared pointer
* LRU memory is tracked and pruned according to variable memory usage per-chunk. LRU entries store fixed-size frames but include interleaved audio. This means 7.1 chunk will take more memory than mono chunk. LRU is pre-allocated up-front (LRU count) according to worst-case usage (all mono chunks).
* API is minimal and designed to seamlessly integrate into existing audio engine code through a wrapper implementation of our decoder interface (ICompressedAudioInfo). Integration in this CL is through the wave proxy player -- but could also be brough to non-proxy-player code in future CLs.
* Testing has been local and not done at scale yet

[RN] Experimental decode cache implementation

#Synced-CL 50891994

diff --git a/Engine/Source/Runtime/Engine/Classes/Sound/AudioPlatformSettings.h b/Engine/Source/Runtime/Engine/Classes/Sound/AudioPlatformSettings.h
new file mode 100644
--- /dev/null
+++ b/Engine/Source/Runtime/Engine/Classes/Sound/AudioPlatformSettings.h
@@ -0,0 +1,42 @@
+// Copyright Epic Games, Inc. All Rights Reserved.
+
+#pragma once
+
+#include "Engine/PlatformSettings.h"
+#include "AudioPlatformSettings.generated.h"
+
+USTRUCT()
+struct FDecodeCacheSettings
+{
+	GENERATED_USTRUCT_BODY()
+
+	// Whether or not the decode cache is enabled by default
+	UPROPERTY(config, EditAnywhere, Category = "Decode Cache", meta = (ClampMin = 5))
+	bool EnableDecodeCache = false;
+
+	// The desired maximum decode cache size in megabytes
+	UPROPERTY(config, EditAnywhere, Category = "Decode Cache", meta = (ClampMin = 5))
+	float MaxDecodeCacheSizeMB = 10.0f;
+
+	// The desired maximum decode cache frames per chunk
+	UPROPERTY(config, EditAnywhere, Category = "Decode Cache", meta = (ClampMin = 1024))
+	int32 FramesPerChunk = 4096;
+
+	// Threshold to opt a sound to use the decode cache based on it's duration (in seconds). If less than threshold, it will use the decode cache.
+	UPROPERTY(config, EditAnywhere, Category = "Decode Cache", meta = (ClampMin = 0))
+	float AutoDurationThreshold = 2.0f;
+};
+
+/*
+* Class to store platform-specific settings.
+*/
+UCLASS(MinimalAPI, config = Engine, defaultconfig)
+class UAudioPlatformSettings : public UPlatformSettings
+{
+	GENERATED_BODY()
+
+public:
+	// Whether or not the decode cache is enabled by default
+	UPROPERTY(config, EditAnywhere, Category = "Decode Cache", meta = (ShowOnlyInnerProperties))
+	FDecodeCacheSettings DecodeCacheSettings;
+};
\ No newline at end of file
diff --git a/Engine/Source/Runtime/Engine/Classes/Sound/AudioSettings.h b/Engine/Source/Runtime/Engine/Classes/Sound/AudioSettings.h
--- a/Engine/Source/Runtime/Engine/Classes/Sound/AudioSettings.h
+++ b/Engine/Source/Runtime/Engine/Classes/Sound/AudioSettings.h
@@ -8,9 +8,11 @@
 #include "UObject/SoftObjectPath.h"
 #include "Engine/DeveloperSettings.h"
 #include "AudioMixerTypes.h"
+#include "AudioPlatformSettings.h"
 #include "AudioSettings.generated.h"
 
 struct FPropertyChangedChainEvent;
+struct FDecodeCacheSettings;
 
 class USoundClass;
 class USoundConcurrency;
@@ -282,6 +284,11 @@ private:
 	UPROPERTY(Transient)
 	TObjectPtr<USoundConcurrency> DefaultSoundConcurrency;
 
+
+	// Per-platform audio options container
+	UPROPERTY(config, EditAnywhere, Category="Platform Settings", meta = (ShowOnlyInnerProperties))
+	FPerPlatformSettings PlatformSettings;
+
 public:
 	// Loads default object instances from soft object path properties
 	ENGINE_API void LoadDefaultObjects();
@@ -311,6 +318,9 @@ public:
 	/** Returns the highest value for MaxChannels among all quality levels */
 	ENGINE_API int32 GetHighestMaxChannels() const;
 
+	// Convenience accessors for per-platform settings
+	ENGINE_API FDecodeCacheSettings GetDecodeCacheSettings() const;
+
 #if WITH_EDITOR
 	/** Returns event to be bound to if caller wants to know when audio settings are modified */
 	FAudioSettingsChanged& OnAudioSettingsChanged() { return AudioSettingsChanged; }
diff --git a/Engine/Source/Runtime/Engine/Classes/Sound/SoundWaveProxyDecodeCache.h b/Engine/Source/Runtime/Engine/Classes/Sound/SoundWaveProxyDecodeCache.h
new file mode 100644
--- /dev/null
+++ b/Engine/Source/Runtime/Engine/Classes/Sound/SoundWaveProxyDecodeCache.h
@@ -0,0 +1,35 @@
+// Copyright Epic Games, Inc. All Rights Reserved.
+
+#pragma once
+
+#include "CoreMinimal.h"
+#include "AudioDecompress.h"
+#include "HAL/LowLevelMemTracker.h"
+#include "UObject/PerPlatformProperties.h"
+
+LLM_DECLARE_TAG_API(Audio_SoundWaveDecodeCache, ENGINE_API);
+
+class FSoundWaveProxy;
+class ICompressedAudioInfo;
+
+/**
+ * A sound wave proxy decode cache. 
+ * Supports caching sound wave proxy decoding to prevent decoding the exact same audio multiple times.
+ * Useful to significantly reduce CPU usage when decoding short-duration or oft-repeated audio files. 
+ */
+namespace SoundWaveProxyDecodeCache
+{
+	using FSoundWaveProxyRef = TSharedRef<FSoundWaveProxy, ESPMode::ThreadSafe>;
+
+	/**
+	* Queries if the decode cache is enabled.
+	*/
+	bool ENGINE_API IsEnabled();
+
+	/**
+	* Used to create an ICompressedAudioInfo instance (a decoder) that will interact with the decode cache.
+	* @param InSoundWaveProxy The sound wave proxy object to use for the decoder. Defines the sound we want to actually decode.
+	* @return A new ICompressedAudioInfo unique ptr that serves as a decoder for the sound wave proxy.
+	*/
+	TUniquePtr<ICompressedAudioInfo> CreateDecoderInstance(const FSoundWaveProxyRef& InSoundWaveProxy);
+};
diff --git a/Engine/Source/Runtime/Engine/Private/AudioDecompress.cpp b/Engine/Source/Runtime/Engine/Private/AudioDecompress.cpp
--- a/Engine/Source/Runtime/Engine/Private/AudioDecompress.cpp
+++ b/Engine/Source/Runtime/Engine/Private/AudioDecompress.cpp
@@ -54,6 +54,30 @@ namespace AudioDecompressPrivate
 #endif //WITH_AUDIO_DECODER_DIAGNOSTICS
 }
 
+int32 ICompressedAudioInfo::ConformDecodeSize(int32 InMaxDesiredDecodeSizeInFrames)
+{
+	/** Minimum number of frames to decode per a call to the decoder.  */
+	static constexpr int32 DefaultMinDecodeSizeInFrames = 128;
+
+	/** Some codecs have strict requirements on decode size. In order to be
+		* functional with all supported codecs, the decode size must be a multiple
+		* of 128.
+		*/
+	static constexpr int32 DecodeSizeQuantizationInFrames = 128;
+
+	static_assert(DefaultMinDecodeSizeInFrames >= DecodeSizeQuantizationInFrames, "Min decode size less than decode size quantization");
+	static_assert((DefaultMinDecodeSizeInFrames % DecodeSizeQuantizationInFrames) == 0, "Min decode size must be equally divisible by decode size quantization");
+
+	InMaxDesiredDecodeSizeInFrames = FMath::Max(InMaxDesiredDecodeSizeInFrames, DefaultMinDecodeSizeInFrames);
+	const int32 QuantizedDecodeSizeInFrames = (InMaxDesiredDecodeSizeInFrames / DecodeSizeQuantizationInFrames) * DecodeSizeQuantizationInFrames;
+	const int32 ConformedDecodeSizeInFrames = FMath::Max(QuantizedDecodeSizeInFrames, DefaultMinDecodeSizeInFrames);
+
+	check((ConformedDecodeSizeInFrames % QuantizedDecodeSizeInFrames) == 0);
+	check(ConformedDecodeSizeInFrames > 0);
+
+	return ConformedDecodeSizeInFrames;
+}
+
 IStreamedCompressedInfo::IStreamedCompressedInfo()
 	: bIsStreaming(false)
 	, SrcBufferData(nullptr)
diff --git a/Engine/Source/Runtime/Engine/Private/AudioSettings.cpp b/Engine/Source/Runtime/Engine/Private/AudioSettings.cpp
--- a/Engine/Source/Runtime/Engine/Private/AudioSettings.cpp
+++ b/Engine/Source/Runtime/Engine/Private/AudioSettings.cpp
@@ -1,6 +1,7 @@
 // Copyright Epic Games, Inc. All Rights Reserved.
 #include "Sound/AudioSettings.h"
 
+#include "Sound/AudioPlatformSettings.h"
 #include "AudioBusSubsystem.h"
 #include "AudioDevice.h"
 #include "Sound/SoundNodeQualityLevel.h"
@@ -29,6 +30,8 @@ UAudioSettings::UAudioSettings(const FObjectInitializer& ObjectInitializer)
 	GlobalMaxPitchScale = 2.0F;
 
 	DefaultAudioCompressionType = EDefaultAudioCompressionType::BinkAudio;
+
+	PlatformSettings.Initialize(UAudioPlatformSettings::StaticClass());
 }
 
 void UAudioSettings::AddDefaultSettings()
@@ -82,12 +85,10 @@ void UAudioSettings::PostEditChangeChainProperty(FPropertyChangedChainEvent& Pro
 
 			bPromptRestartRequired = true;
 		}
-		else if (PropertyName == GET_MEMBER_NAME_CHECKED(UAudioSettings, BaseDefaultSubmix))
-		{
-			bPromptRestartRequired = true;
-		}
-		else if(PropertyName == GET_MEMBER_NAME_CHECKED(UAudioSettings, EQSubmix)
-			|| PropertyName == GET_MEMBER_NAME_CHECKED(UAudioSettings, ReverbSubmix))
+		else if (PropertyName == GET_MEMBER_NAME_CHECKED(UAudioSettings, BaseDefaultSubmix)
+			|| PropertyName == GET_MEMBER_NAME_CHECKED(UAudioSettings, EQSubmix)
+			|| PropertyName == GET_MEMBER_NAME_CHECKED(UAudioSettings, ReverbSubmix)
+			|| PropertyName == GET_MEMBER_NAME_CHECKED(UAudioSettings, PlatformSettings))
 		{
 			bPromptRestartRequired = true;
 		}
@@ -326,5 +327,14 @@ FString UAudioSettings::FindQualityNameByIndex(int32 Index) const
 		   TEXT("");
 }
 
+ENGINE_API FDecodeCacheSettings UAudioSettings::GetDecodeCacheSettings() const
+{
+	if (UAudioPlatformSettings* Settings = UPlatformSettingsManager::Get().GetSettingsForPlatform<UAudioPlatformSettings>())
+	{
+		return Settings->DecodeCacheSettings;
+	}
+	return FDecodeCacheSettings();
+}
+
 #undef LOCTEXT_NAMESPACE
 
diff --git a/Engine/Source/Runtime/Engine/Private/SoundWaveProxyDecodeCacheDecoder.cpp b/Engine/Source/Runtime/Engine/Private/SoundWaveProxyDecodeCacheDecoder.cpp
new file mode 100644
--- /dev/null
+++ b/Engine/Source/Runtime/Engine/Private/SoundWaveProxyDecodeCacheDecoder.cpp
@@ -0,0 +1,319 @@
+// Copyright Epic Games, Inc. All Rights Reserved.
+
+#include "SoundWaveProxyDecodeCacheDecoder.h"
+#include "SoundWaveProxyDecodeCacheImpl.h"
+
+FSoundWaveProxyDecodeCacheDecoder::FSoundWaveProxyDecodeCacheDecoder(const FSoundWaveProxyRef& InWaveProxy)
+{
+	FName Format = InWaveProxy->GetRuntimeFormat();
+	CompressedAudioInfo = TUniquePtr<ICompressedAudioInfo>(IAudioInfoFactoryRegistry::Get().Create(Format, InWaveProxy->GetFName()));
+	if (!CompressedAudioInfo)
+	{
+		UE_LOG(LogAudio, Error, TEXT("FSoundWaveProxyDecodeCacheDecoder::FSoundWaveProxyDecodeCacheDecoder() Failed to create a compressed audio info struct for format %s"), *Format.ToString());
+	}
+	// Need to do this to make sure the wrapper of concrete implementations functions correctly
+	NumChannels = InWaveProxy->GetNumChannels();
+	NumFramesInWave = InWaveProxy->GetNumFrames();
+	WaveGuid = InWaveProxy->GetSoundWaveData()->GetGUID();
+	SampleRate = InWaveProxy->GetSampleRate();
+#if WITH_EDITOR
+	ChunkRevision = InWaveProxy->GetCurrentChunkRevision();
+#endif // #if WITH_EDITOR
+} 
+
+bool FSoundWaveProxyDecodeCacheDecoder::StreamCompressedInfo(USoundWave* Wave, struct FSoundQualityInfo* QualityInfo)
+{
+	if (CompressedAudioInfo)
+	{
+		return CompressedAudioInfo->StreamCompressedInfo(Wave, QualityInfo);
+	}
+	return false;
+}
+
+bool FSoundWaveProxyDecodeCacheDecoder::StreamCompressedInfo(const FSoundWaveProxyPtr& Wave, struct FSoundQualityInfo* QualityInfo)
+{
+	if (CompressedAudioInfo)
+	{
+		return CompressedAudioInfo->StreamCompressedInfo(Wave, QualityInfo);
+	}
+	return false;
+}
+
+bool FSoundWaveProxyDecodeCacheDecoder::ReadCompressedInfo(const uint8* InSrcBufferData, uint32 InSrcBufferDataSize, struct FSoundQualityInfo* QualityInfo)
+{
+	if (CompressedAudioInfo)
+	{
+		return CompressedAudioInfo->ReadCompressedInfo(InSrcBufferData, InSrcBufferDataSize, QualityInfo);
+	}
+	return false;
+}
+
+bool FSoundWaveProxyDecodeCacheDecoder::ReadCompressedData(uint8* Destination, bool bLooping, uint32 BufferSize)
+{
+	int32 NumBytesRead = 0;
+	return GetCompressedDataInternal(Destination, bLooping, BufferSize, NumBytesRead, false);
+}
+
+void FSoundWaveProxyDecodeCacheDecoder::SeekToTime(const float SeekTime)
+{
+	if (NumFramesInWave > 0)
+	{
+		float SeekFrame = FMath::Clamp(static_cast<int32>(SeekTime * SampleRate), 0, NumFramesInWave);
+		SeekToFrame(SeekFrame);
+	}
+	else
+	{
+		CurrentReadFrameIndex = 0;
+		DecodedSoundWaveChunk = nullptr;
+	}
+}
+
+void FSoundWaveProxyDecodeCacheDecoder::SeekToFrame(const uint32 Frame)
+{
+	LLM_SCOPE_BYTAG(Audio_SoundWaveDecodeCache);
+
+	CurrentReadFrameIndex = Frame;
+
+	// Check to see if our seek is outside the range of a current chunk
+	if (DecodedSoundWaveChunk.IsValid())
+	{
+		if (CurrentReadFrameIndex < DecodedSoundWaveChunk->FrameStart || CurrentReadFrameIndex >= DecodedSoundWaveChunk->FrameEnd)
+		{
+			DecodedSoundWaveChunk = nullptr;
+		}
+	}
+}
+
+void FSoundWaveProxyDecodeCacheDecoder::ExpandFile(uint8* DstBuffer, struct FSoundQualityInfo* QualityInfo)
+{
+	if (CompressedAudioInfo)
+	{
+		CompressedAudioInfo->ExpandFile(DstBuffer, QualityInfo);
+	}
+}
+
+void FSoundWaveProxyDecodeCacheDecoder::EnableHalfRate(bool HalfRate)
+{
+	if (CompressedAudioInfo)
+	{
+		CompressedAudioInfo->EnableHalfRate(HalfRate);
+	}
+}
+
+uint32 FSoundWaveProxyDecodeCacheDecoder::GetSourceBufferSize() const
+{
+	if (CompressedAudioInfo)
+	{
+		return CompressedAudioInfo->GetSourceBufferSize();
+	}
+	return 0;
+}
+
+bool FSoundWaveProxyDecodeCacheDecoder::UsesVorbisChannelOrdering() const
+{
+	if (CompressedAudioInfo)
+	{
+		return CompressedAudioInfo->UsesVorbisChannelOrdering();
+	}
+	return false;
+}
+
+int32 FSoundWaveProxyDecodeCacheDecoder::GetStreamBufferSize() const
+{
+	if (CompressedAudioInfo)
+	{
+		return CompressedAudioInfo->GetSourceBufferSize();
+	}
+	return 0;
+}
+
+bool FSoundWaveProxyDecodeCacheDecoder::HasError() const
+{
+	if (CompressedAudioInfo)
+	{
+		return CompressedAudioInfo->HasError();
+	}
+	return true;
+}
+
+bool FSoundWaveProxyDecodeCacheDecoder::StreamCompressedData(uint8* Destination, bool bLooping, uint32 BufferSize, int32& OutNumBytesStreamed)
+{
+	return GetCompressedDataInternal(Destination, bLooping, BufferSize, OutNumBytesStreamed, true);
+}
+
+int32 FSoundWaveProxyDecodeCacheDecoder::GetCurrentChunkIndex() const
+{
+	if (CompressedAudioInfo)
+	{
+		return CompressedAudioInfo->GetCurrentChunkIndex();
+	}
+	return INDEX_NONE;
+}
+
+int32 FSoundWaveProxyDecodeCacheDecoder::GetCurrentChunkOffset() const
+{ 
+	if (CompressedAudioInfo)
+	{
+		return CompressedAudioInfo->GetCurrentChunkOffset();
+	}
+	return INDEX_NONE;
+}
+
+const FSoundWaveProxyPtr& FSoundWaveProxyDecodeCacheDecoder::GetStreamingSoundWave() const
+{ 
+	if (CompressedAudioInfo)
+	{
+		return CompressedAudioInfo->GetStreamingSoundWave();
+	}
+	// Return a null/empty reference in case we don't have a CompressedAudioInfo
+	static const FSoundWaveProxyPtr Empty;
+	return Empty;
+}
+
+bool FSoundWaveProxyDecodeCacheDecoder::StreamCompressedInfoInternal(const FSoundWaveProxyPtr& InWaveProxy, struct FSoundQualityInfo* QualityInfo)
+{
+	if (CompressedAudioInfo)
+	{
+		return CompressedAudioInfo->StreamCompressedInfoInternal(InWaveProxy, QualityInfo);
+	}
+	return false;
+}
+
+bool FSoundWaveProxyDecodeCacheDecoder::GetCompressedDataInternal(uint8* Destination, bool bLooping, uint32 BufferSize, int32& OutNumBytesStreamed, bool bDoStreamingDecode)
+{
+	if (!CompressedAudioInfo)
+	{
+		return false;
+	}
+
+	check(NumChannels > 0);
+	check(SampleRate > 0);
+
+	bool bIsFinished = false;
+
+	// The number of bytes we need to still copy out. 
+	int32 NumFramesToCopyOut = (int32)BufferSize / (MONO_PCM_SAMPLE_SIZE * NumChannels);
+	OutNumBytesStreamed = 0;
+
+	while (NumFramesToCopyOut > 0 && !bIsFinished)
+	{
+		// Check to see if we already have a decoded chunk 
+		if (DecodedSoundWaveChunk.IsValid())
+		{
+			int32 ChunkFrameStart = DecodedSoundWaveChunk->FrameStart;
+			int32 ChunkFrameEnd = DecodedSoundWaveChunk->FrameEnd;
+			int32 ChunkFrameReadIndex = CurrentReadFrameIndex - ChunkFrameStart;
+			int32 NumFramesLeftInThisChunk = FMath::Max<int32>(ChunkFrameEnd - CurrentReadFrameIndex, 0);
+			int32 NumFramesToCopyFromThisChunk = FMath::Min(NumFramesToCopyOut, NumFramesLeftInThisChunk);
+
+			// if we have frames to copy out from this decoded buffer, we don't need to get another chunk
+			if (NumFramesToCopyFromThisChunk > 0)
+			{
+				check(NumFramesToCopyFromThisChunk + ChunkFrameStart <= ChunkFrameEnd);
+				int32 NumBytesToCopy = MONO_PCM_SAMPLE_SIZE * NumFramesToCopyFromThisChunk * NumChannels;
+				int32 CopyFromByteOffset = MONO_PCM_SAMPLE_SIZE * ChunkFrameReadIndex * NumChannels;
+				uint8* CopyFromBufferPtr = DecodedSoundWaveChunk->DecodedByteData.GetData();
+				int32 NumBytesInDecodeBuffer = DecodedSoundWaveChunk->DecodedByteData.Num();
+
+				uint8* Dest = Destination + OutNumBytesStreamed;
+				int32 RemainingBytesInDest = BufferSize - OutNumBytesStreamed;
+				check(NumBytesToCopy <= RemainingBytesInDest);
+
+				uint8* Src = CopyFromBufferPtr + CopyFromByteOffset;
+				check(CopyFromByteOffset + NumBytesToCopy <= NumBytesInDecodeBuffer);
+				FMemory::Memcpy(Dest, Src, NumBytesToCopy);
+
+				OutNumBytesStreamed += NumBytesToCopy;
+				NumFramesToCopyOut -= NumFramesToCopyFromThisChunk;
+				CurrentReadFrameIndex += NumFramesToCopyFromThisChunk;
+			}
+
+			// Check for finished condition
+			// If we are not looping and if we've reached the end of this chunk, and it's the last chunk
+			if (NumFramesToCopyFromThisChunk == NumFramesLeftInThisChunk && DecodedSoundWaveChunk->bIsLastChunk && !bLooping)
+			{
+				bIsFinished = true;
+				break;
+			}
+		}
+
+		// we still have frames we need to copy out, so get a new decoded buffer from the cache
+		// OR if there is no decoded audio in the cache for the data we need, do the decode for a new chunk
+		if (NumFramesToCopyOut > 0 && !bIsFinished)
+		{
+			using namespace SoundWaveProxyDecodeCache;
+
+			// Retrieve the global decode cache
+			FSoundWaveProxyDecodeCache* DecodeCache = FSoundWaveProxyDecodeCache::Get();
+			check(DecodeCache);
+			
+			int32 FramesPerChunk = DecodeCache->GetFramesPerChunk();
+			check(FramesPerChunk > 0);
+
+			// build a decode buffer key that we want to query
+			FDecodedSoundWaveChunkKey NewChunkKey;
+			NewChunkKey.WaveName = WaveName;
+			NewChunkKey.WaveGuid = WaveGuid;
+			NewChunkKey.ChunkIndex = CurrentReadFrameIndex / FramesPerChunk;
+#if WITH_EDITOR
+			NewChunkKey.ChunkRevision = ChunkRevision;
+			NewChunkKey.NumChannels = NumChannels;
+			NewChunkKey.SampleRate = SampleRate;
+#endif
+			DecodedSoundWaveChunk = DecodeCache->GetDecodedChunk(NewChunkKey);
+
+			// if there is no decoded chunk for this key, that means we need to do the decode
+			// with our underlying contained CompressedAudioInfo object
+			if (!DecodedSoundWaveChunk.IsValid())
+			{
+				int32 SeekFrame = NewChunkKey.ChunkIndex * FramesPerChunk;
+				int32 BytesPerChunk = FramesPerChunk * MONO_PCM_SAMPLE_SIZE * NumChannels;
+
+				{
+					LLM_SCOPE_BYTAG(Audio_SoundWaveDecodeCache);
+
+					DecodedSoundWaveChunk = MakeShared<FDecodedSoundWaveChunk>();
+					DecodedSoundWaveChunk->FrameStart = SeekFrame;
+					DecodedSoundWaveChunk->DecodedByteData.AddUninitialized(BytesPerChunk);
+				}
+
+				// first, seek the underlying decoder to the required chunk boundary
+				check(CompressedAudioInfo);
+				CompressedAudioInfo->SeekToFrame(SeekFrame);
+
+				uint8* DecodedByteDataPtr = DecodedSoundWaveChunk->DecodedByteData.GetData();
+
+				int32 NumBytesStreamed = 0;
+				bool bIsLastDecodeChunk = false;
+				if (bDoStreamingDecode)
+				{
+					bIsLastDecodeChunk = CompressedAudioInfo->StreamCompressedData(DecodedByteDataPtr, bLooping, BytesPerChunk, NumBytesStreamed);
+				}
+				else
+				{
+					NumBytesStreamed = BytesPerChunk;
+					bIsLastDecodeChunk = CompressedAudioInfo->ReadCompressedData(DecodedByteDataPtr, bLooping, BytesPerChunk);
+				}
+
+				// No data was decoded, so early out
+				if (NumBytesStreamed == 0)
+				{
+					break;
+				}
+
+				check(bIsLastDecodeChunk || NumBytesStreamed == BytesPerChunk);
+				DecodedSoundWaveChunk->bIsLastChunk = bIsLastDecodeChunk;
+				DecodedSoundWaveChunk->FrameEnd = DecodedSoundWaveChunk->FrameStart + (NumBytesStreamed / (MONO_PCM_SAMPLE_SIZE * NumChannels));
+
+				int32 FramesInChunk = DecodedSoundWaveChunk->FrameEnd - DecodedSoundWaveChunk->FrameStart;
+				check((FramesInChunk * (int32)MONO_PCM_SAMPLE_SIZE * (int32)NumChannels) <= (int32)BytesPerChunk);
+				check(FramesInChunk <= FramesPerChunk);
+
+				DecodeCache->CacheDecodedChunk(NewChunkKey, DecodedSoundWaveChunk);
+			}
+		}
+	}
+
+	return bIsFinished;
+}
+
diff --git a/Engine/Source/Runtime/Engine/Private/SoundWaveProxyDecodeCacheDecoder.h b/Engine/Source/Runtime/Engine/Private/SoundWaveProxyDecodeCacheDecoder.h
new file mode 100644
--- /dev/null
+++ b/Engine/Source/Runtime/Engine/Private/SoundWaveProxyDecodeCacheDecoder.h
@@ -0,0 +1,67 @@
+// Copyright Epic Games, Inc. All Rights Reserved.
+
+#pragma once
+
+#include "Sound/SoundWaveProxyDecodeCache.h"
+#include "SoundWaveProxyDecodeCacheImpl.h"
+#include "Interfaces/IAudioFormat.h"
+#include "AudioDecompress.h"
+
+class FSoundWaveProxyDecodeCache;
+struct FDecodedSoundWaveChunk;
+
+/*
+ * Concrete implementation of ICompressedAudioInfo, which wraps an underlying
+ * decoder and intercepts decode and seek calls to translate it to interactions with a decode cache.
+ */
+class FSoundWaveProxyDecodeCacheDecoder : public ICompressedAudioInfo
+{
+public:
+	using FSoundWaveProxyRef = TSharedRef<FSoundWaveProxy, ESPMode::ThreadSafe>;
+	using FDecodedSoundWaveChunkPtr = TSharedPtr<FDecodedSoundWaveChunk, ESPMode::ThreadSafe>;
+
+	FSoundWaveProxyDecodeCacheDecoder(const FSoundWaveProxyRef& InWaveProxy);
+
+	virtual ~FSoundWaveProxyDecodeCacheDecoder()
+	{}
+
+	virtual bool StreamCompressedInfo(USoundWave* Wave, struct FSoundQualityInfo* QualityInfo) override final;
+	virtual bool StreamCompressedInfo(const FSoundWaveProxyPtr& Wave, struct FSoundQualityInfo* QualityInfo) override final;
+	virtual bool ReadCompressedInfo(const uint8* InSrcBufferData, uint32 InSrcBufferDataSize, struct FSoundQualityInfo* QualityInfo) override final;
+	virtual bool ReadCompressedData(uint8* Destination, bool bLooping, uint32 BufferSize) override final;
+	virtual void SeekToTime(const float SeekTime) override final;
+	virtual void SeekToFrame(const uint32 Frame) override final;
+	virtual void ExpandFile(uint8* DstBuffer, struct FSoundQualityInfo* QualityInfo) override final;
+	virtual void EnableHalfRate(bool HalfRate) override final;
+	virtual uint32 GetSourceBufferSize() const override final;
+	virtual bool UsesVorbisChannelOrdering() const override final;
+	virtual int32 GetStreamBufferSize() const override final;
+	virtual bool HasError() const override final;
+	virtual bool StreamCompressedData(uint8* Destination, bool bLooping, uint32 BufferSize, int32& OutNumBytesStreamed) override final;
+	virtual int32 GetCurrentChunkIndex() const override final;
+	virtual int32 GetCurrentChunkOffset() const override final;
+	virtual const FSoundWaveProxyPtr& GetStreamingSoundWave() const final;
+	virtual bool StreamCompressedInfoInternal(const FSoundWaveProxyPtr& InWaveProxy, struct FSoundQualityInfo* QualityInfo) override final;
+
+private:
+
+	bool GetCompressedDataInternal(uint8* Destination, bool bLooping, uint32 BufferSize, int32& OutNumBytesStreamed, bool bDoStreamingDecode);
+
+	// The actual decoder we are using for decoding
+	TUniquePtr<ICompressedAudioInfo> CompressedAudioInfo;
+
+	// A ptr to the decoded chunk of sound wave data
+	// This is either decoded on the fly and discarded OR pulled from the decode cache
+	// if this decoder is flagged to use the decode cache
+	FDecodedSoundWaveChunkPtr DecodedSoundWaveChunk;
+
+	uint32 NumChannels = 0;
+	uint32 SampleRate = 0;
+	int32 NumFramesInWave = INDEX_NONE;
+	int32 CurrentReadFrameIndex = 0;
+#if WITH_EDITOR
+	int32 ChunkRevision = INDEX_NONE;
+#endif
+	FName WaveName;
+	FGuid WaveGuid;
+};
diff --git a/Engine/Source/Runtime/Engine/Private/SoundWaveProxyDecodeCacheImpl.cpp b/Engine/Source/Runtime/Engine/Private/SoundWaveProxyDecodeCacheImpl.cpp
new file mode 100644
--- /dev/null
+++ b/Engine/Source/Runtime/Engine/Private/SoundWaveProxyDecodeCacheImpl.cpp
@@ -0,0 +1,242 @@
+// Copyright Epic Games, Inc. All Rights Reserved.
+
+#include "SoundWaveProxyDecodeCacheImpl.h"
+#include "Math/UnitConversion.h"
+#include "SoundWaveProxyDecodeCacheDecoder.h"
+#include "HAL/IConsoleManager.h"
+
+
+static int32 EnableDecodeCacheCvar = -1;
+FAutoConsoleVariableRef CVarEnableDecodeCache(
+	TEXT("au.DecodeCache.Enabled"),
+	EnableDecodeCacheCvar,
+	TEXT("Overrides project settings to enable or disable decode cache.\n")
+	TEXT("-1: Use project settings value, 0: Disable, 1: Enable"),
+	ECVF_Default);
+
+static float AutoDurationThresholdCvar = -1.0f;
+FAutoConsoleVariableRef CVarAutoDurationThreshold(
+	TEXT("au.DecodeCache.DurationThreshold"),
+	AutoDurationThresholdCvar,
+	TEXT("If non-negative, overrides project settings to set the auto duration threshold. Sounds with duration below this number will use the decode cache automatically.\n")
+	TEXT("-1: Use project settings value, >= 0.0 will use this value for auto duration threshold. A value of 0.0 means every sound will use decode cache if it is enabled."),
+	ECVF_Default);
+
+
+LLM_DEFINE_TAG(Audio_SoundWaveDecodeCache);
+
+FSoundWaveProxyDecodeCache::FSoundWaveProxyDecodeCache()
+{
+	const UAudioSettings* Settings = GetDefault<UAudioSettings>();
+	check(Settings);
+
+	const FDecodeCacheSettings& DecodeCacheSettings = Settings->GetDecodeCacheSettings();
+
+	// Clamp settings to reasonable limits to prevent out-of-bounds issues
+	DesiredCacheSizeMB = FMath::Clamp(DecodeCacheSettings.MaxDecodeCacheSizeMB, 1.0f, 100.0f);
+	DesiredFramesPerChunk = FMath::Clamp(DecodeCacheSettings.FramesPerChunk, 1024, 48000);
+	AutoDurationThreshold = FMath::Max(DecodeCacheSettings.AutoDurationThreshold, 0.0f);
+
+	float DesiredCacheSizeBytes = FUnitConversion::Convert(DesiredCacheSizeMB, EUnit::Megabytes, EUnit::Bytes);
+
+	// Conform the requested frames per chunk from the settings to one that will actually work
+	FramesPerChunk = ICompressedAudioInfo::ConformDecodeSize(DesiredFramesPerChunk);
+
+	// Calculate the worse-case lru chunk count based on mono buffers of the desired min frames per chunk and the size of the struct used to 
+	// track the data in the chunk entries
+	int32 TotalBytesPerEntry = MONO_PCM_SAMPLE_SIZE * FramesPerChunk + sizeof(FDecodedSoundWaveChunk) + sizeof(FDecodedSoundWaveChunkKey);
+	MaxChunkCount = DesiredCacheSizeBytes / TotalBytesPerEntry;
+
+	// Lets set some extreme but minimal chunk count limit
+	constexpr int32 MinChunkCount = 10;
+	if (MaxChunkCount > MinChunkCount)
+	{
+		check(bHasFailure == false);
+		MaxMemoryBytes = MaxChunkCount * TotalBytesPerEntry;
+
+		// Since there could be rounding issues with fitting the desired budget to frames per chunk, our actual budget may differ slightly
+		float ActualMaxMemoryMB = FUnitConversion::Convert((float)MaxMemoryBytes, EUnit::Bytes, EUnit::Megabytes);
+
+		UE_LOG(LogAudio, Display,
+			TEXT("Initializing decode cache to %.2f Megabytes (Desired: %.2f), MaxChunkCount: %d. Enabled: %s"),
+			ActualMaxMemoryMB, DesiredCacheSizeMB, MaxChunkCount, SoundWaveProxyDecodeCache::IsEnabled() ? TEXT("Yes") : TEXT("No"));
+
+		// Note that we are tracking the LRU overhead here since the LRU cache is pre-allocated when initializing.
+		UsedMemoryBytes = MaxChunkCount * (sizeof(FDecodedSoundWaveChunk) + sizeof(FDecodedSoundWaveChunkKey));
+		check(NumChunksUsed == 0);
+
+		LLM_SCOPE_BYTAG(Audio_SoundWaveDecodeCache);
+
+		// Initialize the Lru to the max chunk count (worse-case all mono audio)
+		LruCache.Empty(MaxChunkCount);
+	}
+	else
+	{
+		bHasFailure = true;
+
+		UE_LOG(LogAudio, Warning, 
+			TEXT("Failed to initialize the decode cache due to settings (MaxDecodeCacheSizeMB: %.2f, FramesPerChunk: %d) which result in extremely small decode chunk count: %d (Min Chunk Count: %d)."), 
+			DecodeCacheSettings.MaxDecodeCacheSizeMB, DecodeCacheSettings.FramesPerChunk, MaxChunkCount, MinChunkCount);
+	}
+}
+
+int32 FSoundWaveProxyDecodeCache::GetFramesPerChunk() const
+{
+	check(FramesPerChunk > 0);
+	return FramesPerChunk;
+}
+
+FSoundWaveProxyCacheInfo FSoundWaveProxyDecodeCache::GetInfo() const
+{
+	FSoundWaveProxyCacheInfo Info;
+	Info.MaxMemoryBytes = MaxMemoryBytes;
+	Info.NumChunksCapacity = MaxChunkCount;
+	Info.UsedChunks = NumChunksUsed;
+	Info.UsedMemoryBytes = UsedMemoryBytes;
+	return Info;
+}
+
+float FSoundWaveProxyDecodeCache::GetAutoDurationThreshold() const
+{
+	if (AutoDurationThresholdCvar >= 0.0f)
+	{
+		return AutoDurationThresholdCvar;
+	}
+
+	return AutoDurationThreshold;
+}
+
+
+TUniquePtr<ICompressedAudioInfo> FSoundWaveProxyDecodeCache::CreateDecoderInstance(const FSoundWaveProxyRef& InSoundWaveProxy)
+{
+	LLM_SCOPE_BYTAG(Audio_SoundWaveDecodeCache);
+
+	return MakeUnique<FSoundWaveProxyDecodeCacheDecoder>(InSoundWaveProxy);
+}
+
+FSoundWaveProxyDecodeCache::FDecodedSoundWaveChunkPtr FSoundWaveProxyDecodeCache::GetDecodedChunk(const FDecodedSoundWaveChunkKey& InKey)
+{
+	FScopeLock Lock(&CacheLock);
+
+	if (SoundWaveProxyDecodeCache::IsEnabled() && !bHasFailure)
+	{
+		FDecodedSoundWaveChunkPtr* Entry = LruCache.FindAndTouch(InKey);
+		if (Entry)
+		{
+			return *Entry;
+		}
+	}
+
+	return nullptr;
+}
+
+void FSoundWaveProxyDecodeCache::CacheDecodedChunk(const FDecodedSoundWaveChunkKey& InKey, FDecodedSoundWaveChunkPtr& InDecodedChunk)
+{
+	FScopeLock Lock(&CacheLock);
+
+	LLM_SCOPE_BYTAG(Audio_SoundWaveDecodeCache);
+
+	// Don't allow caching chunks if we're disabled or have a failure
+	if (!SoundWaveProxyDecodeCache::IsEnabled() || bHasFailure)
+	{
+		return;
+	}
+
+	// First check if chunk is already in the cache. If it is, we'll touch it and early return
+	FDecodedSoundWaveChunkPtr* AlreadyInCache = LruCache.FindAndTouch(InKey);
+	if (AlreadyInCache)
+	{
+		return;
+	}
+
+	int32 NextUsedMemoryBytes = UsedMemoryBytes + InDecodedChunk->DecodedByteData.Num() + sizeof(FDecodedSoundWaveChunk) + sizeof(FDecodedSoundWaveChunkKey);
+
+	// Only add the decoded chunk to the cache if it's below our budget
+	// If try to trim the cache down to below our budget
+	bool bCanCacheChunk = SoundWaveProxyDecodeCache::IsEnabled();
+	while (bCanCacheChunk && NextUsedMemoryBytes >= MaxMemoryBytes)
+	{
+		// The first chunk may exceed the budget and try to evict from an empty cache which will cause an assert
+		if (!LruCache.IsEmpty())
+		{
+			FDecodedSoundWaveChunkKey LeastRecentKey = LruCache.GetLeastRecentKey();
+			FDecodedSoundWaveChunkPtr& LeastRecentChunk = LruCache.FindChecked(LeastRecentKey);
+
+			// if the least-recent entry only has one reference (it's unique), then this decode cache is it's only reference and it's safe to remove
+			if (LeastRecentChunk.IsUnique())
+			{
+				int32 TotalMemSizeOfLeastRecentEntry = LeastRecentChunk->DecodedByteData.Num() + sizeof(FDecodedSoundWaveChunk) + sizeof(FDecodedSoundWaveChunkKey);
+				NextUsedMemoryBytes -= TotalMemSizeOfLeastRecentEntry;
+				NumChunksUsed--;
+				LruCache.RemoveLeastRecent();
+			}
+			else
+			{
+				// Our least recently used chunk is being use, which means the cache has no chunks that can be evicted so we can't use the decode cache
+				bCanCacheChunk = false;
+			}
+			// Check if we need to terminate this loop early because the cache is totally full and everything is being used (i.e. we can't cache this chunk)
+		}
+		else
+		{
+			bCanCacheChunk = false;
+		}
+	}
+
+	if (bCanCacheChunk)
+	{
+		UsedMemoryBytes = NextUsedMemoryBytes;
+		NumChunksUsed++;
+		LruCache.Add(InKey, InDecodedChunk);
+	}
+	else
+	{
+		UE_LOG(LogAudio, Verbose, 
+			TEXT("FSoundWaveProxyDecodeCache::CacheDecodedChunk, Failed to cache decoded audio chunk: cache is full and all chunks are referenced. NumChunksUsed: %d"), NumChunksUsed);
+	}
+}
+
+FSoundWaveProxyDecodeCache* FSoundWaveProxyDecodeCache::Get()
+{
+	static TUniquePtr<FSoundWaveProxyDecodeCache> DecodeCache = MakeUnique<FSoundWaveProxyDecodeCache>();
+	return DecodeCache.Get();
+}
+
+bool SoundWaveProxyDecodeCache::IsEnabled()
+{
+	// -1 means to use the project settings
+	if (EnableDecodeCacheCvar == -1)
+	{
+		// Only access settings once and cache the results
+		auto IsCacheEnabledFromSettings = []() -> bool
+			{
+				const UAudioSettings* Settings = GetDefault<UAudioSettings>();
+				if (Settings)
+				{
+					const FDecodeCacheSettings& DecodeCacheSettings = Settings->GetDecodeCacheSettings();
+					return DecodeCacheSettings.EnableDecodeCache;
+				}
+				return false;
+			};
+
+		static const bool IsEnabledCached = IsCacheEnabledFromSettings();
+		return IsEnabledCached;
+	}
+	else
+	{
+		// otherwise 0 means not enabled, >= 1 means enabled
+		return (EnableDecodeCacheCvar != 0);
+	}
+}
+
+TUniquePtr<ICompressedAudioInfo> SoundWaveProxyDecodeCache::CreateDecoderInstance(const FSoundWaveProxyRef& InSoundWaveProxy)
+{
+	FSoundWaveProxyDecodeCache* DecodeCache = FSoundWaveProxyDecodeCache::Get();
+	float AutoDurationThresholdHold = DecodeCache->GetAutoDurationThreshold();
+	if (InSoundWaveProxy->GetDuration() <= AutoDurationThresholdHold)
+	{
+		return DecodeCache->CreateDecoderInstance(InSoundWaveProxy);
+	}
+	return nullptr;
+}
+
diff --git a/Engine/Source/Runtime/Engine/Private/SoundWaveProxyDecodeCacheImpl.h b/Engine/Source/Runtime/Engine/Private/SoundWaveProxyDecodeCacheImpl.h
new file mode 100644
--- /dev/null
+++ b/Engine/Source/Runtime/Engine/Private/SoundWaveProxyDecodeCacheImpl.h
@@ -0,0 +1,153 @@
+// Copyright Epic Games, Inc. All Rights Reserved.
+
+#pragma once
+
+#include "Audio.h"
+#include "Containers/LruCache.h"
+#include "Math/UnitConversion.h"
+#include "Sound/SoundWaveProxyDecodeCache.h"
+#include "SoundWaveProxyDecodeCacheDecoder.h"
+#include "Templates/SharedPointer.h"
+
+class FSoundWaveProxyDecodeCacheDecoder;
+
+// Key used to uniquely look up already decoded chunks of PCM data from wave proxies
+struct FDecodedSoundWaveChunkKey
+{
+	// The name of the wave proxy (useful for debugging)
+	FName WaveName;
+
+	// The guid of the wave is used for life-time independent cache key
+	FGuid WaveGuid;
+
+	// The chunk index of the decoded chunk, based on the current frame index and chunk size
+	uint32 ChunkIndex = INDEX_NONE;
+
+	// Prevents invalid cache hits in editor in the case of re-imported or other format changes
+#if WITH_EDITOR
+	uint32 ChunkRevision = 0;
+	uint32 NumChannels = 0;
+	uint32 SampleRate = 0;
+#endif
+
+	bool operator==(const FDecodedSoundWaveChunkKey& Other) const
+	{
+		return WaveName == Other.WaveName
+			&& WaveGuid == Other.WaveGuid
+			&& ChunkIndex == Other.ChunkIndex
+#if WITH_EDITOR
+			&& ChunkRevision == Other.ChunkRevision
+			&& NumChannels == Other.NumChannels
+			&& SampleRate == Other.SampleRate
+#endif
+			;
+	}
+
+	friend uint32 GetTypeHash(const FDecodedSoundWaveChunkKey& Key)
+	{
+		uint32 Hash = HashCombineFast(GetTypeHash(Key.WaveName), Key.ChunkIndex);
+		Hash = HashCombineFast(Hash, GetTypeHash(Key.WaveGuid));
+#if WITH_EDITOR
+		Hash = HashCombineFast(Hash, Key.ChunkRevision);
+		Hash = HashCombineFast(Hash, Key.NumChannels);
+		Hash = HashCombineFast(Hash, Key.SampleRate);
+#endif
+		return Hash;
+	}
+};
+
+/**
+ * The decoded audio data struct which is cached in the sound wave proxy decode cache
+ */
+struct FDecodedSoundWaveChunk
+{
+	// The interleaved audio frame that this chunk starts
+	int32 FrameStart = 0;
+
+	// The interleaved frame this chunk ends. Note: this could be less than data size due to it being the last chunk.
+	int32 FrameEnd = 0;
+
+	// If the file is finished in this chunk
+	bool bIsLastChunk = false;
+
+	// The interleaved decoded audio PCM byte data
+	TArray<uint8> DecodedByteData;
+};
+
+/**
+ * Struct holding information about the decode cache.
+ */
+struct FSoundWaveProxyCacheInfo
+{
+	// The max number of chunks possible in the cache (if everything was mono)
+	int32 NumChunksCapacity = 0;
+	// The number of chunks currently being used
+	int32 UsedChunks = 0;
+	// The number of bytes currently being used
+	int32 UsedMemoryBytes = 0;
+	// The max memory allowed by the cache
+	int32 MaxMemoryBytes = 0.0f;
+};
+
+
+/**
+ * A thread-safe decode cache for use with wave proxies. Allows avoiding duplicate decoding of audio which is replayed often.
+ */
+class FSoundWaveProxyDecodeCache
+{
+public:
+	using FSoundWaveProxyRef = TSharedRef<FSoundWaveProxy, ESPMode::ThreadSafe>;
+	using FDecodedSoundWaveChunkPtr = TSharedPtr<FDecodedSoundWaveChunk, ESPMode::ThreadSafe>;
+
+	FSoundWaveProxyDecodeCache();
+	~FSoundWaveProxyDecodeCache() {};
+
+	int32 GetFramesPerChunk() const;
+	FSoundWaveProxyCacheInfo GetInfo() const;
+	float GetAutoDurationThreshold() const;
+	TUniquePtr<ICompressedAudioInfo> CreateDecoderInstance(const FSoundWaveProxyRef& InSoundWaveProxy);
+
+	/**
+	* Returns the private singleton implementation of the sound wave proxy decode cache.
+	* If this is the first time the decode cache is retrieved, it will construct and initialize it.
+	*/
+	static FSoundWaveProxyDecodeCache* Get();
+
+private:
+
+	// Retrieve already decoded chunk with the given key
+	FDecodedSoundWaveChunkPtr GetDecodedChunk(const FDecodedSoundWaveChunkKey& InKey);
+
+	// Cache the given decoded audio into the decode cache
+	void CacheDecodedChunk(const FDecodedSoundWaveChunkKey& InKey, FDecodedSoundWaveChunkPtr& InDecodedChunk);
+
+	// A cache container that stores in-use chunks in a least-recently-used list automatically.
+	TLruCache<FDecodedSoundWaveChunkKey, FDecodedSoundWaveChunkPtr> LruCache;
+
+	// Settings used when initializing the decode cache
+	bool bHasFailure = false;
+	float DesiredCacheSizeMB = 0.0f;
+	float AutoDurationThreshold = 0.0f;
+	int32 DesiredFramesPerChunk = 0;
+
+	// The worse-case max chunk count for the decode cache. 
+	int32 MaxChunkCount = 0;
+
+	// The number of chunks currently being used from the TLruCache
+	int32 NumChunksUsed = 0;
+
+	// The actual max budget in bytes
+	int32 MaxMemoryBytes = 0;
+
+	// The actual frames per chunk
+	int32 FramesPerChunk = 0;
+
+	// The bytes currently being used, used to trigger a trim automatically
+	int32 UsedMemoryBytes = 0;
+
+	// Critical section for accessing the cache
+	mutable FCriticalSection CacheLock;
+
+	friend class FSoundWaveProxyDecodeCacheDecoder;
+};
+
diff --git a/Engine/Source/Runtime/Engine/Private/SoundWaveProxyPlayer.cpp b/Engine/Source/Runtime/Engine/Private/SoundWaveProxyPlayer.cpp
--- a/Engine/Source/Runtime/Engine/Private/SoundWaveProxyPlayer.cpp
+++ b/Engine/Source/Runtime/Engine/Private/SoundWaveProxyPlayer.cpp
@@ -1,37 +1,14 @@
 // Copyright Epic Games, Inc. All Rights Reserved.
 
 #include "Sound/SoundWaveProxyPlayer.h"
-
+#include "Sound/SoundWaveProxyDecodeCache.h"
 #include "Audio.h"
+#include "AudioDevice.h"
 #include "DSP/FloatArrayMath.h"
 #include "Interfaces/IAudioFormat.h"
 #include "ProfilingDebugging/CpuProfilerTrace.h"
 #include "Templates/SharedPointer.h"
 
-int32 ConformDecodeSize(int32 InMaxDesiredDecodeSizeInFrames)
-{
-	/** Minimum number of frames to decode per a call to the decoder.  */
-	static constexpr int32 DefaultMinDecodeSizeInFrames = 128;
-
-	/** Some codecs have strict requirements on decode size. In order to be
-	 * functional with all supported codecs, the decode size must be a multiple
-	 * of 128.
-	 */
-	static constexpr int32 DecodeSizeQuantizationInFrames = 128;
-	
-	static_assert(DefaultMinDecodeSizeInFrames >= DecodeSizeQuantizationInFrames, "Min decode size less than decode size quantization");
-	static_assert((DefaultMinDecodeSizeInFrames % DecodeSizeQuantizationInFrames) == 0, "Min decode size must be equally divisible by decode size quantization");
-
-	InMaxDesiredDecodeSizeInFrames = FMath::Max(InMaxDesiredDecodeSizeInFrames, DefaultMinDecodeSizeInFrames);
-	const int32 QuantizedDecodeSizeInFrames = (InMaxDesiredDecodeSizeInFrames / DecodeSizeQuantizationInFrames) * DecodeSizeQuantizationInFrames;
-	const int32 ConformedDecodeSizeInFrames = FMath::Max(QuantizedDecodeSizeInFrames, DefaultMinDecodeSizeInFrames);
-
-	check((ConformedDecodeSizeInFrames % QuantizedDecodeSizeInFrames) == 0);
-	check(ConformedDecodeSizeInFrames > 0);
-
-	return ConformedDecodeSizeInFrames;
-}
-
 FSoundWaveProxyPlayer::FSoundWaveProxyPlayer(const FSettings& InSettings)
 	: Settings(InSettings)
 	, RuntimeResampler(1)
@@ -754,7 +731,16 @@ bool FSoundWaveProxyPlayer::InitializeDecoder(const FSoundWaveProxyRef& InWavePr
 
 	FName Format = InWaveProxy->GetRuntimeFormat();
 	TUniquePtr<ICompressedAudioInfo> InfoInstance;
-	InfoInstance.Reset(IAudioInfoFactoryRegistry::Get().Create(Format, InWaveProxy->GetFName()));
+
+	if (SoundWaveProxyDecodeCache::IsEnabled())
+	{
+		InfoInstance = SoundWaveProxyDecodeCache::CreateDecoderInstance(InWaveProxy);
+	}
+
+	if (!InfoInstance.IsValid())
+	{
+		InfoInstance.Reset(IAudioInfoFactoryRegistry::Get().Create(Format, InWaveProxy->GetFName()));
+	}
 
 	if (!ensure(InfoInstance.IsValid()))
 	{
@@ -801,7 +787,7 @@ bool FSoundWaveProxyPlayer::InitializeDecoder(const FSoundWaveProxyRef& InWavePr
 	// END HACK
 
 	// initialized decode buffers
-	const int32 DecodeSize = ConformDecodeSize(MaxDecodeSizeInFrames);
+	const int32 DecodeSize = ICompressedAudioInfo::ConformDecodeSize(MaxDecodeSizeInFrames);
 	check(DecodeSize > 0);
 	DecodeBuffer.SetNum(DecodeSize * NumChannels, EAllowShrinking::Yes);
 	return true;
diff --git a/Engine/Source/Runtime/Engine/Public/AudioDecompress.h b/Engine/Source/Runtime/Engine/Public/AudioDecompress.h
--- a/Engine/Source/Runtime/Engine/Public/AudioDecompress.h
+++ b/Engine/Source/Runtime/Engine/Public/AudioDecompress.h
@@ -130,15 +130,14 @@ public:
 	* @param	Wave			Wave that will be read from to retrieve necessary chunk
 	* @param	QualityInfo		Quality Info (to be filled out)
 	*/
-	ENGINE_API bool StreamCompressedInfo(USoundWave* Wave, struct FSoundQualityInfo* QualityInfo);
-	ENGINE_API bool StreamCompressedInfo(const FSoundWaveProxyPtr& Wave, struct FSoundQualityInfo* QualityInfo);
+	ENGINE_API virtual bool StreamCompressedInfo(USoundWave* Wave, struct FSoundQualityInfo* QualityInfo);
+	ENGINE_API virtual bool StreamCompressedInfo(const FSoundWaveProxyPtr& Wave, struct FSoundQualityInfo* QualityInfo);
 
 	/**
 	*  Returns true if a non-recoverable error has occurred.
 	*/
 	ENGINE_API virtual bool HasError() const;
 
-protected:
 	/** Internal override implemented by subclasses. */
 	virtual bool StreamCompressedInfoInternal(const FSoundWaveProxyPtr& InWaveProxy, struct FSoundQualityInfo* QualityInfo) = 0;
 
@@ -168,6 +167,10 @@ public:
 	/** Return the streaming sound wave used by this decoder. Returns nullptr if there is not a streaming sound wave. */
 	virtual const FSoundWaveProxyPtr& GetStreamingSoundWave() const { return StreamingSoundWave; }
 
+	/** Return a conformed decode size that best matches all codecs. */
+	static int32 ConformDecodeSize(int32 InMaxDesiredDecodeSizeInFrames);
+
+
 protected:
 	mutable bool bHasError = false;
 
