[Material Interface] Trivial optimization to prevent allocating a TArray every time GetRelevance is called

#rnx

#Synced-CL 50921747

diff --git a/Engine/Source/Runtime/Engine/Private/Materials/MaterialInterface.cpp b/Engine/Source/Runtime/Engine/Private/Materials/MaterialInterface.cpp
--- a/Engine/Source/Runtime/Engine/Private/Materials/MaterialInterface.cpp
+++ b/Engine/Source/Runtime/Engine/Private/Materials/MaterialInterface.cpp
@@ -663,7 +663,8 @@ FMaterialRelevance UMaterialInterface::GetRelevance_Internal(const UMaterial* Ma
 
 	// Note that even though XX_GameThread() api is called, this function can be called on non game thread via 
 	// GetRelevance_Concurrent()
-	bool bUsesAnisotropy = MaterialResource->GetShadingModels().HasAnyShadingModel({ MSM_DefaultLit, MSM_ClearCoat }) && 
+	static const TArray<EMaterialShadingModel> AnisoShadingModels = { MSM_DefaultLit, MSM_ClearCoat };
+	bool bUsesAnisotropy = MaterialResource->GetShadingModels().HasAnyShadingModel(AnisoShadingModels) && 
 		MaterialResource->MaterialUsesAnisotropy_GameThread();
 
 	const EBlendMode BlendMode = (EBlendMode)GetBlendMode();
