[HWRT] Fix toggling r.RayTracing.Geometry.SkeletalMeshes at runtime leading to invalid state and triggering asserts in UnregisterProxyWithCachedRayTracingState(...).
- moved checking r.RayTracing.Geometry.SkeletalMeshes and r.RayTracing.Geometry.SupportSkeletalMeshes out of HasRayTracingRepresentation.
- also added better checks to detect this situation.

- eventually we should support changes to HasRayTracingRepresentation() without requiring primitive proxy to be recreated but this depends on other ongoing work to FRayTracingGeometryManager, etc.

#Synced-CL 50921461

diff --git a/Engine/Source/Runtime/Engine/Private/PrimitiveSceneProxy.cpp b/Engine/Source/Runtime/Engine/Private/PrimitiveSceneProxy.cpp
--- a/Engine/Source/Runtime/Engine/Private/PrimitiveSceneProxy.cpp
+++ b/Engine/Source/Runtime/Engine/Private/PrimitiveSceneProxy.cpp
@@ -1874,6 +1874,10 @@ ERayTracingPrimitiveFlags FPrimitiveSceneProxy::GetRayTracingPrimitiveFlags()
 		return ERayTracingPrimitiveFlags::Exclude;
 	}
 
+	// TODO: Ideally all these other conditions should be folded into HasRayTracingRepresentation() 
+	// to avoid primitives being unecessarily registered with FRayTracingGeometryManager / Nanite::FRayTracingManager
+	// however would need to make sure they would get registered/unregistered as the conditions change
+
 	if (!(IsVisibleInRayTracing() && ShouldRenderInMainPass() && (IsDrawnInGame() || AffectsIndirectLightingWhileHidden() || CastsHiddenShadow())) && !IsRayTracingFarField())
 	{
 		return ERayTracingPrimitiveFlags::Exclude;
@@ -1886,6 +1890,13 @@ ERayTracingPrimitiveFlags FPrimitiveSceneProxy::GetRayTracingPrimitiveFlags()
 		return ERayTracingPrimitiveFlags::Exclude;
 	}
 
+	static const auto RayTracingSkeletalMeshesCVar = IConsoleManager::Get().FindTConsoleVariableDataInt(TEXT("r.RayTracing.Geometry.SkeletalMeshes"));
+
+	if (bSkinnedMesh && RayTracingSkeletalMeshesCVar && RayTracingSkeletalMeshesCVar->GetValueOnRenderThread() <= 0)
+	{
+		return ERayTracingPrimitiveFlags::Exclude;
+	}
+
 	static const auto RayTracingHISMCVar = IConsoleManager::Get().FindTConsoleVariableDataInt(TEXT("r.RayTracing.Geometry.HierarchicalInstancedStaticMesh"));
 
 	if (bIsHierarchicalInstancedStaticMesh && RayTracingHISMCVar && RayTracingHISMCVar->GetValueOnRenderThread() <= 0)
diff --git a/Engine/Source/Runtime/Engine/Private/SkeletalMeshSceneProxy.cpp b/Engine/Source/Runtime/Engine/Private/SkeletalMeshSceneProxy.cpp
--- a/Engine/Source/Runtime/Engine/Private/SkeletalMeshSceneProxy.cpp
+++ b/Engine/Source/Runtime/Engine/Private/SkeletalMeshSceneProxy.cpp
@@ -143,7 +143,7 @@ FSkeletalMeshSceneProxy::FSkeletalMeshSceneProxy(const FSkinnedMeshSceneProxyDes
 	NumLODs = RenderData->LODRenderData.Num();
 
 #if RHI_RAYTRACING
-	bHasRayTracingRepresentation = RenderData->bSupportRayTracing;
+	bHasRayTracingRepresentation = RenderData->bSupportRayTracing && (bRenderStatic || CVarRayTracingSupportSkeletalMeshes.GetValueOnAnyThread() != 0);
 	bRayTraceStatic = bRenderStatic || bInstancedSkinnedMesh;
 #endif
 
@@ -988,8 +988,7 @@ void FSkeletalMeshSceneProxy::DestroyRenderThreadResources()
 #if RHI_RAYTRACING
 bool FSkeletalMeshSceneProxy::HasRayTracingRepresentation() const
 {
-	return bHasRayTracingRepresentation
-		&& (bRayTraceStatic || (CVarRayTracingSupportSkeletalMeshes.GetValueOnAnyThread() != 0 && CVarRayTracingSkeletalMeshes.GetValueOnAnyThread() != 0));
+	return bHasRayTracingRepresentation;
 }
 
 RayTracing::FGeometryGroupHandle FSkeletalMeshSceneProxy::GetRayTracingGeometryGroupHandle() const
@@ -1089,8 +1088,7 @@ void FSkeletalMeshSceneProxy::GetStaticRayTracingInstances(FStaticRayTracingInst
 
 void FSkeletalMeshSceneProxy::GetDynamicRayTracingInstances(FRayTracingInstanceCollector& Collector)
 {
-	if (!CVarRayTracingSkeletalMeshes.GetValueOnRenderThread()
-		|| !CVarRayTracingSupportSkeletalMeshes.GetValueOnRenderThread())
+	if (!ensure(CVarRayTracingSkeletalMeshes.GetValueOnRenderThread() && CVarRayTracingSupportSkeletalMeshes.GetValueOnRenderThread()))
 	{
 		return;
 	}
diff --git a/Engine/Source/Runtime/Renderer/Private/RendererScene.cpp b/Engine/Source/Runtime/Renderer/Private/RendererScene.cpp
--- a/Engine/Source/Runtime/Renderer/Private/RendererScene.cpp
+++ b/Engine/Source/Runtime/Renderer/Private/RendererScene.cpp
@@ -5662,7 +5662,9 @@ void FScene::Update(FRDGBuilder& GraphBuilder, const FUpdateParameters& Paramete
 				LumenRemovePrimitive(PrimitiveSceneInfo, PrimitiveIndex);
 
 #if RHI_RAYTRACING
-				if (SceneProxy->HasRayTracingRepresentation())
+				ensureMsgf(PrimitiveSceneInfo->bHasRayTracingRepresentation == SceneProxy->HasRayTracingRepresentation(),
+					TEXT("Stale bHasRayTracingRepresentation detected. Primitive might have missed requests to update cached ray tracing state."));
+				if (PrimitiveSceneInfo->bHasRayTracingRepresentation)
 				{
 					const RayTracing::FGeometryGroupHandle RayTracingGeometryGroupHandle = SceneProxy->GetRayTracingGeometryGroupHandle();
 					if (!RayTracingGeometryGroupHandle.IsNull())
@@ -5995,7 +5997,8 @@ void FScene::Update(FRDGBuilder& GraphBuilder, const FUpdateParameters& Paramete
 				LumenAddPrimitive(PrimitiveSceneInfo);
 
 #if RHI_RAYTRACING
-				if (SceneProxy->HasRayTracingRepresentation())
+				PrimitiveSceneInfo->bHasRayTracingRepresentation = SceneProxy->HasRayTracingRepresentation();
+				if (PrimitiveSceneInfo->bHasRayTracingRepresentation)
 				{
 					// TODO: investigate registering proxy when data is actually cached (for example, in FPrimitiveSceneInfo::CacheRayTracingPrimitives(...))
 					const RayTracing::FGeometryGroupHandle RayTracingGeometryGroupHandle = SceneProxy->GetRayTracingGeometryGroupHandle();
diff --git a/Engine/Source/Runtime/Renderer/Private/ScenePrivate.h b/Engine/Source/Runtime/Renderer/Private/ScenePrivate.h
--- a/Engine/Source/Runtime/Renderer/Private/ScenePrivate.h
+++ b/Engine/Source/Runtime/Renderer/Private/ScenePrivate.h
@@ -3028,6 +3028,7 @@ public:
 		Nanite::CoarseMeshStreamingHandle CoarseMeshStreamingHandle;
 		RayTracing::FGeometryGroupHandle RayTracingGeometryGroupHandle;
 	};
+	static_assert(sizeof(FPrimitiveRayTracingData) == 8, "Increasing the size of FPrimitiveRayTracingData will negatively affect performance of ray tracing instance gathering.");
 	TArray<FPrimitiveRayTracingData> PrimitiveRayTracingDatas;
 	/** Packed array of ray tracing primitive group id hash indices. */
 	TArray<Experimental::FHashElementId> PrimitiveRayTracingGroupIds;
diff --git a/Engine/Source/Runtime/Renderer/Public/PrimitiveSceneInfo.h b/Engine/Source/Runtime/Renderer/Public/PrimitiveSceneInfo.h
--- a/Engine/Source/Runtime/Renderer/Public/PrimitiveSceneInfo.h
+++ b/Engine/Source/Runtime/Renderer/Public/PrimitiveSceneInfo.h
@@ -646,6 +646,7 @@ public:
 	bool bNeedsCachedReflectionCaptureUpdate = true;
 
 #if RHI_RAYTRACING
+	bool bHasRayTracingRepresentation : 1 = false;
 	bool bIsRelevantToPathTracing : 1;
 	bool bIsCachedRayTracingInstanceValid : 1;
 	UE_DEPRECATED(5.7, "Use DynamicRayTracingInstanceCachedData or GetRayTracingLODData(...) instead.")
