UE: Improve logging for an ensure related to NaNs in card data

diff --git a/Engine/Source/Runtime/Engine/Private/StaticMesh.cpp b/Engine/Source/Runtime/Engine/Private/StaticMesh.cpp
--- a/Engine/Source/Runtime/Engine/Private/StaticMesh.cpp
+++ b/Engine/Source/Runtime/Engine/Private/StaticMesh.cpp
@@ -2454,13 +2454,20 @@ void FStaticMeshRenderData::SerializeInlineDataRepresentations(FArchive& Ar, USt
 				}
 
 #if !UE_BUILD_SHIPPING
-				if (LOD.CardRepresentationData->ContainsNaN())
+				if (Ar.IsSaving() && LOD.CardRepresentationData->ContainsNaN())
 				{
-					UE_LOG(LogStaticMesh, Display, TEXT("FStaticMeshRenderData::SerializeInlineDataRepresentations found NaN in CardRepresentationData of '%s'. LOD:%d"), *GetPathNameSafe(Owner), ResourceIndex);
+					UE_LOG(LogStaticMesh, Display, TEXT("FStaticMeshRenderData::SerializeInlineDataRepresentations found NaN when saving CardRepresentationData of '%s'. LOD:%d"), *GetPathNameSafe(Owner), ResourceIndex);
 				}
 #endif // UE_BUILD_SHIPPING
 
 				Ar << *(LOD.CardRepresentationData);
+				
+#if !UE_BUILD_SHIPPING
+				if (Ar.IsLoading() && LOD.CardRepresentationData->ContainsNaN())
+				{
+					UE_LOG(LogStaticMesh, Display, TEXT("FStaticMeshRenderData::SerializeInlineDataRepresentations found NaN when loading CardRepresentationData of '%s'. LOD:%d"), *GetPathNameSafe(Owner), ResourceIndex);
+				}
+#endif // UE_BUILD_SHIPPING
 			}
 		}
 	}
