Adaptive memory scheduling for constrained memory cooking, with shader compiler workers.
#rnx

#Synced-CL 50894984

#Synced-CL 50894984

diff --git a/Engine/Source/Runtime/Core/Public/Misc/QueuedThreadPoolWrapper.h b/Engine/Source/Runtime/Core/Public/Misc/QueuedThreadPoolWrapper.h
--- a/Engine/Source/Runtime/Core/Public/Misc/QueuedThreadPoolWrapper.h
+++ b/Engine/Source/Runtime/Core/Public/Misc/QueuedThreadPoolWrapper.h
@@ -66,6 +66,14 @@ public:
 	 */
 	CORE_API void SetMaxConcurrency(int32 MaxConcurrency = -1);
 
+	/**
+	 *  Return the estimated memory currently used by this thread pool, if supported, or -1 if not.
+	 */
+	CORE_API virtual int64 GetEstimatedMemory() const
+	{
+		return -1ll;
+	}
+
 	CORE_API void AddQueuedWork(IQueuedWork* InQueuedWork, EQueuedWorkPriority InPriority = EQueuedWorkPriority::Normal) override;
 	CORE_API bool RetractQueuedWork(IQueuedWork* InQueuedWork) override;
 	CORE_API int32 GetNumThreads() const override;
diff --git a/Engine/Source/Runtime/Engine/Private/AssetCompilingManager.cpp b/Engine/Source/Runtime/Engine/Private/AssetCompilingManager.cpp
--- a/Engine/Source/Runtime/Engine/Private/AssetCompilingManager.cpp
+++ b/Engine/Source/Runtime/Engine/Private/AssetCompilingManager.cpp
@@ -76,8 +76,6 @@ TRACE_DECLARE_INT_COUNTER(AsyncCompilationConcurrency, TEXT("AsyncCompilation/Co
 TRACE_DECLARE_MEMORY_COUNTER(AsyncCompilationAvailableMemory, TEXT("AsyncCompilation/AvailableMemory"));
 TRACE_DECLARE_MEMORY_COUNTER(AsyncCompilationTotalEstimatedMemory, TEXT("AsyncCompilation/TotalEstimatedMemory"));
 
-TRACE_DECLARE_MEMORY_COUNTER(AsyncCompilationSCWCurrentLimit, TEXT("AsyncCompilation/SCWCurrentLimit"));
-TRACE_DECLARE_MEMORY_COUNTER(AsyncCompilationSCWBaseLimit, TEXT("AsyncCompilation/SCWBaseLimit"));
 
 TRACE_DECLARE_INT_COUNTER(AsyncDelayedTask, TEXT("AsyncCompilation/DelayedTask"));
 TRACE_DECLARE_INT_COUNTER(AsyncDelayedTaskActive, TEXT("AsyncCompilation/DelayedTaskActive"));
@@ -137,9 +135,15 @@ public:
 	{
 	}
 
+
+	virtual int64 GetEstimatedMemory() const override
+	{
+		return TotalEstimatedMemory.load(std::memory_order_relaxed);
+	}
+	
 	static int64 GetHardMemoryLimit()
 	{
-		return (int64)CVarAsyncAssetCompilationMaxMemoryUsage.GetValueOnAnyThread() * 1024 * 1024 * 1024;
+		return (int64)CVarAsyncAssetCompilationMaxMemoryUsage.GetValueOnAnyThread() * 1024ll * 1024ll * 1024ll;
 	}
 	static bool IsMemoryConstrained()
 	{
@@ -148,7 +152,7 @@ public:
 
 	int64 GetDefaultMemoryPerAsset() const
 	{
-		return FMath::Max(0ll, (int64)CVarAsyncAssetCompilationMemoryPerCore.GetValueOnAnyThread(false) * 1024 * 1024 * 1024);
+		return FMath::Max(0ll, (int64)CVarAsyncAssetCompilationMemoryPerCore.GetValueOnAnyThread(false) * 1024ll * 1024ll * 1024ll);
 	}
 
 	int64 GetMemoryAvailableForAssetCompilation() const
@@ -220,32 +224,25 @@ public:
 	private:
 		const TFunction<int64 ()> GetAvailableMemory;
 	public:
-		int64 MemoryToLoan;
-		int64 MemoryLoanThreshold;
-		int64 MemoryToBorrowFromSCW;
-
+		int32 MemoryInMiBToBorrowFromSCW;	//	Memory in MiB to borrow from shader compilation when required for asset compilation.
 		bool IsMemorySharingEnabled;
 
 		MemorySharingInfo(TFunction<int64()> InGetAvailableMemory)
 			: GetAvailableMemory(InGetAvailableMemory)
 		{
-			MemoryToLoan = GConfig->GetIntOrDefault(TEXT("AssetCompilation"), TEXT("MemoryToLoanToSCW"), 2000, GEngineIni);	
-			MemoryLoanThreshold = GConfig->GetIntOrDefault(TEXT("AssetCompilation"), TEXT("MemoryLoanToSCWThreshold"), 3000, GEngineIni);	
-			MemoryToBorrowFromSCW = GConfig->GetIntOrDefault(TEXT("AssetCompilation"), TEXT("MemoryToBorrowFromSCW"), 2000, GEngineIni);	
-
+			MemoryInMiBToBorrowFromSCW = GConfig->GetIntOrDefault(TEXT("AssetCompilation"), TEXT("MemoryToBorrowFromSCW"), 2000, GEngineIni);	
 			IsMemorySharingEnabled = GConfig->GetBoolOrDefault(TEXT("AssetCompilation"), TEXT("SCWMemorySharingEnabled"), false, GEngineIni);
-		}	
 
-		int64 GetHighSCWMemory() const
-		{
-			FShaderCompilerMemoryLimits SCWMemoryLimits = GetShaderCompilerCurrentMemoryLimits();
-			return SCWMemoryLimits.BaseLimit + MemoryToLoan;
-		}
+			const FPlatformMemoryStats MemoryStats = FPlatformMemory::GetStats();
+			const FPlatformMemoryConstants& MemoryConstants = FPlatformMemory::GetConstants();
+			UE_LOG(LogAsyncCompilation, Display, TEXT("MemorySharingInfo: UsedMemory=%dMiB, MaxMemoryAvailable=%dMiB"), MemoryStats.UsedVirtual / 1024ull / 1024ull,	MemoryConstants.TotalVirtual / 1024ull / 1024ull);
+		}	
 
-		int64 GetNormalSCWMemory() const
+		void LogMemoryInfo() const
 		{
-			FShaderCompilerMemoryLimits SCWMemoryLimits = GetShaderCompilerCurrentMemoryLimits();		
-			return SCWMemoryLimits.BaseLimit;
+			const FPlatformMemoryStats MemoryStats = FPlatformMemory::GetStats();
+			const FPlatformMemoryConstants& MemoryConstants = FPlatformMemory::GetConstants();
+			UE_LOG(LogAsyncCompilation, Display, TEXT("MemorySharingInfo: UsedMemory=%dMiB, MaxMemoryAvailable=%dMiB Currently Available=%dMiB"), MemoryStats.UsedVirtual / 1024 / 1024,	MemoryConstants.TotalVirtual / 1024 / 1024, GetAvailableMemory() / 1024 / 1024);
 		}
 
 		bool IsSCWMemoryWithinLimit() const 
@@ -255,21 +252,6 @@ public:
 
 			return SCWMemoryUsage.VirtualMemory <= static_cast<uint64>(SCWMemoryLimits.CurrentLimit * 1024ull * 1024ull);
 		}
-		bool IsMemoryAvailableToLoan() const
-		{
-			int64 AvailableAssetCompilationMemory = GetAvailableMemory() / (1024ull * 1024ull);	
-			return AvailableAssetCompilationMemory >= MemoryLoanThreshold;
-		}
-		bool ShouldRetrieveFromSCW() const 
-		{
-			FShaderCompilerMemoryLimits SCWMemoryLimits = GetShaderCompilerCurrentMemoryLimits();		
-			return IsMemorySharingEnabled && IsMemoryAvailableToLoan() == false && SCWMemoryLimits.CurrentLimit == GetHighSCWMemory();
-		}
-		bool CanLoanToSCW() const 
-		{
-			FShaderCompilerMemoryLimits SCWMemoryLimits = GetShaderCompilerCurrentMemoryLimits();		
-			return IsMemorySharingEnabled && IsMemoryAvailableToLoan() && SCWMemoryLimits.CurrentLimit == GetNormalSCWMemory();
-		}
 	};
 
 	class FMemoryBoundScheduledWork : public FQueuedThreadPoolWrapper::FScheduledWork
@@ -327,16 +309,10 @@ public:
 
 				if (RequiredMemory > FMemoryBoundQueuedThreadPoolWrapper::GetHardMemoryLimit())
 				{
-					SetShaderCompilerTargetMemoryLimitInMiB(MemoryInfo.GetNormalSCWMemory() - MemoryInfo.MemoryToBorrowFromSCW);
-					WaitForSCWMemory();
-					FQueuedThreadPoolWrapper::FScheduledWork::DoThreadedWork();
-					SetShaderCompilerTargetMemoryLimitInMiB(MemoryInfo.GetNormalSCWMemory());
-				}
-				else if (MemoryInfo.ShouldRetrieveFromSCW())
-				{
-					SetShaderCompilerTargetMemoryLimitInMiB(MemoryInfo.GetNormalSCWMemory());
+					SetShaderCompilerAssetMemoryLoanInMiB(MemoryInfo.MemoryInMiBToBorrowFromSCW);
 					WaitForSCWMemory();
 					FQueuedThreadPoolWrapper::FScheduledWork::DoThreadedWork();
+					SetShaderCompilerAssetMemoryLoanInMiB(0);
 				}
 				else
 				{
@@ -360,10 +336,6 @@ public:
 		TRACE_COUNTER_SET(AsyncCompilationAvailableMemory, GetMemoryAvailableForAssetCompilation());
 		TRACE_COUNTER_SET(AsyncCompilationTotalEstimatedMemory, TotalEstimatedMemory);
 		TRACE_COUNTER_SET(AsyncCompilationConcurrency, GetCurrentConcurrency());
-
-		FShaderCompilerMemoryLimits SCWMemoryLimits = GetShaderCompilerCurrentMemoryLimits();
-		TRACE_COUNTER_SET(AsyncCompilationSCWCurrentLimit, SCWMemoryLimits.CurrentLimit * 1024ull * 1024ull);
-		TRACE_COUNTER_SET(AsyncCompilationSCWBaseLimit, SCWMemoryLimits.BaseLimit * 1024ull * 1024ull);
 	}
 
 	void OnScheduled(const IQueuedWork* InWork) override
@@ -409,11 +381,6 @@ public:
 			);
 
 		UpdateCounters();
-		
-		if (IsMemoryConstrained() && MemoryInfo.CanLoanToSCW())
-		{
-			SetShaderCompilerTargetMemoryLimitInMiB(MemoryInfo.GetHighSCWMemory());
-		} 		
 	}
 
 	int32 GetAdjustedMaxConcurrency(EQueuedWorkPriority Priority) const
diff --git a/Engine/Source/Runtime/Engine/Private/ShaderCompiler/ShaderCompilerMemoryLimit.cpp b/Engine/Source/Runtime/Engine/Private/ShaderCompiler/ShaderCompilerMemoryLimit.cpp
--- a/Engine/Source/Runtime/Engine/Private/ShaderCompiler/ShaderCompilerMemoryLimit.cpp
+++ b/Engine/Source/Runtime/Engine/Private/ShaderCompiler/ShaderCompilerMemoryLimit.cpp
@@ -35,7 +35,10 @@ void FWindowsResourceProcessTreeMemory::Reset()
 
 void FWindowsResourceProcessTreeMemory::CacheSystemProcessIds()
 {
-	RebuildSystemProcessIdList();
+	if (!bCachedSystemProcessIds)
+	{
+		RebuildSystemProcessIdList();
+	}
 	bCachedSystemProcessIds = true;
 }
 
diff --git a/Engine/Source/Runtime/Engine/Private/ShaderCompiler/ShaderCompilerThreadRunnable.cpp b/Engine/Source/Runtime/Engine/Private/ShaderCompiler/ShaderCompilerThreadRunnable.cpp
--- a/Engine/Source/Runtime/Engine/Private/ShaderCompiler/ShaderCompilerThreadRunnable.cpp
+++ b/Engine/Source/Runtime/Engine/Private/ShaderCompiler/ShaderCompilerThreadRunnable.cpp
@@ -19,6 +19,9 @@
 #include "ProfilingDebugging/StallDetector.h"
 #include "ShaderCompilerMemoryLimit.h"
 #include "ShaderCompileWorkerUtil.h"
+#include "AssetCompilingManager.h"
+#include "Misc/ConfigCacheIni.h"
+#include "Misc/QueuedThreadPoolWrapper.h"
 
 
 float GShaderCompilerTooLongIOThresholdSeconds = 0.3;
@@ -48,6 +51,41 @@ static FAutoConsoleVariableRef CVarShaderWorkerStateChangeHeartbeat(
 	TEXT("This effectively shuts down distributed shader compilation. Default is 15 * 60 seconds."),
 	ECVF_Default);
 
+int32 GShaderCompilerAssetMemoryLoanInMiB = -1;
+
+enum class EShaderCompilerWorkerScaleUpApproach : int32
+{
+	ScaleUpWhenBelowHalfMemoryLimit = 0,
+	ScaleUpWhenSpaceForTwoAverageWorkers = 1
+};
+
+bool GShaderCompilerWorkerSoftLimit = true;
+static FAutoConsoleVariableRef CVarShaderCompilerWorkerSoftLimit(
+	TEXT("r.ShaderCompiler.SoftLimit"),
+	GShaderCompilerWorkerSoftLimit,
+	TEXT("When enabled if not violating the memory, but there isn't sufficient memory for one more worker at the current average memory consumption scale back by one worker.\n") \
+	TEXT("This adapts the active worker count to changes in the memory usage across the workers, usually before a violation occurs improving stability.\n"),
+	ECVF_Default);
+
+float GShaderCompilerWorkerSoftLimitInWorkers = 0.8f;
+static FAutoConsoleVariableRef CVarShaderCompilerWorkerSoftLimitInWorkers(
+	TEXT("r.ShaderCompiler.SoftLimitInWorkers"),
+	GShaderCompilerWorkerSoftLimitInWorkers,
+	TEXT("Multiplier applied to the average memory across shader compiler workers to find the soft limit.\n") \
+	TEXT("For example, 1.5 means leave at least enough memory for 1.5 times the memory currently used by an average shader compiler worker\n") \
+	TEXT("See r.ShaderCompiler.SoftLimit.\n"),
+	ECVF_Default);
+
+int32 GShaderCompilerWorkerScaleUpApproach = 1;
+static FAutoConsoleVariableRef CVarShaderCompilerWorkerScaleUpApproach(
+	TEXT("r.ShaderCompiler.WorkerScaleUpApproach"),
+	GShaderCompilerWorkerScaleUpApproach,
+	TEXT("What approach to use when scaling up or down the active shader compiler workers when running with a memory limit.\n") \
+	TEXT("  0: Scale up when memory is below half the limit") \
+	TEXT("  1: Scale up when there is enough memory for two more workers, based on the average current worker memory use") \
+	TEXT("Only affects behavior when r.ShaderCompiler.MemoryLimit is non-zero"),
+	ECVF_Default);
+
 extern bool GShaderCompilerDumpWorkerDiagnostics;
 
 // Configuration to retry shader compile through workers after a worker has been abandoned
@@ -62,6 +100,9 @@ TRACE_DECLARE_MEMORY_COUNTER(AsyncCompilationSCWIndividualMemoryMax, TEXT("Async
 TRACE_DECLARE_MEMORY_COUNTER(AsyncCompilationSCWIndividualMemoryMedian, TEXT("AsyncCompilation/SCWIndividualMemory_Median"));
 TRACE_DECLARE_MEMORY_COUNTER(AsyncCompilationSCWMemoryViolationSize, TEXT("AsyncCompilation/SCWMemoryViolationSize"));
 TRACE_DECLARE_MEMORY_COUNTER(AsyncCompilationSCWMemory, TEXT("AsyncCompilation/SCWMemory"));
+TRACE_DECLARE_MEMORY_COUNTER(AsyncCompilationSCWMemoryLimit, TEXT("AsyncCompilation/SCWMemoryLimit"));
+TRACE_DECLARE_MEMORY_COUNTER(AsyncCompilationSCWMemoryLoan, TEXT("AsyncCompilation/SCWMemoryLoan"));
+
 TRACE_DECLARE_INT_COUNTER(AsyncCompilationSCWsActive, TEXT("AsyncCompilation/SCWsActive"));
 
 FShaderCompilerMemoryLimits GetShaderCompilerCurrentMemoryLimits()
@@ -74,9 +115,15 @@ FShaderCompilerMemoryLimits GetShaderCompilerCurrentMemoryLimits()
 	return Limits;
 }
 
-void SetShaderCompilerTargetMemoryLimitInMiB(int64 NewLimit)
+void SetShaderCompilerTargetMemoryLimitInMiB(int32 NewLimit)
+{
+	GShaderCompilerCurrentMemoryLimitInMiB = NewLimit;
+}
+
+void SetShaderCompilerAssetMemoryLoanInMiB(int32 NewLoan)
 {
-	GShaderCompilerCurrentMemoryLimitInMiB = static_cast<int32>(NewLimit);
+	GShaderCompilerAssetMemoryLoanInMiB = NewLoan >= 0 ? NewLoan : 0;
+	TRACE_COUNTER_SET(AsyncCompilationSCWMemoryLoan, GShaderCompilerAssetMemoryLoanInMiB);
 }
 
 /** Output structure for polling shader process memory violation status. See FShaderCompileThreadRunnable::QueryMemoryLimitViolationStatus(). */
@@ -87,6 +134,11 @@ struct FShaderCompilerJobMemoryLimitInfo
 
 	/** Size (in bytes) of the shader process memory usage. When QueryMemoryLimitViolationStatus() returns true, this will be greater than JobMemoryLimit, since the process violated the limitation requirements. */
 	int64 MemoryUsed = 0;
+
+	/**
+	 * Number of active shader compiler workers.
+	 */
+	int32 ActiveWorkerCount = 0;
 };
 
 /** Information tracked for each shader compile worker process instance. */
@@ -162,7 +214,10 @@ struct FShaderCompileWorkerInfo
 		if ((!UE_TRACE_CHANNELEXPR_IS_ENABLED(CountersChannel)) || StatsLastMemoryUsedValue == InMemoryUsed)
 			return;
 
-		FCountersTrace::OutputSetValue(InsightsCounterId, static_cast<int64>(InMemoryUsed));
+		if(InMemoryUsed > 0)
+		{
+			FCountersTrace::OutputSetValue(InsightsCounterId, static_cast<int64>(InMemoryUsed));
+		}
 		StatsLastMemoryUsedValue = InMemoryUsed;
 	}
 #endif
@@ -181,8 +236,6 @@ struct FShaderCompileWorkerInfo
 			}
 			FPlatformProcess::CloseProc(WorkerProcess);
 			WorkerProcess = FProcHandle();
-
-			TRACE_COUNTER_DECREMENT(AsyncCompilationSCWsActive);
 		}
 	}
 
@@ -194,8 +247,6 @@ struct FShaderCompileWorkerInfo
 			FPlatformProcess::GetProcReturnCode(WorkerProcess, &ReturnCode);
 			FPlatformProcess::CloseProc(WorkerProcess);
 			WorkerProcess = FProcHandle();
-
-			TRACE_COUNTER_DECREMENT(AsyncCompilationSCWsActive);
 		}
 		return ReturnCode;
 	}
@@ -264,6 +315,10 @@ FShaderCompileThreadRunnable::FShaderCompileThreadRunnable(FShaderCompilingManag
 	{
 		WorkerInfos.Add(MakeUnique<FShaderCompileWorkerInfo>());
 	}
+
+	bDynamicSCWBudgetEnabled = GConfig->GetBoolOrDefault(TEXT("AssetCompilation"), TEXT("DynamicSCWBudget"), false, GEngineIni);
+	DynamicSCWBudgetReserveInMiB = GConfig->GetIntOrDefault(TEXT("AssetCompilation"), TEXT("DynamicSCWBudgetReserveInMiB"), 512, GEngineIni);
+	DynamicSCWBudgetGranularityInMiB = FMath::Max(GConfig->GetIntOrDefault(TEXT("AssetCompilation"), TEXT("DynamicSCWBudgetGranularityInMiB"), 256, GEngineIni), 1);	//	Enforce at least 1 to prevent a divide by zero
 }
 
 FShaderCompileThreadRunnable::~FShaderCompileThreadRunnable()
@@ -777,8 +832,6 @@ bool FShaderCompileThreadRunnable::LaunchWorkersIfNeeded()
 				CurrentWorkerInfo.bLaunchedWorker = true;
 
 				NumberLaunched++;
-
-				TRACE_COUNTER_INCREMENT(AsyncCompilationSCWsActive);
 			}
 		}
 	}
@@ -1110,150 +1163,189 @@ FString FShaderCompileThreadRunnable::GetWorkingDirectoryForWorker(int32 WorkerI
 
 FShaderCompileMemoryUsage FShaderCompileThreadRunnable::CalcPerWorkerMemoryUsageTotal(const TArray<FWorkerMemoryUsage>& InPerWorkerMemoryUsage) const
 {
-	FShaderCompileMemoryUsage MemoryUsage{ 0ull, 0ull };
+	FShaderCompileMemoryUsage MemoryUsage{ 0ull, 0ull, 0 };
 
 	for (const FWorkerMemoryUsage& WorkerMemUsage : InPerWorkerMemoryUsage)
 	{
 		MemoryUsage.VirtualMemory += WorkerMemUsage.VirtualMemory;
 		MemoryUsage.PhysicalMemory += WorkerMemUsage.PhysicalMemory;
 	}
+	MemoryUsage.ActiveWorkerCount = GetNumberOfAvailableWorkers();
+
+	TRACE_COUNTER_SET(AsyncCompilationSCWMemory, MemoryUsage.VirtualMemory);
+	TRACE_COUNTER_SET(AsyncCompilationSCWsActive, MemoryUsage.ActiveWorkerCount);
 
 	return MemoryUsage;
 }
 
 void FShaderCompileThreadRunnable::CheckMemoryLimitViolation()
 {
-	constexpr double kMemoryLimitPollInterval = 0.1; // Check every 0.1s if the memory limit has been exceeded
-	constexpr double kResumingWorkersPollInterval = 1.0; // Check every second since the last time workers have been suspending if we can resume some workers again
+	constexpr double kMemoryLimitPollInterval = 0.1; 		// Check every 0.1s, both for suspending and if not suspending, but there are suspended workers resuming.
 
 	const double CurrentTime = FPlatformTime::Seconds();
+	const bool bTick = CurrentTime - MemoryMonitoringState.LastTimeOfMemoryLimitPoll > kMemoryLimitPollInterval;
+	if (!bTick)
+	{
+		return;
+	}
+	FShaderCompileMemoryUsage TotalSCMemoryUsage = { 0ull, 0ull, 0 };
 
-	const bool bTickSuspendWorkers = CurrentTime - MemoryMonitoringState.LastTimeOfMemoryLimitPoll > kMemoryLimitPollInterval;
-	const bool bTickResumeWorkers = MemoryMonitoringState.bHasSuspendedWorkers &&
-		CurrentTime - MemoryMonitoringState.LastTimeOfSuspeningOrResumingWorkers > kResumingWorkersPollInterval;
-
-	FShaderCompileMemoryUsage TotalSCMemoryUsage = { 0ull, 0ull };
-	if (bTickSuspendWorkers || bTickResumeWorkers)
 	{
 		FScopeLock PerWorkerMemScopeLock(&PerWorkerMemUsageLock);
-		RetrievePerWorkerMemoryUsage(PerWorkerMemoryUsage);
-		TotalSCMemoryUsage = CalcPerWorkerMemoryUsageTotal(PerWorkerMemoryUsage);
+		CachePerWorkerMemoryUsage();
 	}
+	TotalSCMemoryUsage = CalcPerWorkerMemoryUsageTotal(PerWorkerMemoryUsage);
+	MemoryMonitoringState.LastTimeOfMemoryLimitPoll = CurrentTime;
+
+	//	Query the memory violation information once, whether just checking suspend or resume as well.
+	FShaderCompilerJobMemoryLimitInfo LimitInfo;
+	bool bHasMemoryViolation = QueryMemoryLimitViolationStatus(TotalSCMemoryUsage, LimitInfo);
 
 	// Check memory limit violations periodically
-	if (bTickSuspendWorkers)
 	{
-		MemoryMonitoringState.LastTimeOfMemoryLimitPoll = CurrentTime;
-
-		// Check if memory limit has been exceeded
-		FShaderCompilerJobMemoryLimitInfo LimitInfo;
-		const bool bHasMemoryViolation = QueryMemoryLimitViolationStatus(TotalSCMemoryUsage, LimitInfo);
-
 		FScopeLock PerWorkerMemScopeLock(&PerWorkerMemUsageLock);
 		TracePerWorkerSummary(PerWorkerMemoryUsage, bHasMemoryViolation);
-		PerWorkerMemScopeLock.Unlock();
+	}
 
-		if (bHasMemoryViolation)
+	int32 NumWorkersToSuspend = 0;
+	const int32 NumberOfAvailableWorkers = GetNumberOfAvailableWorkers();
+
+	if(bHasMemoryViolation)
+	{
+		// Try to halve the number of workers
+		NumWorkersToSuspend = NumberOfAvailableWorkers / 2;
+	}
+	else
+	{
+		//	Not violating, check for a soft limit.
+		if(GShaderCompilerWorkerSoftLimit)
 		{
-			MemoryMonitoringState.LastTimeOfSuspeningOrResumingWorkers = CurrentTime;
+			if(LimitInfo.ActiveWorkerCount)
+			{
+				const int64 MemoryUsedPerWorker = LimitInfo.MemoryUsed / LimitInfo.ActiveWorkerCount;
 
-			const int32 NumAvailableWorkers = GetNumberOfAvailableWorkers();
-			const int64 MemoryViolationInBytes = LimitInfo.MemoryUsed - LimitInfo.MemoryLimit;
+				//	If there isn't space for one more worker at the current average, suspend one worker.
+				//	Resume checks for space for at least two more workers.
+				const int64 SoftLimitSize = ((float)MemoryUsedPerWorker * GShaderCompilerWorkerSoftLimitInWorkers);
+				const int64 MemoryForAllWorkersPlusSoftLimit = LimitInfo.MemoryUsed + SoftLimitSize;
+				if(MemoryForAllWorkersPlusSoftLimit >= LimitInfo.MemoryLimit)
+				{
+					NumWorkersToSuspend = 1;
+					bHasMemoryViolation = true;
+				}
+			}
+		}
+	}
 
-			if (GShaderCompilerStats)
-				GShaderCompilerStats->RegisterMemoryViolation(MemoryViolationInBytes, NumAvailableWorkers);
-			TRACE_COUNTER_SET(AsyncCompilationSCWMemoryViolationSize, MemoryViolationInBytes);
+	if(NumWorkersToSuspend)
+	{
+		int32 NumBackloggedJobs = 0;
+		const int32 NumSuspendedWorkers = SuspendWorkersAndBacklogJobs(NumWorkersToSuspend, &NumBackloggedJobs);
+		if (NumSuspendedWorkers > 0)
+		{
+			UE_LOGFMT_NSLOC(
+				LogShaderCompilers, Display, "ShaderCompilers", "SuspendingWorkers",
+				"Shader compiler memory usage of {MemoryUsed} MiB, limit of {MemoryLimit} MiB: " \
+				"Backlogged {BackloggedJobs} compile {BackloggedJobsName} from {SuspendedWorkers} suspended {SuspendedWorkersName} ({ActiveWorkerCount}/{TotalWorkerCount} active)",
+				("MemoryUsed", static_cast<int32>(LimitInfo.MemoryUsed / 1024 / 1024)),
+				("MemoryLimit", static_cast<int32>(LimitInfo.MemoryLimit / 1024 / 1024)),
+				("BackloggedJobs", NumBackloggedJobs),
+				("BackloggedJobsName", NumBackloggedJobs == 1 ? TEXT("job") : TEXT("jobs")),
+				("SuspendedWorkers", NumSuspendedWorkers),
+				("SuspendedWorkersName", NumSuspendedWorkers == 1 ? TEXT("worker") : TEXT("workers")),
+				("ActiveWorkerCount", GetNumberOfAvailableWorkers()),
+				("TotalWorkerCount", GetNumberOfWorkers())
+			);
+			MemoryMonitoringState.bHasFailedToSuspendWorkers = false;
+			LimitInfo.ActiveWorkerCount -= NumSuspendedWorkers;
+		}
+		else if (!MemoryMonitoringState.bHasFailedToSuspendWorkers)
+		{
+			UE_LOGFMT_NSLOC(
+				LogShaderCompilers, Warning, "ShaderCompilers", "SuspendingWorkersFailed",
+				"Shader compiler memory usage {MemoryUsed} MiB, limit of {MemoryLimit} MiB, but cannot suspend any more workers",
+				("MemoryUsed", static_cast<int32>(LimitInfo.MemoryUsed / 1024 / 1024)),
+				("MemoryLimit", static_cast<int32>(LimitInfo.MemoryLimit / 1024 / 1024))
+			);
+			MemoryMonitoringState.bHasFailedToSuspendWorkers = true; // Don't show this warning again unless we were able to suspend workers again
+		}
+		MemoryMonitoringState.LastTimeOfSuspendingOrResumingWorkers = CurrentTime;
 
-			// Try to halve the number of workers
-			const int32 NumWorkersToSuspend = NumAvailableWorkers / 2;
+		//	Return having suspended, or tried to suspend some workers.
+		return;
+	}
 
-			int32 NumBackloggedJobs = 0;
-			const int32 NumSuspendedWorkers = SuspendWorkersAndBacklogJobs(NumWorkersToSuspend, &NumBackloggedJobs);
-			if (NumSuspendedWorkers > 0)
-			{
-				UE_LOGFMT_NSLOC(
-					LogShaderCompilers, Display, "ShaderCompilers", "SuspendingWorkers",
-					"Shader compiler memory usage of {MemoryUsed} MiB exceeded limit of {MemoryLimit} MiB: " \
-					"Backlogged {BackloggedJobs} compile {BackloggedJobsName} from {SuspendedWorkers} suspended {SuspendedWorkersName} ({ActiveWorkerCount}/{TotalWorkerCount} active)",
-					("MemoryUsed", static_cast<int32>(LimitInfo.MemoryUsed / 1024 / 1024)),
-					("MemoryLimit", static_cast<int32>(LimitInfo.MemoryLimit / 1024 / 1024)),
-					("BackloggedJobs", NumBackloggedJobs),
-					("BackloggedJobsName", NumBackloggedJobs == 1 ? TEXT("job") : TEXT("jobs")),
-					("SuspendedWorkers", NumSuspendedWorkers),
-					("SuspendedWorkersName", NumSuspendedWorkers == 1 ? TEXT("worker") : TEXT("workers")),
-					("ActiveWorkerCount", GetNumberOfAvailableWorkers()),
-					("TotalWorkerCount", GetNumberOfWorkers())
-				);
-				MemoryMonitoringState.bHasSuspendedWorkers = true;
-				MemoryMonitoringState.bHasFailedToSuspendWorkers = false;
-			}
-			else if (!MemoryMonitoringState.bHasFailedToSuspendWorkers)
-			{
-				UE_LOGFMT_NSLOC(
-					LogShaderCompilers, Warning, "ShaderCompilers", "SuspendingWorkersFailed",
-					"Shader compiler memory usage of {MemoryUsed} MiB exceeded limit of {MemoryLimit} MiB, but cannot suspend any more workers",
-					("MemoryUsed", static_cast<int32>(LimitInfo.MemoryUsed / 1024 / 1024)),
-					("MemoryLimit", static_cast<int32>(LimitInfo.MemoryLimit / 1024 / 1024))
-				);
-				MemoryMonitoringState.bHasFailedToSuspendWorkers = true; // Don't show this warning again unless we were able to suspend workers again
-			}
-		}
-		TRACE_COUNTER_SET(AsyncCompilationSCWMemory, LimitInfo.MemoryUsed);
+	const int32 NumberOfWorkers = GetNumberOfWorkers();
+	int32 NumAvailableWorkers = GetNumberOfAvailableWorkers();
+	if(NumAvailableWorkers == NumberOfWorkers)
+	{
+		return;
 	}
 
-	// Check if we can resume previously suspended workers periodically
-	if (bTickResumeWorkers)
+	//	There is at least one suspended worker, check if we can resume any.
+	int32 NumWorkersToResume = 0;
+	int32 NumResumedWorkers = 0;
+	if(GShaderCompilerWorkerScaleUpApproach <= static_cast<int32>(EShaderCompilerWorkerScaleUpApproach::ScaleUpWhenBelowHalfMemoryLimit))
 	{
-		MemoryMonitoringState.LastTimeOfSuspeningOrResumingWorkers = CurrentTime;
 
-		FShaderCompilerJobMemoryLimitInfo LimitInfo;
-		if (QueryMemoryStatus(TotalSCMemoryUsage, LimitInfo))
+		// If we are below half of our memory limit, resume 50% of available workers.
+		// This approach suspends workers from 100% to 50% and then resumes them back up to 75%.
+		if (LimitInfo.MemoryUsed < LimitInfo.MemoryLimit / 2)
 		{
-			// If we are below half of our memory limit, resume 50% of available workers.
-			// This approach suspends workers from 100% to 50% and then resumes them back up to 75%.
-			if (LimitInfo.MemoryUsed < LimitInfo.MemoryLimit / 2)
+			//	Number of workers to resume is half the currently inactive workers, with 1 added before the divide to always ensure at least one worker
+			//	is attempted to be woken.
+			NumWorkersToResume = ((NumberOfWorkers - NumAvailableWorkers)+1) / 2;
+			NumResumedWorkers = ResumeSuspendedWorkers(NumWorkersToResume);
+		}
+	}
+	else
+	{
+		ensure(GShaderCompilerWorkerScaleUpApproach == static_cast<int32>(EShaderCompilerWorkerScaleUpApproach::ScaleUpWhenSpaceForTwoAverageWorkers));
+
+		//	Scale back up if there looks to be enough memory to fit two average size workers based on current worker memory usage.
+		NumWorkersToResume = 0;
+		const int64 MemoryAvailable = LimitInfo.MemoryLimit;
+		if (MemoryAvailable > 0)
+		{
+			//	Set a memory per worker based on the config, adpating if there are running workers.
+			const int32 EffectiveActiveWorkerCount = LimitInfo.ActiveWorkerCount ? LimitInfo.ActiveWorkerCount : 1;
+			int64 MemoryUsedPerWorker = LimitInfo.MemoryUsed / EffectiveActiveWorkerCount;
+				
+			if(MemoryAvailable > (MemoryUsedPerWorker*(EffectiveActiveWorkerCount+2)))
 			{
-				//	Number of workers to resume is half the currently inactive workers, with 1 added before the divide to always ensure at least one worker
-				//	is attempted to be woken.
-				const int32 NumWorkersToResume = ((GetNumberOfWorkers() - GetNumberOfAvailableWorkers())+1) / 2;
-				const int32 NumResumedWorkers = ResumeSuspendedWorkers(NumWorkersToResume);
-				if (NumResumedWorkers > 0)
-				{
-					UE_LOGFMT_NSLOC(
-						LogShaderCompilers, Display, "ShaderCompilers", "ResumingWorkers",
-						"Resumed {ResumedWorkers} suspended {ResumedWorkersName} since memory usage of {MemoryUsed} MiB is below half the limit of {MemoryLimit} MiB ({ActiveWorkerCount}/{TotalWorkerCount} active)",
-						("ResumedWorkers", NumResumedWorkers),
-						("ResumedWorkersName", NumResumedWorkers == 1 ? TEXT("worker") : TEXT("workers")),
-						("MemoryUsed", static_cast<int32>(LimitInfo.MemoryUsed / 1024 / 1024)),
-						("MemoryLimit", static_cast<int32>(LimitInfo.MemoryLimit / 1024 / 1024)),
-						("ActiveWorkerCount", GetNumberOfAvailableWorkers()),
-						("TotalWorkerCount", GetNumberOfWorkers())
-					);
-				}
-				else
-				{
-					// No more workers that could be resumed
-					MemoryMonitoringState.bHasSuspendedWorkers = false;
-				}
+				NumWorkersToResume = 1;
+			}
+
+			if(NumWorkersToResume)
+			{
+				NumResumedWorkers = ResumeSuspendedWorkers(NumWorkersToResume);
 			}
 		}
 	}
-}
-
-bool FShaderCompileThreadRunnable::QueryMemoryStatus(const FShaderCompileMemoryUsage& InSCMemoryUsage, FShaderCompilerJobMemoryLimitInfo& OutInfo)
-{
-	OutInfo.MemoryLimit = static_cast<int64>(GShaderCompilerCurrentMemoryLimitInMiB) * 1024 * 1024;
-	OutInfo.MemoryUsed = InSCMemoryUsage.VirtualMemory;
 
-	return InSCMemoryUsage.VirtualMemory > 0;
+	if ((NumWorkersToResume > 0) && (NumResumedWorkers > 0))
+	{
+		MemoryMonitoringState.LastTimeOfSuspendingOrResumingWorkers = CurrentTime;
+		UE_LOGFMT_NSLOC(
+			LogShaderCompilers, Display, "ShaderCompilers", "ResumingWorkers",
+			"Resumed {ResumedWorkers} suspended {ResumedWorkersName} memory usage {MemoryUsed} MiB, limit of {MemoryLimit} MiB ({ActiveWorkerCount}/{TotalWorkerCount} active)",
+			("ResumedWorkers", NumResumedWorkers),
+			("ResumedWorkersName", NumResumedWorkers == 1 ? TEXT("worker") : TEXT("workers")),
+			("MemoryUsed", static_cast<int32>(LimitInfo.MemoryUsed / 1024 / 1024)),
+			("MemoryLimit", static_cast<int32>(LimitInfo.MemoryLimit / 1024 / 1024)),
+			("ActiveWorkerCount", GetNumberOfAvailableWorkers()),
+			("TotalWorkerCount", NumberOfWorkers)
+		);
+	}
 }
 
 bool FShaderCompileThreadRunnable::QueryMemoryLimitViolationStatus(const FShaderCompileMemoryUsage& InSCMemoryUsage, FShaderCompilerJobMemoryLimitInfo& OutInfo)
 {
-	const int64 MemoryLimitInBytes = static_cast<int64>(GShaderCompilerCurrentMemoryLimitInMiB) * 1024 * 1024;
+	const int64 MemoryLimitInBytes = static_cast<int64>(GShaderCompilerCurrentMemoryLimitInMiB) * 1024ull * 1024ull;
 
 	OutInfo.MemoryLimit = MemoryLimitInBytes;
 	OutInfo.MemoryUsed = InSCMemoryUsage.VirtualMemory;
+	OutInfo.ActiveWorkerCount = InSCMemoryUsage.ActiveWorkerCount;
 
 	return InSCMemoryUsage.VirtualMemory >= static_cast<uint64>(MemoryLimitInBytes);
 }
@@ -1278,18 +1370,26 @@ bool FShaderCompileThreadRunnable::PrintWorkerMemoryUsage(bool bAllowToWaitForLo
 	}
 }
 
-void FShaderCompileThreadRunnable::RetrievePerWorkerMemoryUsage(TArray<FWorkerMemoryUsage> &OutWorkerMemoryUsage) const
+void FShaderCompileThreadRunnable::CachePerWorkerMemoryUsage()
 {
-	TRACE_CPUPROFILER_EVENT_SCOPE(FShaderCompileThreadRunnable::RetrievePerWorkerMemoryUsage);
+	TRACE_CPUPROFILER_EVENT_SCOPE(FShaderCompileThreadRunnable::CachePerWorkerMemoryUsage);
+
+	const double CurrentTime = FPlatformTime::Seconds();
+	const bool bCheckPerworkerMemoryUsage = (CurrentTime - LastCheckForWorkersMemoryUsageTime > .1f);
+	if (!bCheckPerworkerMemoryUsage)
+	{
+		return;
+	}
+	LastCheckForWorkersMemoryUsageTime = CurrentTime;
 
 	FScopeLock WorkerScopeLock(&WorkerInfosLock);
 	if (WorkerInfos.Num() <= 0)
 	{
-		OutWorkerMemoryUsage.Reset();
+		PerWorkerMemoryUsage.Reset();
 		return;
 	}
 
-	OutWorkerMemoryUsage.Empty(WorkerInfos.Num());
+	PerWorkerMemoryUsage.Empty(WorkerInfos.Num());
 
 	FResourceProcessTreeMemory ProcessTreeMemory;
 	ProcessTreeMemory.CacheSystemProcessIds();
@@ -1306,7 +1406,7 @@ void FShaderCompileThreadRunnable::RetrievePerWorkerMemoryUsage(TArray<FWorkerMe
 			if (ProcessTreeMemory.TryGetMemoryUsage(MemoryStats))
 			{
 				// Virtual memory is committed memory on Windows.
-				OutWorkerMemoryUsage.Add(
+				PerWorkerMemoryUsage.Add(
 					{
 						.WorkerIndex = WorkerIndex,
 						.VirtualMemory = MemoryStats.UsedVirtual,
@@ -1319,7 +1419,7 @@ void FShaderCompileThreadRunnable::RetrievePerWorkerMemoryUsage(TArray<FWorkerMe
 		}
 
 		// Pack the array with 0 so we match the worker array count
-		OutWorkerMemoryUsage.Add(
+		PerWorkerMemoryUsage.Add(
 			{
 				.WorkerIndex = WorkerIndex,
 				.VirtualMemory = 0,
@@ -1329,7 +1429,7 @@ void FShaderCompileThreadRunnable::RetrievePerWorkerMemoryUsage(TArray<FWorkerMe
 		++WorkerIndex;
 	}
 
-	check(OutWorkerMemoryUsage.Num() == WorkerInfos.Num());
+	check(PerWorkerMemoryUsage.Num() == WorkerInfos.Num());
 }
 
 uint64 FShaderCompileThreadRunnable::GetPerWorkerVirtualMemoryMedian(const TArray<FWorkerMemoryUsage>& InWorkerMemoryUsage) const
@@ -1411,26 +1511,8 @@ void FShaderCompileThreadRunnable::TracePerWorkerSummary(const TArray<FWorkerMem
 
 FShaderCompileMemoryUsage FShaderCompileThreadRunnable::GetExternalWorkerMemoryUsage()
 {
-	FShaderCompileMemoryUsage MemoryUsage{};
-	FScopeLock WorkerScopeLock(&WorkerInfosLock);
-
-	FResourceProcessTreeMemory ProcessTreeMemory;
-
-	for (const TUniquePtr<FShaderCompileWorkerInfo>& WorkerInfo : WorkerInfos)
-	{
-		if (WorkerInfo->WorkerProcess.IsValid() && WorkerInfo->WorkerProcessId != 0)
-		{
-			ProcessTreeMemory.AddRootProcessId(WorkerInfo->WorkerProcessId);
-		}
-	}
-
-	FPlatformProcessMemoryStats MemoryStats;
-	if (ProcessTreeMemory.TryGetMemoryUsage(MemoryStats))
-	{
-		// Virtual memory is committed memory on Windows.
-		MemoryUsage.VirtualMemory = MemoryStats.UsedVirtual;
-		MemoryUsage.PhysicalMemory = MemoryStats.UsedPhysical;
-	}
+	TArray<FWorkerMemoryUsage> RecentWorkerMemoryUsage = GetRecentPerWorkerMemoryUsage();
+	FShaderCompileMemoryUsage MemoryUsage = CalcPerWorkerMemoryUsageTotal(RecentWorkerMemoryUsage);
 
 	return MemoryUsage;
 }
@@ -1463,6 +1545,7 @@ static void ReportShaderCompileWorkerDiagnostics(const TArray<TUniquePtr<FShader
 
 int32 FShaderCompileThreadRunnable::CompilingLoop()
 {
+
 	// Generate hash over all worker states to detect hung shader compile jobs
 	TMemoryHasher<FXxHash64Builder, FXxHash64> WorkerStateHasher;
 	bool bHasAnyJobs = false;
@@ -1564,7 +1647,79 @@ int32 FShaderCompileThreadRunnable::CompilingLoop()
 		}
 	}
 
+	if(NumActiveThreads)
+	{
+		FShaderCompileMemoryUsage TotalSCMemoryUsage = { 0ull, 0ull, 0 };
+		{
+			FScopeLock PerWorkerMemScopeLock(&PerWorkerMemUsageLock);
+			CachePerWorkerMemoryUsage();
+		}
+		TotalSCMemoryUsage = CalcPerWorkerMemoryUsageTotal(PerWorkerMemoryUsage);
+
+		if(bDynamicSCWBudgetEnabled)
+		{
+			UpdateDynamicSCWMemoryBudget(TotalSCMemoryUsage);
+		}
+	}
+
 	return NumActiveThreads;
 }
 
+void FShaderCompileThreadRunnable::UpdateDynamicSCWMemoryBudget(const FShaderCompileMemoryUsage& SCMemoryUsage)
+{
+	TRACE_CPUPROFILER_EVENT_SCOPE(UpdateDynamicSCWMemoryBudget);
+
+	const FPlatformMemoryStats MemoryStats = FPlatformMemory::GetStats();
+	const int64 SystemAvailableMemory = (int64)MemoryStats.AvailableVirtual;
 
+	//	Estimate what memory can be used currently for shader compilation.
+	//	System currently free
+	//	Shader compiler currently using
+	//	Asset compilation... what the current allocation is, not the estimate.
+	//	AllocationGranularity is used to reduce the frequency of changes. Small
+
+	int64 Reserve = static_cast<int64>(DynamicSCWBudgetReserveInMiB)*1024ll*1024ll;
+	Reserve += static_cast<int64>(GShaderCompilerAssetMemoryLoanInMiB)*1024ll*1024ll;	//	Include any additional requested. (Typically by asset compilation.)
+	const int64 AllocationGranularity = static_cast<int64>(DynamicSCWBudgetGranularityInMiB)*1024ll*1024ll;
+
+	FQueuedThreadPoolWrapper* AssetThreadPool = static_cast<FQueuedThreadPoolWrapper*>(FAssetCompilingManager::Get().GetThreadPool());
+	if(AssetThreadPool)
+	{
+		const int64 QueuedAssetMemory = AssetThreadPool->GetEstimatedMemory();
+		Reserve -= QueuedAssetMemory;
+	}
+	int64 MemoryAvailable = SystemAvailableMemory - Reserve;	//	Note, can be negative at this point.
+
+	int64 ShaderCompilationUsage = SCMemoryUsage.VirtualMemory;
+	MemoryAvailable += ShaderCompilationUsage;	//	Count the currently in use SCW memory as available for counting how much memory for the SCWs
+
+	//	Log if below 3GiB. This is based on the original budget of 3500. When things are consistently below that, shader compilation speed will be slow.
+	constexpr int64 LowSCWMemoryThreshold = (3ll * 1024ll * 1024ll * 1024ll);
+	if(MemoryAvailable < LowSCWMemoryThreshold)
+	{
+		static int64 MinMemoryAvailable = LowSCWMemoryThreshold;
+
+		if(MemoryAvailable < MinMemoryAvailable)
+		{
+			UE_LOG(LogShaderCompilers, Display, TEXT("Memory available for shader compilation is low, %dMiB. Shader compilation speed will be slow."), MemoryAvailable / (1024ull * 1024ull));
+			MinMemoryAvailable = MemoryAvailable;
+		}
+	}
+
+	static int64 CurrentSCWBudget = 0;
+	int64 DesiredSWCBudget = MemoryAvailable;
+	//	Quantize the desired amount to reduce changes to be significant.
+	DesiredSWCBudget = (DesiredSWCBudget / AllocationGranularity) * AllocationGranularity;
+	//	Limit to twice the configured max.
+	const FShaderCompilerMemoryLimits SCWLimits = GetShaderCompilerCurrentMemoryLimits();
+	const int64 MaxSCWBudget = ((int64)SCWLimits.BaseLimit) * 1024ll * 1024ll * 2ll;
+	DesiredSWCBudget = FMath::Clamp<int64>(DesiredSWCBudget, Reserve, MaxSCWBudget);
+
+	if(DesiredSWCBudget != CurrentSCWBudget)
+	{
+		CurrentSCWBudget = DesiredSWCBudget;
+		int32 CurrentSCWBudgetInMiB = FMath::Min((CurrentSCWBudget / (1024ll*1024ll)), INT32_MAX);
+		SetShaderCompilerTargetMemoryLimitInMiB(CurrentSCWBudgetInMiB);
+		TRACE_COUNTER_SET(AsyncCompilationSCWMemoryLimit, CurrentSCWBudget);
+	}
+}
diff --git a/Engine/Source/Runtime/Engine/Public/ShaderCompiler.h b/Engine/Source/Runtime/Engine/Public/ShaderCompiler.h
--- a/Engine/Source/Runtime/Engine/Public/ShaderCompiler.h
+++ b/Engine/Source/Runtime/Engine/Public/ShaderCompiler.h
@@ -141,6 +141,11 @@ struct FShaderCompileMemoryUsage
 	* The amount of physical memory used.
 	*/
 	uint64 PhysicalMemory;
+
+	/**
+	 * The number of active shader compiler workers
+	 */
+	int32 ActiveWorkerCount;
 };
 
 /** Function signature to iterate over shader compile jobs (single and pipeline jobs). */
@@ -256,18 +261,27 @@ private:
 	/** Tracks the last time that this thread checked if the workers were still active. */
 	double LastCheckForWorkersTime = 0.0;
 
+	/** Tracks the last time the per worker memory usage was checked. This is expensive to query, so if recent enough not queried again. */
+	double LastCheckForWorkersMemoryUsageTime = 0.0;
+
 	/** Whether to read/write files for SCW in parallel (can help situations when this takes too long for a number of reasons) */
 	bool bParallelizeIO = false;
 
 	/** List of jobs that have been backlogged when workers had to be closed due to reaching memory limits. These jobs will be picked up first before new jobs are pulled from the manager job queue. */
 	TArray<FShaderCommonCompileJobPtr> BackloggedJobs;
 
+	/** Is dynamic SCW budget enabled? When enabled, the SCW budget is adjusted based on current memory needs */
+	bool bDynamicSCWBudgetEnabled;
+	/** Memory to reserve in MiB when running the dynamic budget */
+	int32 DynamicSCWBudgetReserveInMiB;
+	/** Granularity of SCW budget. This is to reduce the frequency of changes, limiting it to significant changes */
+	int32 DynamicSCWBudgetGranularityInMiB;
+
 	struct FMemoryMonitoringState
 	{
 		double LastTimeOfMemoryLimitPoll = 0.0;
-		double LastTimeOfSuspeningOrResumingWorkers = 0.0;
+		double LastTimeOfSuspendingOrResumingWorkers = 0.0;
 		bool bHasFailedToSuspendWorkers = false;
-		bool bHasSuspendedWorkers = false;
 	}
 	MemoryMonitoringState;
 
@@ -338,8 +352,8 @@ private:
 	/** Returns the number of suspended workers. Only call inside the critical section WorkerInfosLock. */
 	int32 GetNumberOfSuspendedWorkersUnsafe() const;
 
-	/** Retrieves memory usage of each individual worker */
-	void RetrievePerWorkerMemoryUsage(TArray<FWorkerMemoryUsage>& OutWorkerMemoryUsage) const;
+	/** Caches memory usage of each individual worker into PerWorkerMemoryUsage if stale */
+	void CachePerWorkerMemoryUsage();
 
 	/** Calculates median virtual memory usage from worker memory list */
 	uint64 GetPerWorkerVirtualMemoryMedian(const TArray<FWorkerMemoryUsage>& InWorkerMemoryUsage) const;
@@ -371,11 +385,11 @@ private:
 	/** Checks it the memory limit for shader compile workers has been exceeded and suspend workers as needed. */
 	void CheckMemoryLimitViolation();
 
-	/** Queries the memory status of all worker processes. */
-	bool QueryMemoryStatus(const FShaderCompileMemoryUsage& InSCMemoryUsage, FShaderCompilerJobMemoryLimitInfo& OutInfo);
-
 	/** Queries the status if the worker processes have violated the memory limitation. */
 	bool QueryMemoryLimitViolationStatus(const FShaderCompileMemoryUsage& InSCMemoryUsage, FShaderCompilerJobMemoryLimitInfo& OutInfo);
+
+	/** Dynamically adapts the SCW memory budget to fit within current memory constraints. */
+	void UpdateDynamicSCWMemoryBudget(const FShaderCompileMemoryUsage& InSCMemoryUsage);
 };
 
 class FShaderCompileUtilities
@@ -926,7 +940,8 @@ struct FShaderCompilerMemoryLimits
 };
 
 extern ENGINE_API FShaderCompilerMemoryLimits GetShaderCompilerCurrentMemoryLimits();
-extern ENGINE_API void SetShaderCompilerTargetMemoryLimitInMiB(int64 NewLimit);
+extern ENGINE_API void SetShaderCompilerTargetMemoryLimitInMiB(int32 NewLimit);
+extern ENGINE_API void SetShaderCompilerAssetMemoryLoanInMiB(int32 NewLoan);
 
 /**  
  * Manager of asynchronous and parallel shader compilation.
