PR #14409: Fix description of BoneIsChildOf

#Synced-CL 50903586

diff --git a/Engine/Source/Runtime/Engine/Classes/Components/SkinnedMeshComponent.h b/Engine/Source/Runtime/Engine/Classes/Components/SkinnedMeshComponent.h
--- a/Engine/Source/Runtime/Engine/Classes/Components/SkinnedMeshComponent.h
+++ b/Engine/Source/Runtime/Engine/Classes/Components/SkinnedMeshComponent.h
@@ -1967,13 +1967,13 @@ public:
 	ENGINE_API void GetBoneNames(TArray<FName>& BoneNames);
 
 	/**
-	 * Tests if BoneName is child of (or equal to) ParentBoneName.
+	 * Tests if BoneName is child of ParentBoneName.
 	 *
 	 * @param BoneName Name of the bone
 	 * @param ParentBone Name to check
 	 *
 	 * @return true if child (strictly, not same). false otherwise
-	 * Note - will return false if ChildBoneIndex is the same as ParentBoneIndex ie. must be strictly a child.
+	 * Note - will return false if ChildBoneIndex is the same as ParentBoneIndex i.e. must be strictly a child.
 	 */
 	UFUNCTION(BlueprintCallable, Category="Components|SkinnedMesh")
 	ENGINE_API bool BoneIsChildOf(FName BoneName, FName ParentBoneName) const;
