Removed debug code that apparently has not fired after UE 5.0

#rnx

#Synced-CL 50933424

diff --git a/Engine/Source/Runtime/Engine/Private/Components/StaticMeshComponent.cpp b/Engine/Source/Runtime/Engine/Private/Components/StaticMeshComponent.cpp
--- a/Engine/Source/Runtime/Engine/Private/Components/StaticMeshComponent.cpp
+++ b/Engine/Source/Runtime/Engine/Private/Components/StaticMeshComponent.cpp
@@ -1924,22 +1924,6 @@ void UStaticMeshComponent::PostEditUndo()
 
 	// The component's light-maps are loaded from the transaction, so their resources need to be reinitialized.
 	InitResources();
-
-	// Debug check command trying to track down undo related uninitialized resource
-	if (GetStaticMesh() != NULL && GetStaticMesh()->GetRenderData() && GetStaticMesh()->GetRenderData()->LODResources.Num() > 0)
-	{
-		const int32 LODResourceIdx = GetStaticMesh()->GetRenderData()->GetCurrentFirstLODIdx(0);
-		if (LODResourceIdx != INDEX_NONE)
-		{
-			FRenderResource* Resource = &GetStaticMesh()->GetRenderData()->LODResources[LODResourceIdx].IndexBuffer;
-			ENQUEUE_RENDER_COMMAND(ResourceCheckCommand)(
-				[Resource](FRHICommandList& RHICmdList)
-				{
-					check(Resource->IsInitialized());
-				}
-				);
-		}
-	}
 	Super::PostEditUndo();
 }
 
