Fixed FPredicateOverlapHasSameActor and FPredicateOverlapHasDifferentActor needlessly causing Instanced Actors instantiation.

#rnx

#Synced-CL 50439895

diff --git a/Engine/Source/Runtime/Engine/Classes/Components/PrimitiveComponent.h b/Engine/Source/Runtime/Engine/Classes/Components/PrimitiveComponent.h
--- a/Engine/Source/Runtime/Engine/Classes/Components/PrimitiveComponent.h
+++ b/Engine/Source/Runtime/Engine/Classes/Components/PrimitiveComponent.h
@@ -191,8 +191,14 @@ struct FPredicateOverlapHasSameActor
 
 	bool operator() (const FOverlapInfo& Info)
 	{
+		// since we only care whether the overlaping actor is the same as MyOwnerPtr
+		// we only call GetCachedActor (which resolves the stored handle and returns the actor)
+		// rather than calling FindActor which could cause actor spawning if the entity hit
+		// is an instanced actor. We obviously don't need to spawn an actor to determine 
+		// that it won't be the same as MyOwnerPtr (which has already been spawned).
+		const AActor* HitActor = Info.OverlapInfo.HitObjectHandle.GetCachedActor();
 		// MyOwnerPtr is always valid, so we don't need the IsValid() checks in the WeakObjectPtr comparison operator.
-		return MyOwnerPtr.HasSameIndexAndSerialNumber(Info.OverlapInfo.HitObjectHandle.FetchActor());
+		return HitActor && MyOwnerPtr.HasSameIndexAndSerialNumber(HitActor);
 	}
 
 private:
@@ -209,8 +215,14 @@ struct FPredicateOverlapHasDifferentActor
 
 	bool operator() (const FOverlapInfo& Info)
 	{
+		// since we only care whether the overlaping actor is not the same as MyOwnerPtr
+		// we only call GetCachedActor (which resolves the stored handle and returns the actor)
+		// rather than calling FindActor which could cause actor spawning if the entity hit
+		// is an instanced actor. We obviously don't need to spawn an actor to determine 
+		// that it won't be the same as MyOwnerPtr (which has already been spawned).
+		const AActor* HitActor = Info.OverlapInfo.HitObjectHandle.GetCachedActor();
 		// MyOwnerPtr is always valid, so we don't need the IsValid() checks in the WeakObjectPtr comparison operator.
-		return !MyOwnerPtr.HasSameIndexAndSerialNumber(Info.OverlapInfo.HitObjectHandle.FetchActor());
+		return !HitActor || !MyOwnerPtr.HasSameIndexAndSerialNumber(HitActor);
 	}
 
 private:
