Fix replicated component subobjects after migration by clearing UActorComponent bIsReadyForReplication. This forces the replicated subobject list to be rebuilt when the owning actor calls ResetOwnedComponents when received in AActor::PostMigrate.

#Synced-CL 50910597

diff --git a/Engine/Source/Runtime/Engine/Private/Components/ActorComponent.cpp b/Engine/Source/Runtime/Engine/Private/Components/ActorComponent.cpp
--- a/Engine/Source/Runtime/Engine/Private/Components/ActorComponent.cpp
+++ b/Engine/Source/Runtime/Engine/Private/Components/ActorComponent.cpp
@@ -818,6 +818,12 @@ void UActorComponent::PostMigrate(const struct FUObjectMigrationContext& Migrati
 			GetWorld()->GetNetDriver()->AddObjectToNetworkIdCache(this);
 		}
 
+		// Force components to re-register their replicated subobjects when the owning actor calls ResetOwnedComponents,
+		// since the replicated component/subobject lists will be reset in that function.
+		// Even though we only replicate when ownership is assigned, ResetOwnedComponents is called on every receive,
+		// so we also have to rebuild replicated subobjects on every receive.
+		bIsReadyForReplication = false;
+
 		// We are expecting this flow from AActor, but now we must do this per-component in such a way
 		// that we are not mutating our previously deserialized data.
 		// RegisterAllComponents();
@@ -3361,9 +3367,9 @@ void UActorComponent::Serialize(FArchive& Ar)
 			ComponentFlags |= (bHasBeenInitialized & 1)			<< 1;
 			ComponentFlags |= (bTickFunctionsRegistered & 1)	<< 2;
 			ComponentFlags |= (bHasBegunPlay & 1)				<< 3;
-			ComponentFlags |= (bIsReadyForReplication & 1) 		<< 4;
-			ComponentFlags |= (bPhysicsStateCreated & 1) 		<< 5;
-			ComponentFlags |= (IsComponentTickEnabled() & 1)	<< 6;
+			ComponentFlags |= (bPhysicsStateCreated & 1) 		<< 4;
+			ComponentFlags |= (IsComponentTickEnabled() & 1)	<< 5;
+			// bIsReadyForReplication is intentionally not serialized here because it's reset in UActorComponent::PostMigrate when receiving.
 		}
 
 		Ar << ComponentFlags;
@@ -3374,9 +3380,8 @@ void UActorComponent::Serialize(FArchive& Ar)
 			bHasBeenInitialized 		= (ComponentFlags >> 1) & 1;
 			bTickFunctionsRegistered 	= (ComponentFlags >> 2) & 1;
 			bHasBegunPlay				= (ComponentFlags >> 3) & 1;
-			bIsReadyForReplication 		= (ComponentFlags >> 4) & 1;
-			bPhysicsStateCreated 		= (ComponentFlags >> 5) & 1;
-			bComponentTickEnabled 		= (ComponentFlags >> 6) & 1;
+			bPhysicsStateCreated 		= (ComponentFlags >> 4) & 1;
+			bComponentTickEnabled 		= (ComponentFlags >> 5) & 1;
 		}
 	}
 #endif // UE_WITH_REMOTE_OBJECT_HANDLE
