Physics assets allow a centre of mass offset to be specified on each body. However, the root-most body (which is not well defined/clear) was getting this offset overridden by the CoM offset specified on the component. This change makes the offsets add together.

Note that any existing physics asset with a non-zero CoM offset on the root-most body might change behaviour.

#Synced-CL 50904440

diff --git a/Engine/Source/Runtime/Engine/Private/SkeletalMeshComponentPhysics.cpp b/Engine/Source/Runtime/Engine/Private/SkeletalMeshComponentPhysics.cpp
--- a/Engine/Source/Runtime/Engine/Private/SkeletalMeshComponentPhysics.cpp
+++ b/Engine/Source/Runtime/Engine/Private/SkeletalMeshComponentPhysics.cpp
@@ -1031,7 +1031,7 @@ void USkeletalMeshComponent::InstantiatePhysicsAssetBodies_Internal(const UPhysi
 				BodyInst->bLockTranslation = BodyInstance.bLockTranslation;
 				BodyInst->bLockRotation = BodyInstance.bLockRotation;
 
-				BodyInst->COMNudge = BodyInstance.COMNudge;
+				BodyInst->COMNudge += BodyInstance.COMNudge;
 			}
 			else
 			{
