Removed spurious assert when handling duplicate grass types on old materials after having discovered it's a situation that could arise in the past when copy/pasting a grass input entry on a landscape grass output material expression

#rnx
#Synced-CL 50915523

diff --git a/Engine/Source/Runtime/Engine/Private/Materials/MaterialCachedData.cpp b/Engine/Source/Runtime/Engine/Private/Materials/MaterialCachedData.cpp
--- a/Engine/Source/Runtime/Engine/Private/Materials/MaterialCachedData.cpp
+++ b/Engine/Source/Runtime/Engine/Private/Materials/MaterialCachedData.cpp
@@ -1034,7 +1034,6 @@ PRAGMA_DISABLE_DEPRECATION_WARNINGS
 				if (GrassType != nullptr)
 				{
 					// Just use the grass type's fallback name as a fallback :
-					check(!NamedGrassTypes.Contains(GrassType->GetGrassNameFallback())); // Should be unique as the old GrassTypes were forced unique by the UX code
 					NamedGrassTypes.Add(GrassType->GetGrassNameFallback(), GrassType);
 				}
 			}
