Fix reinstancing linked anim graph crash when graph structure changes and InstanceClass does not match TargetClass anymore. This does not fixes the reinstancing issues caused by this problem, it only prevents a crash.

#Synced-CL 50680057

diff --git a/Engine/Source/Runtime/Engine/Private/Animation/AnimNode_LinkedAnimGraph.cpp b/Engine/Source/Runtime/Engine/Private/Animation/AnimNode_LinkedAnimGraph.cpp
--- a/Engine/Source/Runtime/Engine/Private/Animation/AnimNode_LinkedAnimGraph.cpp
+++ b/Engine/Source/Runtime/Engine/Private/Animation/AnimNode_LinkedAnimGraph.cpp
@@ -517,6 +517,7 @@ void FAnimNode_LinkedAnimGraph::HandleObjectsReinstanced_Impl(UObject* InSourceO
 	if(UseLegacyAnimInstanceReinstancingBehavior == nullptr || !UseLegacyAnimInstanceReinstancingBehavior->GetBool())
 	{
 		UAnimInstance* SourceAnimInstance = CastChecked<UAnimInstance>(InSourceObject);
+
 		FAnimInstanceProxy& SourceProxy = SourceAnimInstance->GetProxyOnAnyThread<FAnimInstanceProxy>();
 
 		// Call Initialize here to ensure any custom proxies are initialized (as they may have been re-created during
@@ -528,7 +529,14 @@ void FAnimNode_LinkedAnimGraph::HandleObjectsReinstanced_Impl(UObject* InSourceO
 		FAnimInstanceProxy& TargetProxy = TargetAnimInstance->GetProxyOnAnyThread<FAnimInstanceProxy>();
 		TargetProxy.Initialize(TargetAnimInstance);
 
-		InitializeProperties(SourceAnimInstance, GetTargetClass());
+		// Reinstancing copies properties using UEngine::CopyPropertiesForUnrelatedObjects which is named-based, but a graph compilation can shuffle the node names around, so at this point
+		// TargetAnimInstance can differ from InstanceClass. If we call InitializeProperties in this scenario, we read property offsets from the wrong class and this results in memory corruption,
+		// when calling PropagateInputProperties, so we skip it. This does not fix the problem but it prevents the crash.
+		UClass* TargetClass = GetTargetClass();
+		if(TargetClass && InTargetObject->IsA(TargetClass))
+		{
+			InitializeProperties(SourceAnimInstance, TargetClass);
+		}
 		DynamicUnlink(SourceAnimInstance);
 		DynamicLink(SourceAnimInstance);
 
