More fixes for GAS compatibility with the MultiServerProxy on behalf of jodon.karlik.

#Synced-CL 50910597

diff --git a/Engine/Source/Runtime/Engine/Private/Actor.cpp b/Engine/Source/Runtime/Engine/Private/Actor.cpp
--- a/Engine/Source/Runtime/Engine/Private/Actor.cpp
+++ b/Engine/Source/Runtime/Engine/Private/Actor.cpp
@@ -2068,7 +2068,8 @@ PRAGMA_ENABLE_DEPRECATION_WARNINGS
 
 		// PreReplication is only called on the server, except when we're recording a Client Replay.
 		// In that case we call PreReplication on the locally controlled Character as well.
-		if ((LocalRole == ROLE_Authority) || ((LocalRole == ROLE_AutonomousProxy) && World && World->IsRecordingClientReplay()))
+		const bool bRoleReplicates = (LocalRole == ROLE_Authority) || NetDriver->IsReplicationOfActorsWithAnyRoleEnabled();
+		if (bRoleReplicates || ((LocalRole == ROLE_AutonomousProxy) && World && World->IsRecordingClientReplay()))
 		{
 			PreReplication(DummyRepChangedPropertyTracker);
 		}
@@ -4733,7 +4734,8 @@ void AActor::DispatchBeginPlay(bool bFromLevelStreaming)
 		UE::Net::FReplicationSystemUtil::StartReplicatingActor(this);
 
 #if UE_SUPPORT_FOR_ACTOR_TICK_DISABLE
-		if (World->IsActorTickAndUserCallbacksEnabled())
+		// If we're loading from level streaming, allow a BeginPlay since it may spawn additional Actors
+		if (bFromLevelStreaming || World->IsActorTickAndUserCallbacksEnabled())
 #endif
 		{
 			BeginPlay();
