Expose ability to get Collision settings from immediate physics.

#Synced-CL 50904440

diff --git a/Engine/Source/Runtime/Engine/Private/PhysicsEngine/ImmediatePhysics/ImmediatePhysicsChaos/ImmediatePhysicsSimulation_Chaos.cpp b/Engine/Source/Runtime/Engine/Private/PhysicsEngine/ImmediatePhysics/ImmediatePhysicsChaos/ImmediatePhysicsSimulation_Chaos.cpp
--- a/Engine/Source/Runtime/Engine/Private/PhysicsEngine/ImmediatePhysics/ImmediatePhysicsChaos/ImmediatePhysicsSimulation_Chaos.cpp
+++ b/Engine/Source/Runtime/Engine/Private/PhysicsEngine/ImmediatePhysics/ImmediatePhysicsChaos/ImmediatePhysicsSimulation_Chaos.cpp
@@ -1221,6 +1221,11 @@ namespace ImmediatePhysics_Chaos
 #endif
 	}
 
+	const Chaos::FPBDCollisionSolverSettings& FSimulation::GetCollisionSolverSettings() const
+	{
+		return Implementation->Collisions.GetSolverSettings();
+	}
+
 	const Chaos::FCollisionDetectorSettings& FSimulation::GetCollisionDetectorSettings() const
 	{
 		return Implementation->Collisions.GetDetectorSettings();
diff --git a/Engine/Source/Runtime/Engine/Public/Physics/ImmediatePhysics/ImmediatePhysicsChaos/ImmediatePhysicsSimulation_Chaos.h b/Engine/Source/Runtime/Engine/Public/Physics/ImmediatePhysics/ImmediatePhysicsChaos/ImmediatePhysicsSimulation_Chaos.h
--- a/Engine/Source/Runtime/Engine/Public/Physics/ImmediatePhysics/ImmediatePhysicsChaos/ImmediatePhysicsSimulation_Chaos.h
+++ b/Engine/Source/Runtime/Engine/Public/Physics/ImmediatePhysics/ImmediatePhysicsChaos/ImmediatePhysicsSimulation_Chaos.h
@@ -130,6 +130,8 @@ namespace ImmediatePhysics_Chaos
 		ENGINE_API void SetSimulationSpaceSettings(
 			const Chaos::FSimulationSpaceSettings& SimulationSpaceSettings);
 
+		ENGINE_API const Chaos::FPBDCollisionSolverSettings& GetCollisionSolverSettings() const;
+
 		ENGINE_API const Chaos::FCollisionDetectorSettings& GetCollisionDetectorSettings() const;
 
 		ENGINE_API void SetCollisionDetectorSettings(const Chaos::FCollisionDetectorSettings& Settings);
