Add a CVar to control whether to precache PSOs from the LandscapeGrassWeightMeshProcessor.

#Synced-CL 50851879

#Synced-CL 50851879

diff --git a/Engine/Source/Runtime/Engine/Private/PSOPrecacheSettings.cpp b/Engine/Source/Runtime/Engine/Private/PSOPrecacheSettings.cpp
--- a/Engine/Source/Runtime/Engine/Private/PSOPrecacheSettings.cpp
+++ b/Engine/Source/Runtime/Engine/Private/PSOPrecacheSettings.cpp
@@ -162,6 +162,14 @@ static FAutoConsoleVariableRef CVarPSOPrecacheWaterInfoDepthPass(
 	ECVF_ReadOnly
 );
 
+static int32 GPSOPrecacheLandscapeGrassWeightPass = -1;
+static FAutoConsoleVariableRef CVarPSOPrecacheWLandscapeGrassWeightPass(
+	TEXT("r.PSOPrecache.LandscapeGrassWeightPass"),
+	GPSOPrecacheLandscapeGrassWeightPass,
+	TEXT("Precaches PSOs for LandscapeGrassWeightPass (-1 use engine default)\n"),
+	ECVF_ReadOnly
+);
+
 static int32 GPSOPrecacheSlateMaterials = -1;
 static FAutoConsoleVariableRef CVarPSOPrecacheSlateMaterials(
 	TEXT("r.PSOPrecache.SlateMaterials"), 
@@ -224,6 +232,7 @@ static void ApplyPSOPrecacheSettingsCVarOverrides(FPSOPrecacheSettings& PSOPreca
 	SET_CVAR_VALUE_BOOL(GPSOPrecacheDeferredDecalMaterials, bDeferredDecals);
 	SET_CVAR_VALUE(GPSOPrecacheParticlePrecachingTime, ParticlePrecachingTime);
 	SET_CVAR_VALUE_BOOL(GPSOPrecacheWaterInfoDepthPass, bWaterInfoDepth);
+	SET_CVAR_VALUE_BOOL(GPSOPrecacheLandscapeGrassWeightPass, bLandscapeGrassWeight);
 	SET_CVAR_VALUE_BOOL(GPSOPrecacheSlateMaterials, bSlateMaterials);
 	SET_CVAR_VALUE_BOOL(GPSOPrecacheSkipSlateMaterialDrawOnPSOTooLate, bSkipSlateMaterialDrawOnPSOTooLate);
 	SET_CVAR_VALUE_BOOL(GPSOPrecacheKeepSlatePSOsInMemoryForActiveMaterials, bKeepSlatePSOsInMemoryForActiveMaterials);
diff --git a/Engine/Source/Runtime/Engine/Public/PSOPrecacheSettings.h b/Engine/Source/Runtime/Engine/Public/PSOPrecacheSettings.h
--- a/Engine/Source/Runtime/Engine/Public/PSOPrecacheSettings.h
+++ b/Engine/Source/Runtime/Engine/Public/PSOPrecacheSettings.h
@@ -111,6 +111,10 @@ struct FPSOPrecacheSettings
 	/** Precache permutations for EMeshPass::WaterInfoTextureDepthPass (default false) **/
 	UPROPERTY()
 	bool bWaterInfoDepth = false;
+
+	/** Precache permutations for FLandscapeGrassWeightMeshProcessor (default true) **/
+	UPROPERTY()
+	bool bLandscapeGrassWeight = true;
 	
 	/** Precache permutations for EMeshPass::DitheredLODFadingOutMaskPass (default false) **/
 	UPROPERTY()
diff --git a/Engine/Source/Runtime/Landscape/Private/LandscapeGrassWeightExporter.cpp b/Engine/Source/Runtime/Landscape/Private/LandscapeGrassWeightExporter.cpp
--- a/Engine/Source/Runtime/Landscape/Private/LandscapeGrassWeightExporter.cpp
+++ b/Engine/Source/Runtime/Landscape/Private/LandscapeGrassWeightExporter.cpp
@@ -19,6 +19,7 @@
 #include "LandscapeSubsystem.h"
 #include "LandscapeAsyncTextureReadback.h"
 #include "Materials/MaterialExpressionLandscapeGrassOutput.h"
+#include "PSOPrecacheSettings.h"
 
 #if WITH_EDITOR
 #include "Misc/FileHelper.h"
@@ -403,6 +404,12 @@ bool FLandscapeGrassWeightMeshProcessor::Process(
 
 void FLandscapeGrassWeightMeshProcessor::CollectPSOInitializers(const FSceneTexturesConfig& SceneTexturesConfig, const FMaterial& Material, const FPSOPrecacheVertexFactoryData& VertexFactoryData, const FPSOPrecacheParams& PreCacheParams, TArray<FPSOPrecacheData>& PSOInitializers)
 {
+	// Early out if disabled
+	if (!UPSOPrecacheSettingsManager::GetSettings().bLandscapeGrassWeight)
+	{
+		return;
+	}
+
 	// Only support the Landscape fixed grid vertex factory type.
 	if (VertexFactoryData.VertexFactoryType != &FLandscapeFixedGridVertexFactory::StaticType)
 	{
