Fix ticker change knock on bug: We cannot guard the removal of from FOutputDeviceRedirector with TickerHandle.IsValid, because it might become invalid due to ticker shutdown.
#rnx

#Synced-CL 50939324

diff --git a/Engine/Source/Runtime/Engine/Private/UnrealEngine.cpp b/Engine/Source/Runtime/Engine/Private/UnrealEngine.cpp
--- a/Engine/Source/Runtime/Engine/Private/UnrealEngine.cpp
+++ b/Engine/Source/Runtime/Engine/Private/UnrealEngine.cpp
@@ -12445,9 +12445,11 @@ UEngine::FErrorsAndWarningsCollector::~FErrorsAndWarningsCollector()
 {
 	if (TickerHandle.IsValid())
 	{
-		FOutputDeviceRedirector::Get()->RemoveOutputDevice(this);
 		FTSTicker::RemoveTicker(TickerHandle);
 	}
+	// Note that TickerHandle might become invalid for other reasons, so it can not be used to guard the call to
+	// RemoveOutputDevice. We don't need a separate guard variable, it's legal and cheap to call even if not added.
+	FOutputDeviceRedirector::Get()->RemoveOutputDevice(this);
 }
 
 bool UEngine::FErrorsAndWarningsCollector::Tick(float Seconds)
