Fixed ShowDebug ANIMATION not properly showing marker names for blendspaces.

#Synced-CL 50920857

diff --git a/Engine/Source/Runtime/Engine/Classes/Animation/BlendSpace.h b/Engine/Source/Runtime/Engine/Classes/Animation/BlendSpace.h
--- a/Engine/Source/Runtime/Engine/Classes/Animation/BlendSpace.h
+++ b/Engine/Source/Runtime/Engine/Classes/Animation/BlendSpace.h
@@ -15,6 +15,7 @@
 #include "AnimNodeBase.h"
 #include "Containers/ArrayView.h"
 #include "Animation/BoneSocketReference.h"
+#include "Animation/AnimTypes.h"
 #include "BlendSpace.generated.h"
 
 class UCachedAnalysisProperties;
@@ -503,6 +504,8 @@ public:
 #endif
 	//~ End UAnimationAsset Interface
 	
+	ENGINE_API TConstArrayView<FAnimSyncMarker> GetAuthoredSyncMarkers() const;
+	
 	// Begin IInterpolationIndexProvider Overrides
 
 	/**
diff --git a/Engine/Source/Runtime/Engine/Private/Animation/AnimInstance.cpp b/Engine/Source/Runtime/Engine/Private/Animation/AnimInstance.cpp
--- a/Engine/Source/Runtime/Engine/Private/Animation/AnimInstance.cpp
+++ b/Engine/Source/Runtime/Engine/Private/Animation/AnimInstance.cpp
@@ -6,6 +6,8 @@
 
 #include "Animation/AnimInstance.h"
 #include "Animation/AnimSequence.h"
+#include "Animation/BlendSpace.h"
+#include "Animation/AnimTypes.h"
 #include "Animation/AnimTrace.h"
 #include "GameFramework/Pawn.h"
 #include "Animation/AnimStats.h"
@@ -1191,27 +1193,31 @@ void OutputTickRecords(const TArray<FAnimTickRecord>& Records, UCanvas* Canvas,
 			FName PrevMarkerName = NAME_None;
 			FName NextMarkerName = NAME_None;
 
-			TArray<FAnimSyncMarker> * AuthoredMarkerNames = nullptr;
+			TConstArrayView<FAnimSyncMarker> AuthoredMarkerNames;
 			
-			if (auto AnimSequence = Cast<UAnimSequence>(Player.SourceAsset))
+			if (const UAnimSequence* AnimSequence = Cast<UAnimSequence>(Player.SourceAsset))
 			{
-				AuthoredMarkerNames = &AnimSequence->AuthoredSyncMarkers;  
+				AuthoredMarkerNames = AnimSequence->AuthoredSyncMarkers;  
 			}
-			else if (auto AnimMontage = Cast<UAnimMontage>(Player.SourceAsset))
+			else if (const UAnimMontage* AnimMontage = Cast<UAnimMontage>(Player.SourceAsset))
 			{
-				AuthoredMarkerNames = &AnimMontage->MarkerData.AuthoredSyncMarkers;
+				AuthoredMarkerNames = AnimMontage->MarkerData.AuthoredSyncMarkers;
 			}
-
-			if (AuthoredMarkerNames)
+			else if (const UBlendSpace* BlendSpace = Cast<UBlendSpace>(Player.SourceAsset))
+			{
+				AuthoredMarkerNames = BlendSpace->GetAuthoredSyncMarkers();
+			}
+			
+			if (!AuthoredMarkerNames.IsEmpty())
 			{
-				if (Player.MarkerTickRecord->PreviousMarker.MarkerIndex >= 0)
+				if (AuthoredMarkerNames.IsValidIndex(Player.MarkerTickRecord->PreviousMarker.MarkerIndex))
 				{
-					PrevMarkerName = (*AuthoredMarkerNames)[Player.MarkerTickRecord->PreviousMarker.MarkerIndex].MarkerName; 
+					PrevMarkerName = AuthoredMarkerNames[Player.MarkerTickRecord->PreviousMarker.MarkerIndex].MarkerName; 
 				}
 
-				if (Player.MarkerTickRecord->NextMarker.MarkerIndex >= 0)
+				if (AuthoredMarkerNames.IsValidIndex(Player.MarkerTickRecord->NextMarker.MarkerIndex))
 				{
-					NextMarkerName = (*AuthoredMarkerNames)[Player.MarkerTickRecord->NextMarker.MarkerIndex].MarkerName; 
+					NextMarkerName = AuthoredMarkerNames[Player.MarkerTickRecord->NextMarker.MarkerIndex].MarkerName; 
 				}
 			}
 
diff --git a/Engine/Source/Runtime/Engine/Private/Animation/BlendSpace.cpp b/Engine/Source/Runtime/Engine/Private/Animation/BlendSpace.cpp
--- a/Engine/Source/Runtime/Engine/Private/Animation/BlendSpace.cpp
+++ b/Engine/Source/Runtime/Engine/Private/Animation/BlendSpace.cpp
@@ -998,6 +998,20 @@ void UBlendSpace::RuntimeValidateMarkerData()
 
 #endif // WITH_EDITOR
 
+TConstArrayView<FAnimSyncMarker> UBlendSpace::GetAuthoredSyncMarkers() const
+{
+	if (SampleIndexWithMarkers != INDEX_NONE && SampleData.Num() > SampleIndexWithMarkers)
+	{
+		const FBlendSample& BlendSample = SampleData[SampleIndexWithMarkers];
+		if (BlendSample.Animation != nullptr)
+		{
+			return MakeConstArrayView(BlendSample.Animation->AuthoredSyncMarkers);
+		}
+	}
+
+	return TConstArrayView<FAnimSyncMarker>(); 
+}
+
 /** When per-bone blend data are required to be sorted, this stores the sorted copy and the index of the original */
 struct FSortedPerBoneInterpolation
 {
