Silence spurious unfixable warning

#rnx
#Synced-CL 50925863

diff --git a/Engine/Source/Runtime/Engine/Private/UnrealEngine.cpp b/Engine/Source/Runtime/Engine/Private/UnrealEngine.cpp
--- a/Engine/Source/Runtime/Engine/Private/UnrealEngine.cpp
+++ b/Engine/Source/Runtime/Engine/Private/UnrealEngine.cpp
@@ -17959,7 +17959,7 @@ public:
 
 public:
 	FCPFUOWriter(UObject* SrcObject, UObject* DstObject, const UEngine::FCopyPropertiesForUnrelatedObjectsParams& Params, Verse::FDuplicationContext& InDuplicatedCells)
-		: FObjectCellWriter(SavedPropertyData, InDuplicatedCells)
+		: FObjectCellWriter(SavedPropertyData, InDuplicatedCells) // -V1050
 		// if the two objects don't share a common native base class, then they may have different
 		// serialization methods, which is dangerous (the data is not guaranteed to be homogeneous)
 		// in that case, we have to stick with tagged properties only
