Substrate - Toon - Experimental indirect specular reflection quantization.
Hidden behind r.Substrate.Experimental.ToonReflectionQuantizationEnabled=0/1.

#rnx
#tests PC editor dx12

#Synced-CL 50938759

diff --git a/Engine/Shaders/Private/DiffuseIndirectComposite.usf b/Engine/Shaders/Private/DiffuseIndirectComposite.usf
--- a/Engine/Shaders/Private/DiffuseIndirectComposite.usf
+++ b/Engine/Shaders/Private/DiffuseIndirectComposite.usf
@@ -609,11 +609,11 @@ void MainPS(
 				}
 
 			#if SUBSTRATE_ENABLED
+				float ToonProfileId = 0.0f;
 				// This path is not calling SubstrateEvaluateForEnvLight so we need to add that specialised code here.
 				BRANCH
 				if (GBuffer.ShadingModelID == SHADINGMODELID_SUBSTRATE_TOON)
 				{
-					float ToonProfileId = 0.0f;
 					float2 PatternUVs = 0.0;
 					UnpackToonCustomData(PackRGBA8(GBuffer.CustomData), PatternUVs, ToonProfileId);
 
@@ -643,6 +643,10 @@ void MainPS(
 				#else
 					IndirectLighting.Specular = AddContrastAndSpecularScale(SpecularIndirectLighting.xyz) * EnvBRDFValue;
 				#endif
+
+			#if SUBSTRATE_ENABLED && SUBSTRATE_EXPERIMENTAL_TOON_REFLECTION_QUANTIZATION
+				IndirectLighting.Specular = SubstrateToonQuantize(ToonProfileId, IndirectLighting.Specular * View.OneOverPreExposure) * View.PreExposure;
+			#endif
 			}
 			}
 			#endif // SUBSTRATE_GBUFFER_FORMAT==1
diff --git a/Engine/Shaders/Private/ReflectionEnvironmentPixelShader.usf b/Engine/Shaders/Private/ReflectionEnvironmentPixelShader.usf
--- a/Engine/Shaders/Private/ReflectionEnvironmentPixelShader.usf
+++ b/Engine/Shaders/Private/ReflectionEnvironmentPixelShader.usf
@@ -578,7 +578,11 @@ void ReflectionEnvironmentSkyLighting(
 				GetSSSRContribution(SubstrateEnvLight, SubstrateBSDFContext.BSDF, BSDFThroughput, SSRContribution);
 			}
 		#endif
-			
+
+		#if SUBSTRATE_EXPERIMENTAL_TOON_REFLECTION_QUANTIZATION
+			SpecularLighting = SubstrateToonQuantize(SubstrateEnvLight.ToonProfileID, SpecularLighting);
+		#endif
+
 			FLightAccumulator Out = (FLightAccumulator)0;
 			LightAccumulator_AddSplit(Out, DiffuseLighting, SpecularLighting, DiffuseLighting, 1.f /*CommonMultiplier*/, SubstrateEnvLight.bPostProcessSubsurface);
 			AccumulateSubstrateDeferredLighting(SubstrateLighting, Out, SubstrateEnvLight.bPostProcessSubsurface, BSDF_GETISTOPLAYER(SubstrateBSDFContext.BSDF));
diff --git a/Engine/Shaders/Private/Substrate/SubstrateEvaluation.ush b/Engine/Shaders/Private/Substrate/SubstrateEvaluation.ush
--- a/Engine/Shaders/Private/Substrate/SubstrateEvaluation.ush
+++ b/Engine/Shaders/Private/Substrate/SubstrateEvaluation.ush
@@ -1768,13 +1768,19 @@ struct FSubstrateEnvLightResult
 	// This can vary depending if the simple clear coat is used.
 	float3 SSRSpecularWeight;
 
+#if SUBSTRATE_EXPERIMENTAL_TOON_REFLECTION_QUANTIZATION
+	float ToonProfileID;
+#endif
+
 	bool   bPostProcessSubsurface; // True if we need to separate the subsurface light contribution for the screen space diffusion process.
 };
 
 FSubstrateEnvLightResult SubstrateEvaluateForEnvLight(FSubstrateBSDFContext BSDFContext, bool bEnableSpecular, FSubstrateIntegrationSettings Settings)
 {
 	FSubstrateEnvLightResult SubstrateEnvLightResult = (FSubstrateEnvLightResult)0;
-
+#if SUBSTRATE_EXPERIMENTAL_TOON_REFLECTION_QUANTIZATION
+	SubstrateEnvLightResult.ToonProfileID = -1.0f;
+#endif
 	const uint BSDFType = BSDF_GETTYPE(BSDFContext.BSDF);
 	switch (BSDFType)
 	{
@@ -2160,6 +2166,10 @@ FSubstrateEnvLightResult SubstrateEvaluateForEnvLight(FSubstrateBSDFContext BSDF
 			const float ToonProfileId = TOON_PROFILEID(BSDFContext.BSDF);
 			FToonProfileEvaluateSettings ToonEvalSettings = GetToonProfileEvaluateSettings(ToonProfileId);
 
+		#if SUBSTRATE_EXPERIMENTAL_TOON_REFLECTION_QUANTIZATION
+			SubstrateEnvLightResult.ToonProfileID			= TOON_PROFILEID(BSDFContext.BSDF);
+		#endif
+
 			float3 DiffEnvBRDF = 0.0f;
 			float3 SpecEnvBRDF = 0.0f;
 			float DirectionalAlbedo_SpecularTransmission = 1.0f;
diff --git a/Engine/Shaders/Private/ToonProfile.usf b/Engine/Shaders/Private/ToonProfile.usf
--- a/Engine/Shaders/Private/ToonProfile.usf
+++ b/Engine/Shaders/Private/ToonProfile.usf
@@ -24,6 +24,7 @@ float DiffuseIndirectScale;
 float SpecularIndirectScale;
 float bSpecularUsesGGX;
 float bEnergyConservation;
+float SpecularIndirectRampRepetition;
 
 RWTexture2DArray<float4>	ToonProfileTexture;
 
@@ -38,6 +39,7 @@ SamplerState				ShadowHatchingPatternSampler;
 
 Buffer<float4>				DiffuseRampBufferSRV;
 Buffer<float4>				SpecularRampBufferSRV;
+Buffer<float4>				SpecularIndirectRampBufferSRV;
 Buffer<float4>				ShadowRampBufferSRV;
 
 [numthreads(8, 8, 1)]
@@ -59,6 +61,10 @@ void ToonProfileUpdateAtlasTileCS(uint3 DispatchThreadId : SV_DispatchThreadID)
 		{
 			ToonProfileTexture[uint3(PixelPos, ToonProfileIndex)] = SpecularRampBufferSRV[PixelPos.x];
 		}
+		else if (PixelPos.y == TOON_PROFILE_SPEC_IND_RAMP_Y)
+		{
+			ToonProfileTexture[uint3(PixelPos, ToonProfileIndex)] = SpecularIndirectRampBufferSRV[PixelPos.x];
+		}
 		else if (PixelPos.y == TOON_PROFILE_SHADOW_RAMP_Y)
 		{
 			ToonProfileTexture[uint3(PixelPos, ToonProfileIndex)] = ShadowRampBufferSRV[PixelPos.x];
@@ -79,7 +85,7 @@ void ToonProfileUpdateAtlasTileCS(uint3 DispatchThreadId : SV_DispatchThreadID)
 			}
 			else if (PixelPos.x == TOON_PROFILE_GI_PARAMS_X)
 			{
-				ToonProfileTexture[uint3(PixelPos, ToonProfileIndex)] = float4(DiffuseIndirectScale, SpecularIndirectScale, 0.0f, 0.0f);
+				ToonProfileTexture[uint3(PixelPos, ToonProfileIndex)] = float4(DiffuseIndirectScale, SpecularIndirectScale, SpecularIndirectRampRepetition, 0.0f);
 			}
 			else
 			{
diff --git a/Engine/Shaders/Private/ToonProfileCommon.ush b/Engine/Shaders/Private/ToonProfileCommon.ush
--- a/Engine/Shaders/Private/ToonProfileCommon.ush
+++ b/Engine/Shaders/Private/ToonProfileCommon.ush
@@ -26,6 +26,12 @@ float3 EvaluateToonProfileSpecularRamp(in float ToonProfileId, in float Specular
 	return EvaluateToonProfileRamp(ToonProfileId, Specular, TOON_PROFILE_SPECULAR_RAMP_Y);
 }
 
+float EvaluateToonProfileSpecularIndirectRamp(in float ToonProfileId, in float Level)
+{
+	// frac for looping ramp over intensity
+	return EvaluateToonProfileRamp(ToonProfileId, Level, TOON_PROFILE_SPEC_IND_RAMP_Y);
+}
+
 float EvaluateToonProfileShadowRamp(in float ToonProfileId, in float Transmittance)
 {
 	return EvaluateToonProfileRamp(ToonProfileId, Transmittance, TOON_PROFILE_SHADOW_RAMP_Y).r;
@@ -190,6 +196,7 @@ struct FToonProfileGISettings
 {
 	float DiffuseIndirectScale;
 	float SpecularIndirectScale;
+	float SpecularIndirectRampRepetition;
 };
 
 FToonProfileGISettings GetToonProfileGISettings(in float ToonProfileId)
@@ -204,10 +211,44 @@ FToonProfileGISettings GetToonProfileGISettings(in float ToonProfileId)
 	FToonProfileGISettings Ret;
 	Ret.DiffuseIndirectScale = Data.x;
 	Ret.SpecularIndirectScale = Data.y;
+	Ret.SpecularIndirectRampRepetition = Data.z;
 	return Ret;
 }
 
+////////////////////////////////////////////////////////////////////////////////////////////////////
 
+#if SUBSTRATE_EXPERIMENTAL_TOON_REFLECTION_QUANTIZATION
+
+float3 SubstrateToonQuantize(
+	in float ToonProfileId,
+	in float3 InSpecularLighting)
+{
+	if (ToonProfileId >= 0.0f)
+	{
+		FToonProfileGISettings ToonProfileGISettings = GetToonProfileGISettings(ToonProfileId);
+		const float RampScale = ToonProfileGISettings.SpecularIndirectRampRepetition;
+		if (RampScale <= 0.0)
+		{
+			return InSpecularLighting;
+		}
 
+		float3 SpecularLighting = log(max(0.001, InSpecularLighting));
+		float3 SpecularLightingHSV = LinearRGB_2_HSV(SpecularLighting);
 
+		SpecularLightingHSV.z *= RampScale;
+
+		const float RampSpecLevel = frac(SpecularLightingHSV.z);
+		const float UntouchedSpec = SpecularLightingHSV.z - RampSpecLevel;
+
+		SpecularLightingHSV.z = UntouchedSpec + EvaluateToonProfileSpecularIndirectRamp(ToonProfileId, RampSpecLevel);
+
+		SpecularLightingHSV.z *= (1.0f / RampScale);
+
+		SpecularLighting = HSV_2_LinearRGB(SpecularLightingHSV);
+		SpecularLighting = exp(SpecularLighting);
+		return SpecularLighting;
+	}
+	return InSpecularLighting;
+}
 
+#endif
diff --git a/Engine/Shaders/Shared/ToonProfileDefinitions.h b/Engine/Shaders/Shared/ToonProfileDefinitions.h
--- a/Engine/Shaders/Shared/ToonProfileDefinitions.h
+++ b/Engine/Shaders/Shared/ToonProfileDefinitions.h
@@ -5,8 +5,9 @@
 #define TOON_PROFILE_DIFFUSE_RAMP_Y		0
 #define TOON_PROFILE_SPECULAR_RAMP_Y	1
 #define TOON_PROFILE_SHADOW_RAMP_Y		2
-#define TOON_PROFILE_PARAMETERS_Y		3
-#define TOON_PROFILE_TEXTURES_Y			4
+#define TOON_PROFILE_SPEC_IND_RAMP_Y	3
+#define TOON_PROFILE_PARAMETERS_Y		4
+#define TOON_PROFILE_TEXTURES_Y			5
 
 #define TOON_PROFILE_EVAL_PARAMS_X		0
 #define TOON_PROFILE_EVAL_PARAMS2_X		1
diff --git a/Engine/Source/Runtime/Engine/Classes/Engine/ToonProfile.h b/Engine/Source/Runtime/Engine/Classes/Engine/ToonProfile.h
--- a/Engine/Source/Runtime/Engine/Classes/Engine/ToonProfile.h
+++ b/Engine/Source/Runtime/Engine/Classes/Engine/ToonProfile.h
@@ -116,6 +116,18 @@ struct FToonProfileStruct
 	UPROPERTY(Category = "Global Illumination", EditAnywhere, meta = (DisplayName = "Specular Indirect Scale", UIMin = 0.0, UIMax = 1.0, ClampMin = 0.0))
 	float SpecularIndirectScale;
 
+	/**
+	* Defines the specular global illumination ramp that can be used to quantise the specular indirect reflections. Requires r.Substrate.Experimental.ToonReflectionQuantizationEnabled=1.
+	*/
+	UPROPERTY(Category = "Global Illumination", EditAnywhere, meta = (DisplayName = "Specular Indirect Ramp", AllowZoomOutput = "false", ShowZoomButtons = "false", ViewMinInput = "0", ViewMaxInput = "1", ViewMinOutput = "0", ViewMaxOutput = "1", TimelineLength = "1", ShowInputGridNumbers = "false", ShowOutputGridNumbers = "false"))
+	FRuntimeCurveLinearColor SpecularIndirectRamp;
+
+	/**
+	 * The scale of the the specular indirect ramp application over the range of luminance. When 0, the ramp is not applied to specular indirect reflections. Requires r.Substrate.Experimental.ToonReflectionQuantizationEnabled=1.
+	 */
+	UPROPERTY(Category = "Global Illumination", EditAnywhere, meta = (UIMin = 0.0, UIMax = 1.0, ClampMin = 0.0))
+	float SpecularIndirectRampRepetition;
+
 	/**
 	 * Experimental.
 	 */
diff --git a/Engine/Source/Runtime/Engine/Private/Rendering/ToonProfile.cpp b/Engine/Source/Runtime/Engine/Private/Rendering/ToonProfile.cpp
--- a/Engine/Source/Runtime/Engine/Private/Rendering/ToonProfile.cpp
+++ b/Engine/Source/Runtime/Engine/Private/Rendering/ToonProfile.cpp
@@ -84,6 +84,7 @@ public:
 		SHADER_PARAMETER(float, SpecularIndirectScale)
 		SHADER_PARAMETER(float, bSpecularUsesGGX)
 		SHADER_PARAMETER(float, bEnergyConservation)
+		SHADER_PARAMETER(float, SpecularIndirectRampRepetition)
 
 		SHADER_PARAMETER(float, ToonProfileShadowExtinctionCoefficientInInvCentimeter)
 
@@ -100,6 +101,7 @@ public:
 
 		SHADER_PARAMETER_RDG_BUFFER_SRV(Buffer, DiffuseRampBufferSRV)
 		SHADER_PARAMETER_RDG_BUFFER_SRV(Buffer, SpecularRampBufferSRV)
+		SHADER_PARAMETER_RDG_BUFFER_SRV(Buffer, SpecularIndirectRampBufferSRV)
 		SHADER_PARAMETER_RDG_BUFFER_SRV(Buffer, ShadowRampBufferSRV)
 	END_SHADER_PARAMETER_STRUCT()
 
@@ -417,6 +419,7 @@ IPooledRenderTarget* FToonProfileTextureManager::GetAtlasTexture(FRDGBuilder& Gr
 			PassParameters->SpecularIndirectScale = Data.SpecularIndirectScale;
 			PassParameters->bSpecularUsesGGX = Data.bSpecularUsesGGX ? 1.0f : 0.0f;
 			PassParameters->bEnergyConservation = Data.bEnergyConservation ? 1.0f : 0.0f;
+			PassParameters->SpecularIndirectRampRepetition = Data.SpecularIndirectRampRepetition;
 
 			auto CreateBufferFRomCurve = [&](FRuntimeCurveLinearColor ColorCurve)
 			{
@@ -447,6 +450,9 @@ IPooledRenderTarget* FToonProfileTextureManager::GetAtlasTexture(FRDGBuilder& Gr
 			FRDGBufferRef SpecularRampBuffer = CreateBufferFRomCurve(Data.SpecularRamp);
 			PassParameters->SpecularRampBufferSRV = GraphBuilder.CreateSRV(SpecularRampBuffer, Format);
 
+			FRDGBufferRef SpecularIndirectRampBuffer = CreateBufferFRomCurve(Data.SpecularIndirectRamp);
+			PassParameters->SpecularIndirectRampBufferSRV = GraphBuilder.CreateSRV(SpecularIndirectRampBuffer, Format);
+
 			FRDGBufferRef ShadowRampBuffer = CreateBufferFRomCurve(Data.ShadowRamp);
 			PassParameters->ShadowRampBufferSRV = GraphBuilder.CreateSRV(ShadowRampBuffer, Format);
 
@@ -529,6 +535,9 @@ FToonProfileStruct::FToonProfileStruct()
 	AddPointToCurve(SpecularRamp, 0.51, 1.0);
 	AddPointToCurve(SpecularRamp, 1.00, 1.0);
 
+	AddPointToCurve(SpecularIndirectRamp, 0.0, 0.0);
+	AddPointToCurve(SpecularIndirectRamp, 1.0, 1.0);
+
 	DiffuseDitherPattern = nullptr;
 	DiffuseDitherPatternStrength = 0.0f;	// By default, skip sampling dither pattern.
 	DiffuseDitherPatternSize = 1.0f;
@@ -545,6 +554,7 @@ FToonProfileStruct::FToonProfileStruct()
 
 	DiffuseIndirectScale = 1.0f;
 	SpecularIndirectScale = 1.0f;
+	SpecularIndirectRampRepetition = 0.0f;
 
 	bSpecularUsesGGX = 0;
 	bEnergyConservation = 0;
diff --git a/Engine/Source/Runtime/Engine/Private/ShaderCompiler/ShaderCompiler.cpp b/Engine/Source/Runtime/Engine/Private/ShaderCompiler/ShaderCompiler.cpp
--- a/Engine/Source/Runtime/Engine/Private/ShaderCompiler/ShaderCompiler.cpp
+++ b/Engine/Source/Runtime/Engine/Private/ShaderCompiler/ShaderCompiler.cpp
@@ -4113,6 +4113,13 @@ void GlobalBeginCompileShader(
 
 		const bool IsStochasticLightingEnabled = bSubstrate && Substrate::IsStochasticLightingEnabled(Target.GetPlatform());
 		SET_SHADER_DEFINE(Input.Environment, SUBSTRATE_STOCHASTIC_LIGHTING_ENABLED, IsStochasticLightingEnabled ? 1 : 0);
+
+
+		{
+			static IConsoleVariable* CVar = IConsoleManager::Get().FindConsoleVariable(TEXT("r.Substrate.Experimental.ToonReflectionQuantizationEnabled"));
+			const bool bExperimentalToonReflectionQuantizationEnabled = bSubstrate && CVar && CVar->GetInt() > 0 ? 1 : 0;
+			SET_SHADER_DEFINE(Input.Environment, SUBSTRATE_EXPERIMENTAL_TOON_REFLECTION_QUANTIZATION, bExperimentalToonReflectionQuantizationEnabled ? 1 : 0);
+		}
 	}
 
 	{
diff --git a/Engine/Source/Runtime/RenderCore/Private/RenderUtils.cpp b/Engine/Source/Runtime/RenderCore/Private/RenderUtils.cpp
--- a/Engine/Source/Runtime/RenderCore/Private/RenderUtils.cpp
+++ b/Engine/Source/Runtime/RenderCore/Private/RenderUtils.cpp
@@ -2211,6 +2211,12 @@ static TAutoConsoleVariable<int32> CVarSubstrateSimpleTiledLightingPasses(
 	TEXT("Enabling simple tiled lighting passes when greater than 0, can help performance on low end platforms not supporting Rect primitives. Only stencil marking will be tiled. And tiled passes will use larger tiles. Not compatible with Lumen nor MegaLight."),
 	ECVF_RenderThreadSafe);
 
+static TAutoConsoleVariable<int32> CVarSubstrateExperimentalToonReflectionQuantizationEnabled(
+	TEXT("r.Substrate.Experimental.ToonReflectionQuantizationEnabled"),
+	0,
+	TEXT("Enable experimental toon reflection quantization."),
+	ECVF_ReadOnly | ECVF_RenderThreadSafe);
+
 namespace Substrate
 {
 	bool IsSubstrateEnabled()
diff --git a/Engine/Source/Runtime/RenderCore/Private/Shader.cpp b/Engine/Source/Runtime/RenderCore/Private/Shader.cpp
--- a/Engine/Source/Runtime/RenderCore/Private/Shader.cpp
+++ b/Engine/Source/Runtime/RenderCore/Private/Shader.cpp
@@ -2540,6 +2540,15 @@ void ShaderMapAppendKey(EShaderPlatform Platform, FShaderKeyGenerator& KeyGen)
 			KeyGen.AppendSeparator();
 			KeyGen.Append(TEXT("STRTSP"));
 		}
+
+		{
+			static IConsoleVariable* CVar = IConsoleManager::Get().FindConsoleVariable(TEXT("r.Substrate.Experimental.ToonReflectionQuantizationEnabled"));
+			if (CVar && CVar->GetInt() > 0)
+			{
+				KeyGen.AppendSeparator();
+				KeyGen.Append(TEXT("EXPTOONRQ"));
+			}
+		}
 	}
 
 	{
