Fix for NMT to run when enabled and requested by materials. (after 50942835)

#tests PC editor dx12 with NMT enabled.
#rnx

#Synced-CL 50966633

diff --git a/Engine/Source/Runtime/Engine/Private/Materials/MaterialInterface.cpp b/Engine/Source/Runtime/Engine/Private/Materials/MaterialInterface.cpp
--- a/Engine/Source/Runtime/Engine/Private/Materials/MaterialInterface.cpp
+++ b/Engine/Source/Runtime/Engine/Private/Materials/MaterialInterface.cpp
@@ -344,7 +344,7 @@ const TArray<TObjectPtr<UObject>>& UMaterialInterface::GetDefaultTextures() cons
 bool UMaterialInterface::IsUsingNewHLSLGenerator() const
 {
 	static const auto CVar = IConsoleManager::Get().FindTConsoleVariableDataBool(TEXT("r.Material.Translator.EnableNew"));
-	if (!bEnableRestrictiveMaterialMode || !CVar->GetValueOnAnyThread())
+	if (bEnableRestrictiveMaterialMode || !CVar->GetValueOnAnyThread())
 	{
 		return false;
 	}
