Substrate - NMT - added support for eye and hair material expression nodes.

#tests PC editor with NMT enabled
#rnx

#Synced-CL 50966633

diff --git a/Engine/Source/Runtime/Engine/Private/Materials/MaterialExpressionsIR.cpp b/Engine/Source/Runtime/Engine/Private/Materials/MaterialExpressionsIR.cpp
--- a/Engine/Source/Runtime/Engine/Private/Materials/MaterialExpressionsIR.cpp
+++ b/Engine/Source/Runtime/Engine/Private/Materials/MaterialExpressionsIR.cpp
@@ -5669,6 +5669,112 @@ void UMaterialExpressionSubstrateSimpleClearCoatBSDF::Build(MIR::FEmitter& Em)
 	}
 }
 
+void UMaterialExpressionSubstrateEyeBSDF::Build(MIR::FEmitter& Em)
+{
+	FSubstrateTranslatorData* SubstrateTranslatorData = Em.GetSubstrateTranslatorData();
+
+	FGuid PathUniqueId = SubstrateTranslatorData->SubstrateTreeStackGetPathUniqueId();
+	FSubstrateOperator& SubstrateOperator = SubstrateTranslatorData->SubstrateCompilationGetOperator(PathUniqueId);
+
+	FValueRef DiffuseColorArg = CompileWithDefaultFloat3(Em, DiffuseColor, 0.18f, 0.18f, 0.18f);
+	FValueRef RoughnessArg = CompileWithDefaultFloat1(Em, Roughness, 0.5f);
+	FValueRef EmissiveColorArg = CompileWithDefaultFloat3(Em, EmissiveColor, 0.0f, 0.0f, 0.0f);
+
+	FValueRef IrisMaskArg = CompileWithDefaultFloat1(Em, IrisMask, 0.0f);
+	FValueRef IrisDistanceArg = CompileWithDefaultFloat1(Em, IrisDistance, 0.0f);
+	FValueRef IrisNormalArg = CompileWithDefaultNormalWS(Em, Material, IrisNormal);
+	FValueRef IrisPlaneNormalArg = CompileWithDefaultNormalWS(Em, Material, IrisPlaneNormal);
+
+	FValueRef SubsurfaceProfileArg = Em.ConstantFloat(0.0f);	// SUBSTRATE_TODO when profile is supported with NMT
+
+	FValueRef NormalArg = CompileWithDefaultNormalWS(Em, Material, CorneaNormal);
+	const FSubstrateRegisteredSharedLocalBasis NewRegisteredSharedLocalBasis = MIRSubstrateCompilationInfoCreateSharedLocalBasis2(SubstrateTranslatorData, &NormalArg, nullptr);
+	SubstrateOperator.BSDFRegisteredSharedLocalBasis = NewRegisteredSharedLocalBasis;
+	FValueRef SharedLocalBasisIndexMacroArg = Em.InlineHLSL(MIR::FType::MakeIntScalar(), *GetSubstrateSharedLocalBasisIndexMacroInner(NewRegisteredSharedLocalBasis, SubstrateTranslatorData->CurrentSubstrateCompilationContext));
+
+	FValueRef EyeBSDF = Em.InlineHLSL(
+		MIR::FType::MakeSubstrateData(),
+		TEXTVIEW("GetSubstrateEyeBSDF($0, $1, $2, $3, $4, $5, ExtractSubsurfaceProfileInt($6), $7, $8)"),
+		{
+			DiffuseColorArg,
+			RoughnessArg,
+			IrisMaskArg,
+			IrisDistanceArg,
+			IrisNormalArg,
+			IrisPlaneNormalArg,
+			SubsurfaceProfileArg,
+			EmissiveColorArg,
+			SharedLocalBasisIndexMacroArg
+		});
+
+	FSubstrateOperator* PromoteToOperator = !SubstrateOperator.bUseParameterBlending || (SubstrateOperator.bUseParameterBlending && SubstrateOperator.bRootOfParameterBlendingSubTree) ? &SubstrateOperator : nullptr;
+	if (PromoteToOperator)
+	{
+
+		FValueRef PromototedSlab = Em.SubstratePromoteToOperator(
+			EyeBSDF,
+			Em.ConstantInt(PromoteToOperator->Index),
+			Em.ConstantInt(PromoteToOperator->BSDFIndex),
+			Em.ConstantInt(PromoteToOperator->LayerDepth),
+			Em.ConstantInt(PromoteToOperator->bIsBottom ? 1 : 0));
+
+		Em.Output(0, PromototedSlab);
+		return;
+	}
+
+	Em.Output(0, EyeBSDF);
+}
+
+void UMaterialExpressionSubstrateHairBSDF::Build(MIR::FEmitter& Em)
+{
+	FSubstrateTranslatorData* SubstrateTranslatorData = Em.GetSubstrateTranslatorData();
+
+	FGuid PathUniqueId = SubstrateTranslatorData->SubstrateTreeStackGetPathUniqueId();
+	FSubstrateOperator& SubstrateOperator = SubstrateTranslatorData->SubstrateCompilationGetOperator(PathUniqueId);
+
+	FValueRef EmissiveColorArg = CompileWithDefaultFloat3(Em, EmissiveColor, 0.0f, 0.0f, 0.0f);
+	FValueRef BaseColorArg = CompileWithDefaultFloat3(Em, BaseColor, 0.18f, 0.18f, 0.18f);
+	FValueRef RoughnessArg = CompileWithDefaultFloat1(Em, Roughness, 0.5f);
+	FValueRef SpecularArg = CompileWithDefaultFloat1(Em, Specular, 0.5f);
+	FValueRef ScatterArg = CompileWithDefaultFloat1(Em, Scatter, 0.0f);
+	FValueRef BacklitArg = CompileWithDefaultFloat1(Em, Backlit, 0.0f);
+
+	FValueRef TangentArg = CompileWithDefaultTangentWS(Em, Material, Tangent);
+	const FSubstrateRegisteredSharedLocalBasis NewRegisteredSharedLocalBasis = MIRSubstrateCompilationInfoCreateSharedLocalBasis2(SubstrateTranslatorData, &TangentArg, nullptr);
+	SubstrateOperator.BSDFRegisteredSharedLocalBasis = NewRegisteredSharedLocalBasis;
+	FValueRef SharedLocalBasisIndexMacroArg = Em.InlineHLSL(MIR::FType::MakeIntScalar(), *GetSubstrateSharedLocalBasisIndexMacroInner(NewRegisteredSharedLocalBasis, SubstrateTranslatorData->CurrentSubstrateCompilationContext));
+
+	FValueRef HairBSDF = Em.InlineHLSL(
+		MIR::FType::MakeSubstrateData(),
+		TEXTVIEW("GetSubstrateHairBSDF($0, $1, $2, $3, $4, $5, false/*bComplexTransmittance*/, $6)"),
+		{
+			BaseColorArg,
+			ScatterArg,
+			SpecularArg,
+			RoughnessArg,
+			BacklitArg,
+			EmissiveColorArg,
+			SharedLocalBasisIndexMacroArg
+		});
+
+	FSubstrateOperator* PromoteToOperator = !SubstrateOperator.bUseParameterBlending || (SubstrateOperator.bUseParameterBlending && SubstrateOperator.bRootOfParameterBlendingSubTree) ? &SubstrateOperator : nullptr;
+	if (PromoteToOperator)
+	{
+
+		FValueRef PromototedSlab = Em.SubstratePromoteToOperator(
+			HairBSDF,
+			Em.ConstantInt(PromoteToOperator->Index),
+			Em.ConstantInt(PromoteToOperator->BSDFIndex),
+			Em.ConstantInt(PromoteToOperator->LayerDepth),
+			Em.ConstantInt(PromoteToOperator->bIsBottom ? 1 : 0));
+
+		Em.Output(0, PromototedSlab);
+		return;
+	}
+
+	Em.Output(0, HairBSDF);
+}
+
 static FValueRef SubstrateBlendNormal(MIR::FEmitter& Em, FValueRef NormalCodeChunk0, FValueRef NormalCodeChunk1, FValueRef MixCodeChunk)
 {
 	FValueRef SafeMixCodeChunk = Em.Saturate(MixCodeChunk);
diff --git a/Engine/Source/Runtime/Engine/Public/Materials/MaterialExpressionSubstrate.h b/Engine/Source/Runtime/Engine/Public/Materials/MaterialExpressionSubstrate.h
--- a/Engine/Source/Runtime/Engine/Public/Materials/MaterialExpressionSubstrate.h
+++ b/Engine/Source/Runtime/Engine/Public/Materials/MaterialExpressionSubstrate.h
@@ -654,7 +654,7 @@ class UMaterialExpressionSubstrateHairBSDF : public UMaterialExpressionSubstrate
 	//~ Begin UMaterialExpression Interface
 #if WITH_EDITOR
 	virtual int32 Compile(class FMaterialCompiler* Compiler, int32 OutputIndex) override;
-	// SUBSTRATE_TODO virtual void Build(MIR::FEmitter& Emitter) override;
+	virtual void Build(MIR::FEmitter& Emitter) override;
 	virtual void GetCaption(TArray<FString>& OutCaptions) const override;
 	virtual EMaterialValueType GetOutputValueType(int32 OutputIndex) override;
 	virtual EMaterialValueType GetInputValueType(int32 InputIndex) override;
@@ -725,7 +725,7 @@ class UMaterialExpressionSubstrateEyeBSDF : public UMaterialExpressionSubstrateB
 	//~ Begin UMaterialExpression Interface
 #if WITH_EDITOR
 	virtual int32 Compile(class FMaterialCompiler* Compiler, int32 OutputIndex) override;
-	// SUBSTRATE_TODO virtual void Build(MIR::FEmitter& Emitter) override;
+	virtual void Build(MIR::FEmitter& Emitter) override;
 	virtual void GetCaption(TArray<FString>& OutCaptions) const override;
 	virtual EMaterialValueType GetOutputValueType(int32 OutputIndex) override;
 	virtual EMaterialValueType GetInputValueType(int32 InputIndex) override;
