PR #14416: Mark GetCurveValueWithDefault() in AnimInstance as const

#Synced-CL 50967439

diff --git a/Engine/Source/Runtime/Engine/Classes/Animation/AnimInstance.h b/Engine/Source/Runtime/Engine/Classes/Animation/AnimInstance.h
--- a/Engine/Source/Runtime/Engine/Classes/Animation/AnimInstance.h
+++ b/Engine/Source/Runtime/Engine/Classes/Animation/AnimInstance.h
@@ -1257,7 +1257,7 @@ public:
 	* @param	OutValue		The curve's value.
 	*/
 	UFUNCTION(BlueprintPure, Category="Animation|Curves", meta=(BlueprintThreadSafe))
-	ENGINE_API bool GetCurveValueWithDefault(FName CurveName, float DefaultValue, float& OutValue);
+	ENGINE_API bool GetCurveValueWithDefault(FName CurveName, float DefaultValue, float& OutValue) const;
 
 	/** This returns last up-to-date list of active curve names */
 	UFUNCTION(BlueprintPure, Category = "Animation|Curves", meta=(BlueprintThreadSafe))
diff --git a/Engine/Source/Runtime/Engine/Private/Animation/AnimInstance.cpp b/Engine/Source/Runtime/Engine/Private/Animation/AnimInstance.cpp
--- a/Engine/Source/Runtime/Engine/Private/Animation/AnimInstance.cpp
+++ b/Engine/Source/Runtime/Engine/Private/Animation/AnimInstance.cpp
@@ -2124,7 +2124,7 @@ bool UAnimInstance::GetCurveValue(FName CurveName, float& OutValue) const
 	return false;
 }
 
-bool UAnimInstance::GetCurveValueWithDefault(FName CurveName, float DefaultValue, float& OutValue)
+bool UAnimInstance::GetCurveValueWithDefault(FName CurveName, float DefaultValue, float& OutValue) const
 {
 	OutValue = DefaultValue;
 	return GetCurveValue(CurveName, OutValue);
