Substrate - NMT
- Added support for VolumetricFogCloud and SingleLayerWater material expression nodes.
- Removed bad test in MaterialIRToHLSLTranslator that was preventing water code to be generated correctly because the water shading model was not accounted for and ShadingModelCount was then being 0. GetRequiresDisableForwardLocalLights should only control DISABLE_FORWARD_LOCAL_LIGHTS.

#tests PC editor with NMT enabled
#rnx

#Synced-CL 50966633

diff --git a/Engine/Source/Runtime/Engine/Private/Materials/MaterialExpressionsIR.cpp b/Engine/Source/Runtime/Engine/Private/Materials/MaterialExpressionsIR.cpp
--- a/Engine/Source/Runtime/Engine/Private/Materials/MaterialExpressionsIR.cpp
+++ b/Engine/Source/Runtime/Engine/Private/Materials/MaterialExpressionsIR.cpp
@@ -5676,16 +5676,15 @@ void UMaterialExpressionSubstrateEyeBSDF::Build(MIR::FEmitter& Em)
 	FGuid PathUniqueId = SubstrateTranslatorData->SubstrateTreeStackGetPathUniqueId();
 	FSubstrateOperator& SubstrateOperator = SubstrateTranslatorData->SubstrateCompilationGetOperator(PathUniqueId);
 
-	FValueRef DiffuseColorArg = CompileWithDefaultFloat3(Em, DiffuseColor, 0.18f, 0.18f, 0.18f);
-	FValueRef RoughnessArg = CompileWithDefaultFloat1(Em, Roughness, 0.5f);
-	FValueRef EmissiveColorArg = CompileWithDefaultFloat3(Em, EmissiveColor, 0.0f, 0.0f, 0.0f);
+	FValueRef DiffuseColorArg		= CompileWithDefaultFloat3(Em, DiffuseColor,	0.18f, 0.18f, 0.18f);
+	FValueRef RoughnessArg			= CompileWithDefaultFloat1(Em, Roughness,		0.5f);
+	FValueRef EmissiveColorArg		= CompileWithDefaultFloat3(Em, EmissiveColor,	0.0f, 0.0f, 0.0f);
+	FValueRef IrisMaskArg			= CompileWithDefaultFloat1(Em, IrisMask,		0.0f);
+	FValueRef IrisDistanceArg		= CompileWithDefaultFloat1(Em, IrisDistance,	0.0f);
+	FValueRef IrisNormalArg			= CompileWithDefaultNormalWS(Em, Material,		IrisNormal);
+	FValueRef IrisPlaneNormalArg	= CompileWithDefaultNormalWS(Em, Material,		IrisPlaneNormal);
 
-	FValueRef IrisMaskArg = CompileWithDefaultFloat1(Em, IrisMask, 0.0f);
-	FValueRef IrisDistanceArg = CompileWithDefaultFloat1(Em, IrisDistance, 0.0f);
-	FValueRef IrisNormalArg = CompileWithDefaultNormalWS(Em, Material, IrisNormal);
-	FValueRef IrisPlaneNormalArg = CompileWithDefaultNormalWS(Em, Material, IrisPlaneNormal);
-
-	FValueRef SubsurfaceProfileArg = Em.ConstantFloat(0.0f);	// SUBSTRATE_TODO when profile is supported with NMT
+	FValueRef SubsurfaceProfileArg	= Em.ConstantFloat(0.0f);	// SUBSTRATE_TODO when profile is supported with NMT
 
 	FValueRef NormalArg = CompileWithDefaultNormalWS(Em, Material, CorneaNormal);
 	const FSubstrateRegisteredSharedLocalBasis NewRegisteredSharedLocalBasis = MIRSubstrateCompilationInfoCreateSharedLocalBasis2(SubstrateTranslatorData, &NormalArg, nullptr);
@@ -5732,14 +5731,14 @@ void UMaterialExpressionSubstrateHairBSDF::Build(MIR::FEmitter& Em)
 	FGuid PathUniqueId = SubstrateTranslatorData->SubstrateTreeStackGetPathUniqueId();
 	FSubstrateOperator& SubstrateOperator = SubstrateTranslatorData->SubstrateCompilationGetOperator(PathUniqueId);
 
-	FValueRef EmissiveColorArg = CompileWithDefaultFloat3(Em, EmissiveColor, 0.0f, 0.0f, 0.0f);
-	FValueRef BaseColorArg = CompileWithDefaultFloat3(Em, BaseColor, 0.18f, 0.18f, 0.18f);
-	FValueRef RoughnessArg = CompileWithDefaultFloat1(Em, Roughness, 0.5f);
-	FValueRef SpecularArg = CompileWithDefaultFloat1(Em, Specular, 0.5f);
-	FValueRef ScatterArg = CompileWithDefaultFloat1(Em, Scatter, 0.0f);
-	FValueRef BacklitArg = CompileWithDefaultFloat1(Em, Backlit, 0.0f);
+	FValueRef EmissiveColorArg	= CompileWithDefaultFloat3(Em, EmissiveColor,	0.0f, 0.0f, 0.0f);
+	FValueRef BaseColorArg		= CompileWithDefaultFloat3(Em, BaseColor,		0.18f, 0.18f, 0.18f);
+	FValueRef RoughnessArg		= CompileWithDefaultFloat1(Em, Roughness,		0.5f);
+	FValueRef SpecularArg		= CompileWithDefaultFloat1(Em, Specular,		0.5f);
+	FValueRef ScatterArg		= CompileWithDefaultFloat1(Em, Scatter,			0.0f);
+	FValueRef BacklitArg		= CompileWithDefaultFloat1(Em, Backlit,			0.0f);
 
-	FValueRef TangentArg = CompileWithDefaultTangentWS(Em, Material, Tangent);
+	FValueRef TangentArg		= CompileWithDefaultTangentWS(Em, Material, Tangent);
 	const FSubstrateRegisteredSharedLocalBasis NewRegisteredSharedLocalBasis = MIRSubstrateCompilationInfoCreateSharedLocalBasis2(SubstrateTranslatorData, &TangentArg, nullptr);
 	SubstrateOperator.BSDFRegisteredSharedLocalBasis = NewRegisteredSharedLocalBasis;
 	FValueRef SharedLocalBasisIndexMacroArg = Em.InlineHLSL(MIR::FType::MakeIntScalar(), *GetSubstrateSharedLocalBasisIndexMacroInner(NewRegisteredSharedLocalBasis, SubstrateTranslatorData->CurrentSubstrateCompilationContext));
@@ -5775,6 +5774,139 @@ void UMaterialExpressionSubstrateHairBSDF::Build(MIR::FEmitter& Em)
 	Em.Output(0, HairBSDF);
 }
 
+void UMaterialExpressionSubstrateVolumetricFogCloudBSDF::Build(MIR::FEmitter& Em)
+{
+	FSubstrateTranslatorData* SubstrateTranslatorData = Em.GetSubstrateTranslatorData();
+
+	FGuid PathUniqueId = SubstrateTranslatorData->SubstrateTreeStackGetPathUniqueId();
+	FSubstrateOperator& SubstrateOperator = SubstrateTranslatorData->SubstrateCompilationGetOperator(PathUniqueId);
+	
+	FValueRef EmissiveColorArg		= CompileWithDefaultFloat3(Em, EmissiveColor,		0.0f, 0.0f, 0.0f);
+	FValueRef AlbedoArg				= CompileWithDefaultFloat3(Em, Albedo,				0.0f, 0.0f, 0.0f);
+	FValueRef ExtinctionArg			= CompileWithDefaultFloat3(Em, Extinction,			0.0f, 0.0f, 0.0f);
+	FValueRef AmbientOcclusionArg	= CompileWithDefaultFloat1(Em, AmbientOcclusion,	1.0f);
+
+	bool bActualEmissiveOnly = bEmissiveOnly;
+	UMaterial* BaseMaterial = Material->GetBaseMaterial();
+	// Material is probably an instanced material, so we check that this is the case before potentially overriding the ShadingModel from the material instance.
+	UMaterialInterface* BaseMaterialInterface = static_cast<UMaterialInterface*>(BaseMaterial);
+	UMaterialInterface* MaterialInterface = Material;
+	if (BaseMaterial && (BaseMaterialInterface != MaterialInterface))
+	{
+		FMaterialShadingModelField BaseMaterialShadingModels = BaseMaterial->GetShadingModels();
+		FMaterialShadingModelField MaterialShadingModels = Material->GetShadingModels();
+
+		// If the potentially instanced material does not have the same shading model as the instance material, this means it would have overridden the shading model.
+		// From the UI, only a single shading model is selectable, so we simply apply the one coming form the material instance.
+		if (MaterialShadingModels.IsValid() && MaterialShadingModels.CountShadingModels() == 1 && MaterialShadingModels != BaseMaterialShadingModels)
+		{
+			// Use the mode required by the overriden shading model.
+			bActualEmissiveOnly = MaterialShadingModels.GetFirstShadingModel() == MSM_Unlit ? true : false;
+		}
+	}
+
+	// Override to EmissiveOnly if that is the final shading model.
+	if (bActualEmissiveOnly)
+	{
+		AlbedoArg = ExtinctionArg = Em.ConstantFloat3(FVector3f::ZeroVector);
+		AmbientOcclusionArg = Em.ConstantFloat(1.0f);
+	}
+
+	FValueRef VolumetricFogCloudBSDF = Em.InlineHLSL(
+		MIR::FType::MakeSubstrateData(),
+		TEXTVIEW("GetSubstrateVolumeFogCloudBSDF($0, $1, $2, $3)"),
+		{
+			AlbedoArg,
+			ExtinctionArg,
+			EmissiveColorArg,
+			AmbientOcclusionArg
+		});
+
+	FSubstrateOperator* PromoteToOperator = !SubstrateOperator.bUseParameterBlending || (SubstrateOperator.bUseParameterBlending && SubstrateOperator.bRootOfParameterBlendingSubTree) ? &SubstrateOperator : nullptr;
+	if (PromoteToOperator)
+	{
+
+		FValueRef PromototedSlab = Em.SubstratePromoteToOperator(
+			VolumetricFogCloudBSDF,
+			Em.ConstantInt(PromoteToOperator->Index),
+			Em.ConstantInt(PromoteToOperator->BSDFIndex),
+			Em.ConstantInt(PromoteToOperator->LayerDepth),
+			Em.ConstantInt(PromoteToOperator->bIsBottom ? 1 : 0));
+
+		Em.Output(0, PromototedSlab);
+		return;
+	}
+
+	Em.Output(0, VolumetricFogCloudBSDF);
+}
+
+void UMaterialExpressionSubstrateSingleLayerWaterBSDF::Build(MIR::FEmitter& Em)
+{
+	FSubstrateTranslatorData* SubstrateTranslatorData = Em.GetSubstrateTranslatorData();
+
+	FGuid PathUniqueId = SubstrateTranslatorData->SubstrateTreeStackGetPathUniqueId();
+	FSubstrateOperator& SubstrateOperator = SubstrateTranslatorData->SubstrateCompilationGetOperator(PathUniqueId);
+
+	if (SubstrateOperator.bUseParameterBlending)
+	{
+		return Em.Errorf(TEXT("Substrate SingleLayerWater BSDF node cannot be used with parameter blending."));
+	}
+	else if (SubstrateOperator.bRootOfParameterBlendingSubTree)
+	{
+		return Em.Errorf(TEXT("Substrate SingleLayerWater BSDF node cannot be the root of a parameter blending sub tree."));
+	}
+
+	FValueRef NormalArg = CompileWithDefaultNormalWS(Em, Material, Normal);
+	const FSubstrateRegisteredSharedLocalBasis NewRegisteredSharedLocalBasis = MIRSubstrateCompilationInfoCreateSharedLocalBasis2(SubstrateTranslatorData, &NormalArg, nullptr);
+	SubstrateOperator.BSDFRegisteredSharedLocalBasis = NewRegisteredSharedLocalBasis;
+	FValueRef SharedLocalBasisIndexMacroArg = Em.InlineHLSL(MIR::FType::MakeIntScalar(), *GetSubstrateSharedLocalBasisIndexMacroInner(NewRegisteredSharedLocalBasis, SubstrateTranslatorData->CurrentSubstrateCompilationContext));
+
+	FValueRef BaseColorArg				= CompileWithDefaultFloat3(Em, BaseColor,				0.0f, 0.0f, 0.0f);
+	FValueRef MetallicArg				= CompileWithDefaultFloat1(Em, Metallic,				0.0f);
+	FValueRef SpecularArg				= CompileWithDefaultFloat1(Em, Specular,				0.5f);
+	FValueRef RoughnessArg				= CompileWithDefaultFloat1(Em, Roughness,				0.5f);
+	FValueRef EmissiveColorArg			= CompileWithDefaultFloat3(Em, EmissiveColor,			0.0f, 0.0f, 0.0f);
+	FValueRef TopMaterialOpacityArg		= CompileWithDefaultFloat1(Em, TopMaterialOpacity,		0.0f);
+	FValueRef WaterAlbedoArg			= CompileWithDefaultFloat3(Em, WaterAlbedo,				0.0f, 0.0f, 0.0f);
+	FValueRef WaterExtinctionArg		= CompileWithDefaultFloat3(Em, WaterExtinction,			0.0f, 0.0f, 0.0f);
+	FValueRef WaterPhaseGArg			= CompileWithDefaultFloat1(Em, WaterPhaseG,				0.0f);
+	FValueRef ColorScaleBehindWaterArg	= CompileWithDefaultFloat3(Em, ColorScaleBehindWater,	1.0f, 1.0f, 1.0f);
+
+	FValueRef SingleLayerWaterBSDF = Em.InlineHLSL(
+		MIR::FType::MakeSubstrateData(),
+		TEXTVIEW("GetSubstrateSingleLayerWaterBSDF($0, $1, $2, $3, $4, $5, $6, $7, $8, $9, $10)"),
+		{
+			BaseColorArg,
+			MetallicArg,
+			SpecularArg,
+			RoughnessArg,
+			EmissiveColorArg,
+			TopMaterialOpacityArg,
+			WaterAlbedoArg,
+			WaterExtinctionArg,
+			WaterPhaseGArg,
+			ColorScaleBehindWaterArg,
+			SharedLocalBasisIndexMacroArg
+		});
+
+	FSubstrateOperator* PromoteToOperator = !SubstrateOperator.bUseParameterBlending || (SubstrateOperator.bUseParameterBlending && SubstrateOperator.bRootOfParameterBlendingSubTree) ? &SubstrateOperator : nullptr;
+	if (PromoteToOperator)
+	{
+
+		FValueRef PromototedSlab = Em.SubstratePromoteToOperator(
+			SingleLayerWaterBSDF,
+			Em.ConstantInt(PromoteToOperator->Index),
+			Em.ConstantInt(PromoteToOperator->BSDFIndex),
+			Em.ConstantInt(PromoteToOperator->LayerDepth),
+			Em.ConstantInt(PromoteToOperator->bIsBottom ? 1 : 0));
+
+		Em.Output(0, PromototedSlab);
+		return;
+	}
+
+	Em.Output(0, SingleLayerWaterBSDF);
+}
+
 static FValueRef SubstrateBlendNormal(MIR::FEmitter& Em, FValueRef NormalCodeChunk0, FValueRef NormalCodeChunk1, FValueRef MixCodeChunk)
 {
 	FValueRef SafeMixCodeChunk = Em.Saturate(MixCodeChunk);
diff --git a/Engine/Source/Runtime/Engine/Private/Materials/MaterialIRToHLSLTranslator.cpp b/Engine/Source/Runtime/Engine/Private/Materials/MaterialIRToHLSLTranslator.cpp
--- a/Engine/Source/Runtime/Engine/Private/Materials/MaterialIRToHLSLTranslator.cpp
+++ b/Engine/Source/Runtime/Engine/Private/Materials/MaterialIRToHLSLTranslator.cpp
@@ -2736,11 +2736,6 @@ static void GetShaderCompilerEnvironment(const FMaterial* Material, const FMater
 				continue;
 			}
 
-			if (Model == MSM_SingleLayerWater && !FDataDrivenShaderPlatformInfo::GetRequiresDisableForwardLocalLights(ShaderPlatform))
-			{
-				continue;
-			}
-
 			if (Model == MSM_SingleLayerWater && bSingleLayerWaterUsesSimpleShading)
 			{
 				// Value must match SINGLE_LAYER_WATER_SHADING_QUALITY_MOBILE_WITH_DEPTH_TEXTURE in SingleLayerWaterCommon.ush!
diff --git a/Engine/Source/Runtime/Engine/Public/Materials/MaterialExpressionSubstrate.h b/Engine/Source/Runtime/Engine/Public/Materials/MaterialExpressionSubstrate.h
--- a/Engine/Source/Runtime/Engine/Public/Materials/MaterialExpressionSubstrate.h
+++ b/Engine/Source/Runtime/Engine/Public/Materials/MaterialExpressionSubstrate.h
@@ -540,7 +540,7 @@ class UMaterialExpressionSubstrateVolumetricFogCloudBSDF : public UMaterialExpre
 	//~ Begin UMaterialExpression Interface
 #if WITH_EDITOR
 	virtual int32 Compile(class FMaterialCompiler* Compiler, int32 OutputIndex) override;
-	// SUBSTRATE_TODO virtual void Build(MIR::FEmitter& Emitter) override;
+	virtual void Build(MIR::FEmitter& Emitter) override;
 	virtual void GetCaption(TArray<FString>& OutCaptions) const override;
 	virtual EMaterialValueType GetOutputValueType(int32 OutputIndex) override;
 	virtual EMaterialValueType GetInputValueType(int32 InputIndex) override;
@@ -810,7 +810,7 @@ class UMaterialExpressionSubstrateSingleLayerWaterBSDF : public UMaterialExpress
 	//~ Begin UMaterialExpression Interface
 #if WITH_EDITOR
 	virtual int32 Compile(class FMaterialCompiler* Compiler, int32 OutputIndex) override;
-	// SUBSTRATE_TODO virtual void Build(MIR::FEmitter& Emitter) override;
+	virtual void Build(MIR::FEmitter& Emitter) override;
 	virtual void GetCaption(TArray<FString>& OutCaptions) const override;
 	virtual EMaterialValueType GetOutputValueType(int32 OutputIndex) override;
 	virtual EMaterialValueType GetInputValueType(int32 InputIndex) override;
