Substrate - NMT - made common optional BSDF promotion to operator code in a function.
#tests PC editor with NMT enabled
#rnx

#Synced-CL 50966633

#Synced-CL 50966633

diff --git a/Engine/Source/Runtime/Engine/Private/Materials/MaterialExpressionsIR.cpp b/Engine/Source/Runtime/Engine/Private/Materials/MaterialExpressionsIR.cpp
--- a/Engine/Source/Runtime/Engine/Private/Materials/MaterialExpressionsIR.cpp
+++ b/Engine/Source/Runtime/Engine/Private/Materials/MaterialExpressionsIR.cpp
@@ -5080,6 +5080,33 @@ static FValueRef CompileWithDefaultTangentWS(MIR::FEmitter& Em, UMaterial* Mater
 	return CompileWithDefaultDirectionWS(Em, Material, TangentInput, bConvertToRequestedSpace, EDefaultDirection::EDefaultDirectionTangent);
 }
 
+static FValueRef FinalizeBSDFAndPromoteToOperatorIfNeeded(MIR::FEmitter& Em, FSubstrateOperator& SubstrateOperator, FValueRef BSDF)
+{
+	FSubstrateOperator* PromoteToOperator = !SubstrateOperator.bUseParameterBlending || (SubstrateOperator.bUseParameterBlending && SubstrateOperator.bRootOfParameterBlendingSubTree) ? &SubstrateOperator : nullptr;
+	if (PromoteToOperator)
+	{
+		if (PromoteToOperator->Index == INDEX_NONE || PromoteToOperator->BSDFIndex == INDEX_NONE)
+		{
+			Em.Errorf(TEXT("Invalid SubstrateSlabBSDF operator and BSDF indices during promotion.\r\n"));
+			FValueRef DefaultSlab = Em.SubstrateDefaultSlab();
+			Em.Output(0, DefaultSlab);
+			return DefaultSlab;
+		}
+
+		FValueRef PromototedSlab = Em.SubstratePromoteToOperator(
+			BSDF,
+			Em.ConstantInt(PromoteToOperator->Index),
+			Em.ConstantInt(PromoteToOperator->BSDFIndex),
+			Em.ConstantInt(PromoteToOperator->LayerDepth),
+			Em.ConstantInt(PromoteToOperator->bIsBottom ? 1 : 0));
+		Em.Output(0, PromototedSlab);
+		return PromototedSlab;
+	}
+
+	Em.Output(0, BSDF);
+	return BSDF;
+}
+
 void UMaterialExpressionSubstrateSlabBSDF::Build(MIR::FEmitter& Em)
 {
 	// SUBSTRATE_TODO make this similar to UMaterialExpressionSubstrateSlabBSDF::Compile
@@ -5246,52 +5273,7 @@ void UMaterialExpressionSubstrateSlabBSDF::Build(MIR::FEmitter& Em)
 
 	FValueRef SharedLocalBasisIndexMacroArg = Em.InlineHLSL(MIR::FType::MakeIntScalar(), *GetSubstrateSharedLocalBasisIndexMacroInner(NewRegisteredSharedLocalBasis, SubstrateTranslatorData->CurrentSubstrateCompilationContext));
 
-	FSubstrateOperator* PromoteToOperator = !SubstrateOperator.bUseParameterBlending || (SubstrateOperator.bUseParameterBlending && SubstrateOperator.bRootOfParameterBlendingSubTree) ? &SubstrateOperator : nullptr;
-	if(PromoteToOperator)
-	{
-		MIR::FSubstrateSlabDesc SubstrateSlabDesc = {
-				NormalArg,
-				DiffuseAlbedoArg,
-				F0Arg,
-				F90Arg,
-				RoughnessArg,
-				AnisotropyArg,
-				SubsurfaceProfileArg,
-				SSSMFPArg,
-				SSSMFPScaleArg,
-				SSSPhaseAnisotropyArg,
-				Em.ConstantFloat(float(SubSurfaceType)),
-				EmissiveColorArg,
-				SecondRoughnessArg,
-				SecondRoughnessWeightArg,
-				Em.ConstantFloat(0.0f),												//		SecondRoughnessAsSimpleClearCoatSUBSTRATE_TODO
-				Em.ConstantFloat(ClearCoatUseSecondNormal),							//		*SubstrateGetCastParameterCode(ClearCoatUseSecondNormal, MCT_Float),SUBSTRATE_TODO
-				Em.ConstantFloat3(FVector3f::Zero()),								//		*SubstrateGetCastParameterCode(ClearCoatBottomNormal, MCT_Float3),SUBSTRATE_TODO
-				FuzzAmountArg,
-				FuzzColorArg,
-				FuzzRoughnessArg,
-				GlintValueArg,
-				GlintUVArg,
-				SpecularProfileArg,
-				ThicknessArg,
-				bIsThinSurface ? Em.ConstantTrue() : Em.ConstantFalse(),
-				bIsAtTheBottomOfTopology ? Em.ConstantTrue() : Em.ConstantFalse(),
-				SharedLocalBasisIndexMacroArg };
-
-		FValueRef Slab = Em.SubstrateSlab(SubstrateSlabDesc);
-
-		FValueRef PromototedSlab = Em.SubstratePromoteToOperator(
-			Slab,
-			Em.ConstantInt(PromoteToOperator->Index),
-			Em.ConstantInt(PromoteToOperator->BSDFIndex),
-			Em.ConstantInt(PromoteToOperator->LayerDepth),
-			Em.ConstantInt(PromoteToOperator->bIsBottom ? 1 : 0));
-
-		Em.Output(0, PromototedSlab);
-	}
-	else
-	{
-		MIR::FSubstrateSlabDesc SubstrateSlabDesc = {
+	MIR::FSubstrateSlabDesc SubstrateSlabDesc = {
 			NormalArg,
 			DiffuseAlbedoArg,
 			F0Arg,
@@ -5306,9 +5288,9 @@ void UMaterialExpressionSubstrateSlabBSDF::Build(MIR::FEmitter& Em)
 			EmissiveColorArg,
 			SecondRoughnessArg,
 			SecondRoughnessWeightArg,
-			Em.ConstantFloat(0.0f),												//		SecondRoughnessAsSimpleClearCoat SUBSTRATE_TODO
-			Em.ConstantFloat(ClearCoatUseSecondNormal),							//		*SubstrateGetCastParameterCode(ClearCoatUseSecondNormal, MCT_Float), SUBSTRATE_TODO
-			Em.ConstantFloat3(FVector3f::Zero()),								//		*SubstrateGetCastParameterCode(ClearCoatBottomNormal, MCT_Float3), SUBSTRATE_TODO
+			Em.ConstantFloat(0.0f),												//		SecondRoughnessAsSimpleClearCoatSUBSTRATE_TODO
+			Em.ConstantFloat(ClearCoatUseSecondNormal),							//		*SubstrateGetCastParameterCode(ClearCoatUseSecondNormal, MCT_Float),SUBSTRATE_TODO
+			Em.ConstantFloat3(FVector3f::Zero()),								//		*SubstrateGetCastParameterCode(ClearCoatBottomNormal, MCT_Float3),SUBSTRATE_TODO
 			FuzzAmountArg,
 			FuzzColorArg,
 			FuzzRoughnessArg,
@@ -5318,11 +5300,11 @@ void UMaterialExpressionSubstrateSlabBSDF::Build(MIR::FEmitter& Em)
 			ThicknessArg,
 			bIsThinSurface ? Em.ConstantTrue() : Em.ConstantFalse(),
 			bIsAtTheBottomOfTopology ? Em.ConstantTrue() : Em.ConstantFalse(),
-			SharedLocalBasisIndexMacroArg};
-		FValueRef Slab = Em.SubstrateSlab(SubstrateSlabDesc);
+			SharedLocalBasisIndexMacroArg };
 
-		Em.Output(0, Slab);
-	}
+	FValueRef Slab = Em.SubstrateSlab(SubstrateSlabDesc);
+
+	FinalizeBSDFAndPromoteToOperatorIfNeeded(Em, SubstrateOperator, Slab);
 }
 
 MIR::FValueRef UMaterialExpressionSubstrateShadingModels::BuildCommon(MIR::FEmitter& Em, UMaterial* Material,
@@ -5446,55 +5428,6 @@ MIR::FValueRef UMaterialExpressionSubstrateShadingModels::BuildCommon(MIR::FEmit
 	FValueRef BasisIndexMacroArg = Em.InlineHLSL(MIR::FType::MakeIntScalar(), *BasisIndexMacro);
 	FValueRef ClearCoat_BasisIndexMacroArg = Em.InlineHLSL(MIR::FType::MakeIntScalar(), *ClearCoat_BasisIndexMacro);
 
-	FSubstrateOperator* PromoteToOperator = !SubstrateOperator.bUseParameterBlending || (SubstrateOperator.bUseParameterBlending && SubstrateOperator.bRootOfParameterBlendingSubTree) ? &SubstrateOperator : nullptr;
-	if (PromoteToOperator)
-	{
-		if (PromoteToOperator->Index == INDEX_NONE || PromoteToOperator->BSDFIndex == INDEX_NONE)
-		{
-			Em.Errorf(TEXT("Invalid SubstrateSlabBSDF operator and BSDF indices during promotion in Material %s (asset: %s).\r\n"), *Material->GetName(), *Material->GetPathName());
-			return Em.SubstrateDefaultSlab();
-		}
-
-		MIR::FSubstrateShadingModelsDesc SubstrateShadingModelsDesc = {
-			BaseColorArg,
-			SpecularArg,
-			MetallicArg,
-			RoughnessArg,
-			AnisotropyArg,
-			SubSurfaceColorArg,
-			SubsurfaceProfileIdArg,
-			ClearCoatArg,
-			ClearCoatRoughnessArg,
-			EmissiveColorArg,
-			OpacityArg,
-			ThinTranslucentTransmittanceColorArg,
-			ThinTranslucentSurfaceCoverageArg,
-			WaterScatteringCoefficientsArg,
-			WaterAbsorptionCoefficientsArg,
-			WaterPhaseGArg,
-			ColorScaleBehindWaterArg,
-			ShadingModelArg,
-			NormalArg,
-			TangentArg,
-			ClearCoatNormalArg,
-			CustomTangentArg,
-			BasisIndexMacroArg,
-			ClearCoat_BasisIndexMacroArg,
-			bHasDynamicShadingModels };
-
-		FValueRef Slab = Em.SubstrateShadingModels(SubstrateShadingModelsDesc);
-
-		FValueRef PromototedSlab = Em.SubstratePromoteToOperator(
-			Slab,
-			Em.ConstantInt(PromoteToOperator->Index),
-			Em.ConstantInt(PromoteToOperator->BSDFIndex),
-			Em.ConstantInt(PromoteToOperator->LayerDepth),
-			Em.ConstantInt(PromoteToOperator->bIsBottom ? 1 : 0));
-
-		Em.Output(0, PromototedSlab);
-		return PromototedSlab;
-	}
-
 	MIR::FSubstrateShadingModelsDesc SubstrateShadingModelsDesc = {
 		BaseColorArg,
 		SpecularArg,
@@ -5524,8 +5457,7 @@ MIR::FValueRef UMaterialExpressionSubstrateShadingModels::BuildCommon(MIR::FEmit
 
 	FValueRef Slab = Em.SubstrateShadingModels(SubstrateShadingModelsDesc);
 
-	Em.Output(0, Slab);
-	return Slab;
+	return FinalizeBSDFAndPromoteToOperatorIfNeeded(Em, SubstrateOperator, Slab);
 }
 
 void UMaterialExpressionSubstrateShadingModels::Build(MIR::FEmitter& Em)
@@ -5590,52 +5522,7 @@ void UMaterialExpressionSubstrateSimpleClearCoatBSDF::Build(MIR::FEmitter& Em)
 	SubstrateOperator.BSDFRegisteredSharedLocalBasis = NewRegisteredSharedLocalBasis;
 	FValueRef SharedLocalBasisIndexMacroArg = Em.InlineHLSL(MIR::FType::MakeIntScalar(), *GetSubstrateSharedLocalBasisIndexMacroInner(NewRegisteredSharedLocalBasis, SubstrateTranslatorData->CurrentSubstrateCompilationContext));
 
-	FSubstrateOperator* PromoteToOperator = !SubstrateOperator.bUseParameterBlending || (SubstrateOperator.bUseParameterBlending && SubstrateOperator.bRootOfParameterBlendingSubTree) ? &SubstrateOperator : nullptr;
-	if (PromoteToOperator)
-	{
-		MIR::FSubstrateSlabDesc SubstrateSlabDesc = {
-				NormalArg,
-				DiffuseAlbedoArg,
-				F0Arg,
-				F90Arg,
-				RoughnessArg,
-				AnisotropyArg,
-				SubsurfaceProfileArg,
-				SSSMFPArg,
-				SSSMFPScaleArg,
-				SSSPhaseAnisotropyArg,
-				SubSurfaceTypeArg,
-				EmissiveColorArg,
-				ClearCoatRoughnessArg,												//		Second roughness
-				ClearCoatArg,														//		Second roughness weight
-				Em.ConstantFloat(1.0f),												//		SecondRoughnessAsSimpleClearCoat = true
-				Em.ConstantFloat(ClearCoatUseSecondNormal),							//		*SubstrateGetCastParameterCode(ClearCoatUseSecondNormal, MCT_Float),SUBSTRATE_TODO
-				Em.ConstantFloat3(FVector3f::Zero()),								//		*SubstrateGetCastParameterCode(ClearCoatBottomNormal, MCT_Float3),SUBSTRATE_TODO
-				FuzzAmountArg,
-				FuzzColorArg,
-				FuzzRoughnessArg,
-				GlintValueArg,
-				GlintUVArg,
-				SpecularProfileArg,
-				ThicknessArg,
-				bIsThinSurface ? Em.ConstantTrue() : Em.ConstantFalse(),
-				bIsAtTheBottomOfTopology ? Em.ConstantTrue() : Em.ConstantFalse(),
-				SharedLocalBasisIndexMacroArg };
-
-		FValueRef Slab = Em.SubstrateSlab(SubstrateSlabDesc);
-
-		FValueRef PromototedSlab = Em.SubstratePromoteToOperator(
-			Slab,
-			Em.ConstantInt(PromoteToOperator->Index),
-			Em.ConstantInt(PromoteToOperator->BSDFIndex),
-			Em.ConstantInt(PromoteToOperator->LayerDepth),
-			Em.ConstantInt(PromoteToOperator->bIsBottom ? 1 : 0));
-
-		Em.Output(0, PromototedSlab);
-	}
-	else
-	{
-		MIR::FSubstrateSlabDesc SubstrateSlabDesc = {
+	MIR::FSubstrateSlabDesc SubstrateSlabDesc = {
 			NormalArg,
 			DiffuseAlbedoArg,
 			F0Arg,
@@ -5663,10 +5550,10 @@ void UMaterialExpressionSubstrateSimpleClearCoatBSDF::Build(MIR::FEmitter& Em)
 			bIsThinSurface ? Em.ConstantTrue() : Em.ConstantFalse(),
 			bIsAtTheBottomOfTopology ? Em.ConstantTrue() : Em.ConstantFalse(),
 			SharedLocalBasisIndexMacroArg };
-		FValueRef Slab = Em.SubstrateSlab(SubstrateSlabDesc);
 
-		Em.Output(0, Slab);
-	}
+	FValueRef Slab = Em.SubstrateSlab(SubstrateSlabDesc);
+
+	FinalizeBSDFAndPromoteToOperatorIfNeeded(Em, SubstrateOperator, Slab);
 }
 
 void UMaterialExpressionSubstrateEyeBSDF::Build(MIR::FEmitter& Em)
@@ -5706,22 +5593,7 @@ void UMaterialExpressionSubstrateEyeBSDF::Build(MIR::FEmitter& Em)
 			SharedLocalBasisIndexMacroArg
 		});
 
-	FSubstrateOperator* PromoteToOperator = !SubstrateOperator.bUseParameterBlending || (SubstrateOperator.bUseParameterBlending && SubstrateOperator.bRootOfParameterBlendingSubTree) ? &SubstrateOperator : nullptr;
-	if (PromoteToOperator)
-	{
-
-		FValueRef PromototedSlab = Em.SubstratePromoteToOperator(
-			EyeBSDF,
-			Em.ConstantInt(PromoteToOperator->Index),
-			Em.ConstantInt(PromoteToOperator->BSDFIndex),
-			Em.ConstantInt(PromoteToOperator->LayerDepth),
-			Em.ConstantInt(PromoteToOperator->bIsBottom ? 1 : 0));
-
-		Em.Output(0, PromototedSlab);
-		return;
-	}
-
-	Em.Output(0, EyeBSDF);
+	FinalizeBSDFAndPromoteToOperatorIfNeeded(Em, SubstrateOperator, EyeBSDF);
 }
 
 void UMaterialExpressionSubstrateHairBSDF::Build(MIR::FEmitter& Em)
@@ -5756,22 +5628,7 @@ void UMaterialExpressionSubstrateHairBSDF::Build(MIR::FEmitter& Em)
 			SharedLocalBasisIndexMacroArg
 		});
 
-	FSubstrateOperator* PromoteToOperator = !SubstrateOperator.bUseParameterBlending || (SubstrateOperator.bUseParameterBlending && SubstrateOperator.bRootOfParameterBlendingSubTree) ? &SubstrateOperator : nullptr;
-	if (PromoteToOperator)
-	{
-
-		FValueRef PromototedSlab = Em.SubstratePromoteToOperator(
-			HairBSDF,
-			Em.ConstantInt(PromoteToOperator->Index),
-			Em.ConstantInt(PromoteToOperator->BSDFIndex),
-			Em.ConstantInt(PromoteToOperator->LayerDepth),
-			Em.ConstantInt(PromoteToOperator->bIsBottom ? 1 : 0));
-
-		Em.Output(0, PromototedSlab);
-		return;
-	}
-
-	Em.Output(0, HairBSDF);
+	FinalizeBSDFAndPromoteToOperatorIfNeeded(Em, SubstrateOperator, HairBSDF);
 }
 
 void UMaterialExpressionSubstrateVolumetricFogCloudBSDF::Build(MIR::FEmitter& Em)
@@ -5822,22 +5679,7 @@ void UMaterialExpressionSubstrateVolumetricFogCloudBSDF::Build(MIR::FEmitter& Em
 			AmbientOcclusionArg
 		});
 
-	FSubstrateOperator* PromoteToOperator = !SubstrateOperator.bUseParameterBlending || (SubstrateOperator.bUseParameterBlending && SubstrateOperator.bRootOfParameterBlendingSubTree) ? &SubstrateOperator : nullptr;
-	if (PromoteToOperator)
-	{
-
-		FValueRef PromototedSlab = Em.SubstratePromoteToOperator(
-			VolumetricFogCloudBSDF,
-			Em.ConstantInt(PromoteToOperator->Index),
-			Em.ConstantInt(PromoteToOperator->BSDFIndex),
-			Em.ConstantInt(PromoteToOperator->LayerDepth),
-			Em.ConstantInt(PromoteToOperator->bIsBottom ? 1 : 0));
-
-		Em.Output(0, PromototedSlab);
-		return;
-	}
-
-	Em.Output(0, VolumetricFogCloudBSDF);
+	FinalizeBSDFAndPromoteToOperatorIfNeeded(Em, SubstrateOperator, VolumetricFogCloudBSDF);
 }
 
 void UMaterialExpressionSubstrateSingleLayerWaterBSDF::Build(MIR::FEmitter& Em)
@@ -5889,22 +5731,7 @@ void UMaterialExpressionSubstrateSingleLayerWaterBSDF::Build(MIR::FEmitter& Em)
 			SharedLocalBasisIndexMacroArg
 		});
 
-	FSubstrateOperator* PromoteToOperator = !SubstrateOperator.bUseParameterBlending || (SubstrateOperator.bUseParameterBlending && SubstrateOperator.bRootOfParameterBlendingSubTree) ? &SubstrateOperator : nullptr;
-	if (PromoteToOperator)
-	{
-
-		FValueRef PromototedSlab = Em.SubstratePromoteToOperator(
-			SingleLayerWaterBSDF,
-			Em.ConstantInt(PromoteToOperator->Index),
-			Em.ConstantInt(PromoteToOperator->BSDFIndex),
-			Em.ConstantInt(PromoteToOperator->LayerDepth),
-			Em.ConstantInt(PromoteToOperator->bIsBottom ? 1 : 0));
-
-		Em.Output(0, PromototedSlab);
-		return;
-	}
-
-	Em.Output(0, SingleLayerWaterBSDF);
+	FinalizeBSDFAndPromoteToOperatorIfNeeded(Em, SubstrateOperator, SingleLayerWaterBSDF);
 }
 
 static FValueRef SubstrateBlendNormal(MIR::FEmitter& Em, FValueRef NormalCodeChunk0, FValueRef NormalCodeChunk1, FValueRef MixCodeChunk)
