[StreamableManager]
- Remove forced synchronous loads that are not required anymore with zenloader

#Synced-CL 0

diff --git a/Engine/Source/Runtime/Engine/Private/StreamableManager.cpp b/Engine/Source/Runtime/Engine/Private/StreamableManager.cpp
--- a/Engine/Source/Runtime/Engine/Private/StreamableManager.cpp
+++ b/Engine/Source/Runtime/Engine/Private/StreamableManager.cpp
@@ -2121,10 +2121,6 @@ TSharedPtr<FStreamableHandle> FStreamableManager::RequestAsyncLoad(FStreamableAs
 
 TSharedPtr<FStreamableHandle> FStreamableManager::RequestSyncLoadInternal(TArray<FSoftObjectPath>&& TargetsToStream, bool bManageActiveHandle, FString&& DebugName, UE::FSourceLocation&& Location)
 {
-	// If in async loading thread or from callback always do sync as recursive tick is unsafe
-	// Otherwise, only do a sync load if there are no background sync loads, this is faster but will cause a sync flush
-	bForceSynchronousLoads = IsInAsyncLoadingThread() || FPlatformProperties::RequiresCookedData() || !IsAsyncLoading();
-
 	FStreamableAsyncLoadParams AsyncParams;
 	AsyncParams.TargetsToStream = MoveTemp(TargetsToStream);
 	AsyncParams.Priority = AsyncLoadHighPriority;
@@ -2133,8 +2129,6 @@ TSharedPtr<FStreamableHandle> FStreamableManager::RequestSyncLoadInternal(TArray
 	// Do an async load and wait to complete. In some cases this will do a sync load due to safety issues
 	TSharedPtr<FStreamableHandle> Request = RequestAsyncLoad(MoveTemp(AsyncParams), MoveTemp(DebugName), MoveTemp(Location));
 
-	bForceSynchronousLoads = false;
-
 	if (Request.IsValid())
 	{
 		EAsyncPackageState::Type Result = Request->WaitUntilComplete();
