- Added pre-caching on the game thread for asynchronous physics body instance creation to ensure there are no data races (when computing the physical materials used by body instances).
- Cache also optimizes body instance creation for ISM components as each instance now simply shares the FBodyInstance cache instead of recomputing physical materials per instance.
- FPhysScene_AsyncPhysicsStateJobQueue is now an FGCObject and collects objects returned by IAsyncPhysicsStateProcessor to keep them rooted during asynchronous processing.
- Updated both FastGeo and ActorComponents to properly support async physics creation pre-caching.
- In UPhysicalMaterial::PostLoad, the material handle is now primed to allow thread-safe access in the asynchronous physics path.

#Synced-CL 50924777

diff --git a/Engine/Plugins/Experimental/FastGeoStreaming/Source/FastGeoStreaming/Internal/FastGeoComponent.h b/Engine/Plugins/Experimental/FastGeoStreaming/Source/FastGeoStreaming/Internal/FastGeoComponent.h
--- a/Engine/Plugins/Experimental/FastGeoStreaming/Source/FastGeoStreaming/Internal/FastGeoComponent.h
+++ b/Engine/Plugins/Experimental/FastGeoStreaming/Source/FastGeoStreaming/Internal/FastGeoComponent.h
@@ -63,18 +63,12 @@ public:
 	virtual void CreateRenderState(FRegisterComponentContext* Context) {}
 	virtual void DestroyRenderState(FFastGeoDestroyRenderStateContext* Context) {}
 	virtual void UpdateVisibility() {}
-#ifdef FASTGEO_DEBUG_PHYSICS
+	virtual void OnAsyncCreatePhysicsStateBegin_GameThread() {}
 	virtual void OnAsyncCreatePhysicsState();
+	virtual void OnAsyncCreatePhysicsStateEnd_GameThread();
 	virtual void OnAsyncDestroyPhysicsStateBegin_GameThread();
 	virtual void OnAsyncDestroyPhysicsState();
 	virtual void OnAsyncDestroyPhysicsStateEnd_GameThread();
-	void OnAsyncCreatePhysicsStateEnd_GameThread();
-#else
-	virtual void OnAsyncCreatePhysicsState() {}
-	virtual void OnAsyncDestroyPhysicsStateBegin_GameThread() {}
-	virtual void OnAsyncDestroyPhysicsState() {}
-	virtual void OnAsyncDestroyPhysicsStateEnd_GameThread() {}
-#endif
 	virtual void GetBodyInstances(TArray<FBodyInstance*>& OutBodyInstances) {}
 
 #if !WITH_EDITOR && UE_WITH_PSO_PRECACHING
diff --git a/Engine/Plugins/Experimental/FastGeoStreaming/Source/FastGeoStreaming/Internal/FastGeoContainer.h b/Engine/Plugins/Experimental/FastGeoStreaming/Source/FastGeoStreaming/Internal/FastGeoContainer.h
--- a/Engine/Plugins/Experimental/FastGeoStreaming/Source/FastGeoStreaming/Internal/FastGeoContainer.h
+++ b/Engine/Plugins/Experimental/FastGeoStreaming/Source/FastGeoStreaming/Internal/FastGeoContainer.h
@@ -170,10 +170,10 @@ public:
 	virtual bool AllowsAsyncPhysicsStateDestruction() const override final;
 	virtual bool IsAsyncPhysicsStateCreated() const override final;
 	virtual UObject* GetAsyncPhysicsStateObject() const override final;
-	virtual void OnAsyncCreatePhysicsStateBegin_GameThread() override;
+	virtual void OnAsyncCreatePhysicsStateBegin_GameThread(TArray<UObject*>& OutRootedObjects) override;
 	virtual bool OnAsyncCreatePhysicsState(const UE::FTimeout& Timeout) override;
 	virtual void OnAsyncCreatePhysicsStateEnd_GameThread() override;
-	virtual void OnAsyncDestroyPhysicsStateBegin_GameThread() override;
+	virtual void OnAsyncDestroyPhysicsStateBegin_GameThread(TArray<UObject*>& OutRootedObjects) override;
 	virtual bool OnAsyncDestroyPhysicsState(const UE::FTimeout& Timeout) override;
 	virtual void OnAsyncDestroyPhysicsStateEnd_GameThread() override;
 	virtual void CollectBodySetupsWithPhysicsMeshesToCreate(TSet<UBodySetup*>& OutBodySetups) const override;
diff --git a/Engine/Plugins/Experimental/FastGeoStreaming/Source/FastGeoStreaming/Internal/FastGeoInstancedStaticMeshComponent.h b/Engine/Plugins/Experimental/FastGeoStreaming/Source/FastGeoStreaming/Internal/FastGeoInstancedStaticMeshComponent.h
--- a/Engine/Plugins/Experimental/FastGeoStreaming/Source/FastGeoStreaming/Internal/FastGeoInstancedStaticMeshComponent.h
+++ b/Engine/Plugins/Experimental/FastGeoStreaming/Source/FastGeoStreaming/Internal/FastGeoInstancedStaticMeshComponent.h
@@ -19,6 +19,7 @@ public:
 	//~ Begin FFastGeoComponent Interface
 	virtual void Serialize(FArchive& Ar) override;
 	virtual void OnAsyncCreatePhysicsState() override;
+	virtual void OnAsyncCreatePhysicsStateEnd_GameThread() override;
 	virtual void OnAsyncDestroyPhysicsState() override;
 	virtual void OnAsyncDestroyPhysicsStateBegin_GameThread() override;
 	virtual void OnAsyncDestroyPhysicsStateEnd_GameThread() override;
diff --git a/Engine/Plugins/Experimental/FastGeoStreaming/Source/FastGeoStreaming/Internal/FastGeoPhysicsBodyInstanceOwner.h b/Engine/Plugins/Experimental/FastGeoStreaming/Source/FastGeoStreaming/Internal/FastGeoPhysicsBodyInstanceOwner.h
--- a/Engine/Plugins/Experimental/FastGeoStreaming/Source/FastGeoStreaming/Internal/FastGeoPhysicsBodyInstanceOwner.h
+++ b/Engine/Plugins/Experimental/FastGeoStreaming/Source/FastGeoStreaming/Internal/FastGeoPhysicsBodyInstanceOwner.h
@@ -37,6 +37,7 @@ public:
 	virtual UMaterialInterface* GetPhysicsMaterialBase() const override;
 	virtual int32 GetNumMaterials() const override;
 	virtual UMaterialInterface* GetMaterial(int32 Index) const override;
+	virtual UBodySetup* GetPhysicsBodySetup() const override;
 	//~ End IPhysicsBodyInstanceOwner interface
 
 private:
diff --git a/Engine/Plugins/Experimental/FastGeoStreaming/Source/FastGeoStreaming/Internal/FastGeoPrimitiveComponent.h b/Engine/Plugins/Experimental/FastGeoStreaming/Source/FastGeoStreaming/Internal/FastGeoPrimitiveComponent.h
--- a/Engine/Plugins/Experimental/FastGeoStreaming/Source/FastGeoStreaming/Internal/FastGeoPrimitiveComponent.h
+++ b/Engine/Plugins/Experimental/FastGeoStreaming/Source/FastGeoStreaming/Internal/FastGeoPrimitiveComponent.h
@@ -102,7 +102,9 @@ public:
 	//~ Begin FFastGeoComponent interface
 	virtual void Serialize(FArchive& Ar) override;
 	virtual void InitializeDynamicProperties() override;
+	virtual void OnAsyncCreatePhysicsStateBegin_GameThread() override;
 	virtual void OnAsyncCreatePhysicsState() override;
+	virtual void OnAsyncCreatePhysicsStateEnd_GameThread() override;
 	virtual void OnAsyncDestroyPhysicsStateBegin_GameThread() override;
 	virtual void OnAsyncDestroyPhysicsStateEnd_GameThread() override;
 	virtual void OnAsyncDestroyPhysicsState() override;
diff --git a/Engine/Plugins/Experimental/FastGeoStreaming/Source/FastGeoStreaming/Private/FastGeoComponent.cpp b/Engine/Plugins/Experimental/FastGeoStreaming/Source/FastGeoStreaming/Private/FastGeoComponent.cpp
--- a/Engine/Plugins/Experimental/FastGeoStreaming/Source/FastGeoStreaming/Private/FastGeoComponent.cpp
+++ b/Engine/Plugins/Experimental/FastGeoStreaming/Source/FastGeoStreaming/Private/FastGeoComponent.cpp
@@ -110,37 +110,45 @@ void FFastGeoComponent::ApplyWorldTransform(const FTransform& InTransform)
 	WorldTransform = LocalTransform * InTransform;
 }
 
-#ifdef FASTGEO_DEBUG_PHYSICS
 void FFastGeoComponent::OnAsyncCreatePhysicsState()
 {
+#ifdef FASTGEO_DEBUG_PHYSICS
 	check(IsCollisionEnabled());
 	check(PhysicsStateCreation == EPhysicsStateCreation::NotCreated);
 	PhysicsStateCreation = EPhysicsStateCreation::Creating;
+#endif
 }
 
 void FFastGeoComponent::OnAsyncCreatePhysicsStateEnd_GameThread()
 {
+#ifdef FASTGEO_DEBUG_PHYSICS
 	check(PhysicsStateCreation == EPhysicsStateCreation::Creating);
 	PhysicsStateCreation = EPhysicsStateCreation::Created;
+#endif
 }
 
 void FFastGeoComponent::OnAsyncDestroyPhysicsStateBegin_GameThread()
 {
+#ifdef FASTGEO_DEBUG_PHYSICS
 	check(PhysicsStateCreation == EPhysicsStateCreation::Created);
 	PhysicsStateCreation = EPhysicsStateCreation::Destroying;
+#endif
 }
 
 void FFastGeoComponent::OnAsyncDestroyPhysicsState()
 {
+#ifdef FASTGEO_DEBUG_PHYSICS
 	check(IsCollisionEnabled());
+#endif
 }
 
 void FFastGeoComponent::OnAsyncDestroyPhysicsStateEnd_GameThread()
 {
+#ifdef FASTGEO_DEBUG_PHYSICS
 	check(PhysicsStateCreation == EPhysicsStateCreation::Destroying);
 	PhysicsStateCreation = EPhysicsStateCreation::NotCreated;
-}
 #endif
+}
 
 bool FFastGeoComponent::IsRenderStateCreated() const
 {
diff --git a/Engine/Plugins/Experimental/FastGeoStreaming/Source/FastGeoStreaming/Private/FastGeoContainer.cpp b/Engine/Plugins/Experimental/FastGeoStreaming/Source/FastGeoStreaming/Private/FastGeoContainer.cpp
--- a/Engine/Plugins/Experimental/FastGeoStreaming/Source/FastGeoStreaming/Private/FastGeoContainer.cpp
+++ b/Engine/Plugins/Experimental/FastGeoStreaming/Source/FastGeoStreaming/Private/FastGeoContainer.cpp
@@ -29,6 +29,7 @@
 #include "Templates/Invoke.h"
 #include "WorldPartition/WorldPartitionHelpers.h"
 #include "WorldPartition/WorldPartition.h"
+#include "Engine/Engine.h"
 
 namespace FastGeo
 {
@@ -929,12 +930,21 @@ UObject* UFastGeoContainer::GetAsyncPhysicsStateObject() const
 	return const_cast<UFastGeoContainer*>(this);
 }
 
-void UFastGeoContainer::OnAsyncCreatePhysicsStateBegin_GameThread()
+void UFastGeoContainer::OnAsyncCreatePhysicsStateBegin_GameThread(TArray<UObject*>& OutRootedObjects)
 {
+	TRACE_CPUPROFILER_EVENT_SCOPE(UFastGeoContainer::OnAsyncCreatePhysicsStateBegin_GameThread);
 	check(PhysicsStateCreation == EPhysicsStateCreation::NotCreated);
 	check(PendingCreate.PhysicsState.ComponentsToProcess);
 	check(PendingCreate.PhysicsState.ComponentsToProcess->Num());
 	PhysicsStateCreation = EPhysicsStateCreation::Creating;
+
+	// No need to fill OutRootedObjects as Assets already keeps objects alive
+	bool bSingleThreaded = !FApp::ShouldUseThreadingForPerformance();
+	ParallelFor(CollisionComponents.Num(), [&](int32 Index)
+	{
+		FTaskTagScope Scope(ETaskTag::EParallelGameThread);
+		CollisionComponents[Index]->OnAsyncCreatePhysicsStateBegin_GameThread();
+	}, bSingleThreaded);
 }
 
 bool UFastGeoContainer::OnAsyncCreatePhysicsState(const UE::FTimeout& Timeout)
@@ -964,18 +974,15 @@ void UFastGeoContainer::OnAsyncCreatePhysicsStateEnd_GameThread()
 	check(PhysicsStateCreation == EPhysicsStateCreation::Creating);
 	PhysicsStateCreation = EPhysicsStateCreation::Created;
 	
-#ifdef FASTGEO_DEBUG_PHYSICS
-	// This loop is only necessary for debugging
 	for (FFastGeoComponent* Component : CollisionComponents)
 	{
 		Component->OnAsyncCreatePhysicsStateEnd_GameThread();
 	}
-#endif
 
 	RegisterBodyInstances();
 }
 
-void UFastGeoContainer::OnAsyncDestroyPhysicsStateBegin_GameThread()
+void UFastGeoContainer::OnAsyncDestroyPhysicsStateBegin_GameThread(TArray<UObject*>& OutRootedObjects)
 {
 	TRACE_CPUPROFILER_EVENT_SCOPE(UFastGeoContainer::OnAsyncDestroyPhysicsStateBegin_GameThread);
 	check(PhysicsStateCreation == EPhysicsStateCreation::Created);
diff --git a/Engine/Plugins/Experimental/FastGeoStreaming/Source/FastGeoStreaming/Private/FastGeoInstancedStaticMeshComponent.cpp b/Engine/Plugins/Experimental/FastGeoStreaming/Source/FastGeoStreaming/Private/FastGeoInstancedStaticMeshComponent.cpp
--- a/Engine/Plugins/Experimental/FastGeoStreaming/Source/FastGeoStreaming/Private/FastGeoInstancedStaticMeshComponent.cpp
+++ b/Engine/Plugins/Experimental/FastGeoStreaming/Source/FastGeoStreaming/Private/FastGeoInstancedStaticMeshComponent.cpp
@@ -283,6 +283,19 @@ void FFastGeoInstancedStaticMeshComponent::OnAsyncCreatePhysicsState()
 	FFastGeoComponent::OnAsyncCreatePhysicsState();
 }
 
+void FFastGeoInstancedStaticMeshComponent::OnAsyncCreatePhysicsStateEnd_GameThread()
+{
+	for (FBodyInstance* Body : InstanceBodies)
+	{
+		if (Body)
+		{
+			Body->ClearCachedPhysicsCreationInputs();
+		}
+	}
+
+	Super::OnAsyncCreatePhysicsStateEnd_GameThread();
+}
+
 void FFastGeoInstancedStaticMeshComponent::OnAsyncDestroyPhysicsStateBegin_GameThread()
 {
 	FFastGeoComponent::OnAsyncDestroyPhysicsStateBegin_GameThread();
diff --git a/Engine/Plugins/Experimental/FastGeoStreaming/Source/FastGeoStreaming/Private/FastGeoPhysicsBodyInstanceOwner.cpp b/Engine/Plugins/Experimental/FastGeoStreaming/Source/FastGeoStreaming/Private/FastGeoPhysicsBodyInstanceOwner.cpp
--- a/Engine/Plugins/Experimental/FastGeoStreaming/Source/FastGeoStreaming/Private/FastGeoPhysicsBodyInstanceOwner.cpp
+++ b/Engine/Plugins/Experimental/FastGeoStreaming/Source/FastGeoStreaming/Private/FastGeoPhysicsBodyInstanceOwner.cpp
@@ -128,4 +128,10 @@ UMaterialInterface* FFastGeoPhysicsBodyInstanceOwner::GetMaterial(int32 Index) c
 {
 	check(OwnerContainer.IsValid());
 	return OwnerComponent->GetMaterial(Index);
+}
+
+UBodySetup* FFastGeoPhysicsBodyInstanceOwner::GetPhysicsBodySetup() const
+{
+	check(OwnerContainer.IsValid());
+	return OwnerComponent->GetBodySetup();
 }
\ No newline at end of file
diff --git a/Engine/Plugins/Experimental/FastGeoStreaming/Source/FastGeoStreaming/Private/FastGeoPrimitiveComponent.cpp b/Engine/Plugins/Experimental/FastGeoStreaming/Source/FastGeoStreaming/Private/FastGeoPrimitiveComponent.cpp
--- a/Engine/Plugins/Experimental/FastGeoStreaming/Source/FastGeoStreaming/Private/FastGeoPrimitiveComponent.cpp
+++ b/Engine/Plugins/Experimental/FastGeoStreaming/Source/FastGeoStreaming/Private/FastGeoPrimitiveComponent.cpp
@@ -920,6 +920,22 @@ FBodyInstance* FFastGeoPrimitiveComponent::GetBodyInstance(FName BoneName, bool
 	return const_cast<FBodyInstance*>(&BodyInstance);
 }
 
+void FFastGeoPrimitiveComponent::OnAsyncCreatePhysicsStateBegin_GameThread()
+{
+	BodyInstanceOwner.Initialize(this);
+	BodyInstance.CachePhysicsCreationInputs(&BodyInstanceOwner);
+	BodyInstanceOwner.Uninitialize();
+
+	Super::OnAsyncCreatePhysicsStateBegin_GameThread();
+}
+
+void FFastGeoPrimitiveComponent::OnAsyncCreatePhysicsStateEnd_GameThread()
+{
+	BodyInstance.ClearCachedPhysicsCreationInputs();
+
+	Super::OnAsyncCreatePhysicsStateEnd_GameThread();
+}
+
 UFastGeoSurrogateComponent* FFastGeoPrimitiveComponent::GetSurrogateComponent() const
 {
 	if (UFastGeoContainer* Container = GetOwnerContainer(); ensure(Container))
diff --git a/Engine/Plugins/Experimental/FastGeoStreaming/Source/FastGeoStreaming/Private/FastGeoStaticMeshComponent.cpp b/Engine/Plugins/Experimental/FastGeoStreaming/Source/FastGeoStreaming/Private/FastGeoStaticMeshComponent.cpp
--- a/Engine/Plugins/Experimental/FastGeoStreaming/Source/FastGeoStreaming/Private/FastGeoStaticMeshComponent.cpp
+++ b/Engine/Plugins/Experimental/FastGeoStreaming/Source/FastGeoStreaming/Private/FastGeoStaticMeshComponent.cpp
@@ -132,6 +132,14 @@ const Nanite::FResources* FFastGeoStaticMeshComponentBase::GetNaniteResources()
 
 UBodySetup* FFastGeoStaticMeshComponentBase::GetBodySetup() const
 {
+	if (!IsInGameThread())
+	{
+		if (BodyInstance.GetPhysicsCreationInputs().IsValid())
+		{
+			return BodyInstance.GetPhysicsCreationInputs()->ResolvedBodySetup;
+		}
+	}
+
 	return GetStaticMesh() ? GetStaticMesh()->GetBodySetup() : nullptr;
 }
 
diff --git a/Engine/Source/Runtime/Engine/Classes/Components/ActorComponent.h b/Engine/Source/Runtime/Engine/Classes/Components/ActorComponent.h
--- a/Engine/Source/Runtime/Engine/Classes/Components/ActorComponent.h
+++ b/Engine/Source/Runtime/Engine/Classes/Components/ActorComponent.h
@@ -899,7 +899,7 @@ public:
 	ENGINE_API virtual bool OnAsyncCreatePhysicsState(const UE::FTimeout& Timeout) override;
 	ENGINE_API virtual void OnAsyncCreatePhysicsStateEnd_GameThread() override;
 	virtual bool OnAsyncDestroyPhysicsState(const UE::FTimeout& Timeout) override { return true; }
-	ENGINE_API virtual void OnAsyncDestroyPhysicsStateBegin_GameThread() override;
+	ENGINE_API virtual void OnAsyncDestroyPhysicsStateBegin_GameThread(TArray<UObject*>& OutRootedObjects) override;
 	ENGINE_API virtual void OnAsyncDestroyPhysicsStateEnd_GameThread() override;
 	//~ End IAsyncPhysicsStateProcessor
 	/**
diff --git a/Engine/Source/Runtime/Engine/Classes/Components/InstancedStaticMeshComponent.h b/Engine/Source/Runtime/Engine/Classes/Components/InstancedStaticMeshComponent.h
--- a/Engine/Source/Runtime/Engine/Classes/Components/InstancedStaticMeshComponent.h
+++ b/Engine/Source/Runtime/Engine/Classes/Components/InstancedStaticMeshComponent.h
@@ -526,8 +526,9 @@ public:
 protected:
 	ENGINE_API virtual void OnCreatePhysicsState() override;
 	ENGINE_API virtual void OnDestroyPhysicsState() override;
+	ENGINE_API virtual void OnAsyncCreatePhysicsStateEnd_GameThread() override;
 	ENGINE_API virtual bool OnAsyncDestroyPhysicsState(const UE::FTimeout& Timeout) override;
-	ENGINE_API virtual void OnAsyncDestroyPhysicsStateBegin_GameThread() override;
+	ENGINE_API virtual void OnAsyncDestroyPhysicsStateBegin_GameThread(TArray<UObject*>& OutRootedObjects) override;
 
 public:
 	ENGINE_API virtual bool CanEditSimulatePhysics() override;
diff --git a/Engine/Source/Runtime/Engine/Classes/Components/PrimitiveComponent.h b/Engine/Source/Runtime/Engine/Classes/Components/PrimitiveComponent.h
--- a/Engine/Source/Runtime/Engine/Classes/Components/PrimitiveComponent.h
+++ b/Engine/Source/Runtime/Engine/Classes/Components/PrimitiveComponent.h
@@ -2522,9 +2522,10 @@ public:
 	//~ End UActorComponent Interface
 
 	//~ Begin IAsyncPhysicsStateProcessor
+	ENGINE_API virtual void OnAsyncCreatePhysicsStateBegin_GameThread(TArray<UObject*>& OutRootedObjects) override;
 	ENGINE_API virtual void OnAsyncCreatePhysicsStateEnd_GameThread() override;
 	ENGINE_API virtual bool OnAsyncDestroyPhysicsState(const UE::FTimeout& Timeout) override;
-	ENGINE_API virtual void OnAsyncDestroyPhysicsStateBegin_GameThread() override;
+	ENGINE_API virtual void OnAsyncDestroyPhysicsStateBegin_GameThread(TArray<UObject*>& OutRootedObjects) override;
 	ENGINE_API virtual void OnAsyncDestroyPhysicsStateEnd_GameThread() override;
 	ENGINE_API virtual void CollectBodySetupsWithPhysicsMeshesToCreate(TSet<UBodySetup*>& OutBodySetups) const override;
 	//~ End IAsyncPhysicsStateProcessor
@@ -3195,6 +3196,9 @@ public:
 	/** Returns the transform of this body instance owner. */
 	virtual FTransform GetPhysicsOwnerTransform() const override { return GetComponentTransform(); }
 
+	/** Return the BodySetup to use for this body instance owner (single body case) */
+	virtual UBodySetup* GetPhysicsBodySetup() const override { return const_cast<UPrimitiveComponent*>(this)->GetBodySetup(); }
+
 	/** Returns the body instance's physical material override. */
 	ENGINE_API virtual UPhysicalMaterial* GetPhysicsMaterialOverride() const override;
 
diff --git a/Engine/Source/Runtime/Engine/Classes/Components/SkeletalMeshComponent.h b/Engine/Source/Runtime/Engine/Classes/Components/SkeletalMeshComponent.h
--- a/Engine/Source/Runtime/Engine/Classes/Components/SkeletalMeshComponent.h
+++ b/Engine/Source/Runtime/Engine/Classes/Components/SkeletalMeshComponent.h
@@ -1866,10 +1866,10 @@ protected:
 public:
 	ENGINE_API virtual bool AllowsAsyncPhysicsStateCreation() const override;
 	ENGINE_API virtual bool AllowsAsyncPhysicsStateDestruction() const override;
-	ENGINE_API virtual void OnAsyncCreatePhysicsStateBegin_GameThread() override;
+	ENGINE_API virtual void OnAsyncCreatePhysicsStateBegin_GameThread(TArray<UObject*>& OutRootedObjects) override;
 	ENGINE_API virtual bool OnAsyncCreatePhysicsState(const UE::FTimeout& TimeOut) override;
 	ENGINE_API virtual void OnAsyncCreatePhysicsStateEnd_GameThread() override;
-	ENGINE_API virtual void OnAsyncDestroyPhysicsStateBegin_GameThread() override;
+	ENGINE_API virtual void OnAsyncDestroyPhysicsStateBegin_GameThread(TArray<UObject*>& OutRootedObjects) override;
 	ENGINE_API virtual bool OnAsyncDestroyPhysicsState(const UE::FTimeout& Timeout) override;
 
 	ENGINE_API virtual void InitializeComponent() override;
diff --git a/Engine/Source/Runtime/Engine/Classes/PhysicsEngine/BodyInstance.h b/Engine/Source/Runtime/Engine/Classes/PhysicsEngine/BodyInstance.h
--- a/Engine/Source/Runtime/Engine/Classes/PhysicsEngine/BodyInstance.h
+++ b/Engine/Source/Runtime/Engine/Classes/PhysicsEngine/BodyInstance.h
@@ -58,6 +58,25 @@ namespace EDOFMode
 }
 
 struct FBodyInstance;
+struct FBodyInstancePhysicsCreationInputs;
+
+using FBodyInstancePhysicsCreationInputsPtr = TSharedPtr<FBodyInstancePhysicsCreationInputs, ESPMode::ThreadSafe>;
+
+// Cache used when creating physics body instances
+struct FBodyInstancePhysicsCreationInputs
+{
+	void GetReferencedObjects(TArray<UObject*>& OutReferencedObjects);
+
+	// Cached body setup resolved on GT
+	UBodySetup* ResolvedBodySetup = nullptr;
+
+	// Cached simple material raw ptr
+	UPhysicalMaterial* ResolvedSimpleMaterial = nullptr;
+
+	// Cached complex materials and masks
+	TArray<UPhysicalMaterial*> ResolvedComplexMaterials;
+	TArray<FPhysicalMaterialMaskParams> ResolvedComplexMasks;
+};
 
 #ifndef CHAOS_DEBUG_NAME
 #define CHAOS_DEBUG_NAME 0
@@ -1522,6 +1541,15 @@ private:
 	/** Used to map between shapes and welded bodies. We do not create entries if the owning body instance is root*/
 	TSharedPtr<TMap<FPhysicsShapeHandle, FWeldInfo>> ShapeToBodiesMap;
 
+
+public:
+	// Resolved Inputs used by physics state creation (currently only used by async mode)
+	ENGINE_API const FBodyInstancePhysicsCreationInputsPtr& CachePhysicsCreationInputs(IPhysicsBodyInstanceOwner* BodyInstanceOwner);
+	void ClearCachedPhysicsCreationInputs() { PhysicsCreationInputs.Reset(); }
+	const FBodyInstancePhysicsCreationInputsPtr& GetPhysicsCreationInputs() const { return PhysicsCreationInputs; }
+
+private:
+	FBodyInstancePhysicsCreationInputsPtr PhysicsCreationInputs;
 };
 
 template<>
diff --git a/Engine/Source/Runtime/Engine/Classes/PhysicsEngine/PhysicsBodyInstanceOwnerInterface.h b/Engine/Source/Runtime/Engine/Classes/PhysicsEngine/PhysicsBodyInstanceOwnerInterface.h
--- a/Engine/Source/Runtime/Engine/Classes/PhysicsEngine/PhysicsBodyInstanceOwnerInterface.h
+++ b/Engine/Source/Runtime/Engine/Classes/PhysicsEngine/PhysicsBodyInstanceOwnerInterface.h
@@ -40,18 +40,6 @@ class IPhysicsBodyInstanceOwner
 public:
 	virtual ~IPhysicsBodyInstanceOwner() = default;
 
-	UE_DEPRECATED(5.8, "Use GetPhysicsBodyInstanceOwnerFromHitResult instead.")
-	static IPhysicsBodyInstanceOwner* GetPhysicsBodyInstandeOwnerFromHitResult(const FHitResult& Result)
-	{
-		return const_cast<IPhysicsBodyInstanceOwner*>(GetPhysicsBodyInstanceOwnerFromHitResult(Result));
-	}
-
-	UE_DEPRECATED(5.8, "Use GetPhysicsBodyInstanceOwnerFromOverlapResult instead.")
-	static IPhysicsBodyInstanceOwner* GetPhysicsBodyInstandeOwnerFromOverlapResult(const FOverlapResult& OverlapResult)
-	{
-		return const_cast<IPhysicsBodyInstanceOwner*>(GetPhysicsBodyInstanceOwnerFromOverlapResult(OverlapResult));
-	}
-
 	/** Returns the IPhysicsBodyInstanceOwner based on a given hit result. 
 	 * @param HiResult Hit result
 	 * @param bPreferComponent Whether to prioritize the hit result component (if valid) when resolving the IPhysicsBodyInstanceOwner.
@@ -86,8 +74,8 @@ public:
 	/** Returns the form of collision for this body instance owner. */
 	virtual ECollisionEnabled::Type GetCollisionEnabled() const = 0;
 
-	UE_DEPRECATED(5.8, "GetComplexPhysicalMaterials was removed.")
-	virtual void GetComplexPhysicalMaterials(TArray<UPhysicalMaterial*>& OutPhysMaterials, TArray<FPhysicalMaterialMaskParams>* OutPhysMaterialMasks) const {}
+	/** Return the BodySetup to use for this body instance owner (single body case) */
+	virtual UBodySetup* GetPhysicsBodySetup() const = 0;
 
 	/** Gets the response type given a specific channel. */
 	virtual ECollisionResponse GetCollisionResponseToChannel(ECollisionChannel Channel) const = 0;
@@ -109,17 +97,35 @@ public:
 	 */
 	virtual FBodyInstance* GetBodyInstance(FName BoneName = NAME_None, bool bGetWelded = true, int32 Index = INDEX_NONE) const = 0;
 
-	/** Returns the body instance's physics material override (used by FBodyInstance::GetSimplePhysicalMaterial / FBodyInstance::GetComplexPhysicalMaterials). */
+	/** Returns the body instance's physics material override. */
 	virtual UPhysicalMaterial* GetPhysicsMaterialOverride() const = 0;
 
-	/** Returns the body instance owner mesh material (used by FBodyInstance::GetSimplePhysicalMaterial). */
+	/** Returns the body instance owner mesh material. */
 	virtual UMaterialInterface* GetPhysicsMaterialBase() const = 0;
 
-	/** Returns the body instance owner number of materials (used by FBodyInstance::GetComplexPhysicalMaterials). */
+	/** Returns the body instance owner number of materials. */
 	virtual int32 GetNumMaterials() const = 0;
 
-	/** Returns the body instance owner material for the provided material index (used by FBodyInstance::GetComplexPhysicalMaterials). */
+	/** Returns the body instance owner material for the provided material index. */
 	virtual UMaterialInterface* GetMaterial(int32 Index) const = 0;
 
+	//~ Begin Deprecation
+
+	UE_DEPRECATED(5.8, "Use GetPhysicsBodyInstanceOwnerFromHitResult instead.")
+	static IPhysicsBodyInstanceOwner* GetPhysicsBodyInstandeOwnerFromHitResult(const FHitResult& Result)
+	{
+		return const_cast<IPhysicsBodyInstanceOwner*>(GetPhysicsBodyInstanceOwnerFromHitResult(Result));
+	}
+
+	UE_DEPRECATED(5.8, "Use GetPhysicsBodyInstanceOwnerFromOverlapResult instead.")
+	static IPhysicsBodyInstanceOwner* GetPhysicsBodyInstandeOwnerFromOverlapResult(const FOverlapResult& OverlapResult)
+	{
+		return const_cast<IPhysicsBodyInstanceOwner*>(GetPhysicsBodyInstanceOwnerFromOverlapResult(OverlapResult));
+	}
+	UE_DEPRECATED(5.8, "GetComplexPhysicalMaterials was removed.")
+	virtual void GetComplexPhysicalMaterials(TArray<UPhysicalMaterial*>& OutPhysMaterials, TArray<FPhysicalMaterialMaskParams>* OutPhysMaterialMasks) const {}
+
+	//~ End Deprecation
+
 	friend struct FBodyInstance;
 };
\ No newline at end of file
diff --git a/Engine/Source/Runtime/Engine/Private/Components/ActorComponent.cpp b/Engine/Source/Runtime/Engine/Private/Components/ActorComponent.cpp
--- a/Engine/Source/Runtime/Engine/Private/Components/ActorComponent.cpp
+++ b/Engine/Source/Runtime/Engine/Private/Components/ActorComponent.cpp
@@ -2378,7 +2378,7 @@ void UActorComponent::OnAsyncCreatePhysicsStateEnd_GameThread()
 	check(IsPreRegistered());
 }
 
-void UActorComponent::OnAsyncDestroyPhysicsStateBegin_GameThread()
+void UActorComponent::OnAsyncDestroyPhysicsStateBegin_GameThread(TArray<UObject*>& OutRootedObjects)
 {
 	GlobalDestroyPhysicsDelegate.Broadcast(this);
 	ensureMsgf(bRegistered, TEXT("Component has physics state when not registered (%s)"), *GetFullName()); // should not have physics state unless we are registered
diff --git a/Engine/Source/Runtime/Engine/Private/Components/PrimitiveComponent.cpp b/Engine/Source/Runtime/Engine/Private/Components/PrimitiveComponent.cpp
--- a/Engine/Source/Runtime/Engine/Private/Components/PrimitiveComponent.cpp
+++ b/Engine/Source/Runtime/Engine/Private/Components/PrimitiveComponent.cpp
@@ -966,10 +966,17 @@ void UPrimitiveComponent::OnCreatePhysicsState()
 	}
 }
 
+void UPrimitiveComponent::OnAsyncCreatePhysicsStateBegin_GameThread(TArray<UObject*>& OutRootedObjects)
+{
+	BodyInstance.CachePhysicsCreationInputs(this)->GetReferencedObjects(OutRootedObjects);
+}
+
 void UPrimitiveComponent::OnAsyncCreatePhysicsStateEnd_GameThread()
 {
 	check(IsPhysicsStateCreated());
 
+	BodyInstance.ClearCachedPhysicsCreationInputs();
+
 #if UE_ENABLE_DEBUG_DRAWING
 	SendRenderDebugPhysics();
 #endif
@@ -1066,7 +1073,7 @@ void UPrimitiveComponent::SendPhysicsTransform(ETeleportType Teleport)
 	BodyInstance.UpdateBodyScale(GetComponentTransform().GetScale3D());
 }
 
-void UPrimitiveComponent::OnAsyncDestroyPhysicsStateBegin_GameThread()
+void UPrimitiveComponent::OnAsyncDestroyPhysicsStateBegin_GameThread(TArray<UObject*>& OutRootedObjects)
 {
 	check(AllowsAsyncPhysicsStateDestruction());
 
@@ -1082,7 +1089,7 @@ void UPrimitiveComponent::OnAsyncDestroyPhysicsStateBegin_GameThread()
 	AsyncTermBodyPayload = BodyInstance.StartAsyncTermBody_GameThread();
 	check(!BodyInstance.IsValidBodyInstance());
 	
-	Super::OnAsyncDestroyPhysicsStateBegin_GameThread();
+	Super::OnAsyncDestroyPhysicsStateBegin_GameThread(OutRootedObjects);
 }
 
 bool UPrimitiveComponent::OnAsyncDestroyPhysicsState(const UE::FTimeout& Timeout)
@@ -5244,6 +5251,7 @@ void UPrimitiveComponent::PrecachePSOs()
 		return;
 	}
 
+	TRACE_CPUPROFILER_EVENT_SCOPE(UPrimitiveComponent::PrecachePSOs);
 	// clear the current request data
 	PSOPrecacheComponentData.Reset();
 
diff --git a/Engine/Source/Runtime/Engine/Private/Components/StaticMeshComponent.cpp b/Engine/Source/Runtime/Engine/Private/Components/StaticMeshComponent.cpp
--- a/Engine/Source/Runtime/Engine/Private/Components/StaticMeshComponent.cpp
+++ b/Engine/Source/Runtime/Engine/Private/Components/StaticMeshComponent.cpp
@@ -1308,6 +1308,14 @@ void UStaticMeshComponent::GetUsedTextures(TArray<UTexture*>& OutTextures, EMate
 
 UBodySetup* UStaticMeshComponent::GetBodySetup()
 {
+	if (!IsInGameThread() && IsAsyncCreatePhysicsStateRunning())
+	{
+		if (ensureAlwaysMsgf(BodyInstance.GetPhysicsCreationInputs().IsValid(), TEXT("UStaticMeshComponent::GetBodySetup called outside of game thread, expected a valid cache for %s"), *GetFullName()))
+		{
+			return BodyInstance.GetPhysicsCreationInputs()->ResolvedBodySetup;
+		}
+	}
+
 	if (GetStaticMesh())
 	{
 		return GetStaticMesh()->GetBodySetup();
diff --git a/Engine/Source/Runtime/Engine/Private/InstancedStaticMesh.cpp b/Engine/Source/Runtime/Engine/Private/InstancedStaticMesh.cpp
--- a/Engine/Source/Runtime/Engine/Private/InstancedStaticMesh.cpp
+++ b/Engine/Source/Runtime/Engine/Private/InstancedStaticMesh.cpp
@@ -3013,9 +3013,9 @@ void UInstancedStaticMeshComponent::OnDestroyPhysicsState()
 	bNavigationRelevant = IsNavigationRelevant();
 }
 
-void UInstancedStaticMeshComponent::OnAsyncDestroyPhysicsStateBegin_GameThread()
+void UInstancedStaticMeshComponent::OnAsyncDestroyPhysicsStateBegin_GameThread(TArray<UObject*>& OutRootedObjects)
 {
-	USceneComponent::OnAsyncDestroyPhysicsStateBegin_GameThread();
+	USceneComponent::OnAsyncDestroyPhysicsStateBegin_GameThread(OutRootedObjects);
 
 	// Move InstanceBodies in AsyncDestroyPhysicsStatePayload
 	check(AsyncDestroyPhysicsStatePayload.IsEmpty());
@@ -3024,6 +3024,19 @@ void UInstancedStaticMeshComponent::OnAsyncDestroyPhysicsStateBegin_GameThread()
 	bNavigationRelevant = IsNavigationRelevant();
 }
 
+void UInstancedStaticMeshComponent::OnAsyncCreatePhysicsStateEnd_GameThread()
+{
+	Super::OnAsyncCreatePhysicsStateEnd_GameThread();
+
+	for (FBodyInstance* Body : InstanceBodies)
+	{
+		if (Body)
+		{
+			Body->ClearCachedPhysicsCreationInputs();
+		}
+	}
+}
+
 bool UInstancedStaticMeshComponent::OnAsyncDestroyPhysicsState(const UE::FTimeout& Timeout)
 {
 	TRACE_CPUPROFILER_EVENT_SCOPE(UInstancedStaticMeshComponent::OnAsyncDestroyPhysicsState);
diff --git a/Engine/Source/Runtime/Engine/Private/PhysicsEngine/BodyInstance.cpp b/Engine/Source/Runtime/Engine/Private/PhysicsEngine/BodyInstance.cpp
--- a/Engine/Source/Runtime/Engine/Private/PhysicsEngine/BodyInstance.cpp
+++ b/Engine/Source/Runtime/Engine/Private/PhysicsEngine/BodyInstance.cpp
@@ -39,6 +39,7 @@
 
 #include "Components/BrushComponent.h"
 #include "PhysicalMaterials/PhysicalMaterial.h"
+#include "PhysicalMaterials/PhysicalMaterialMask.h"
 #include "PhysicsEngine/PhysicsSettings.h"
 
 #define LOCTEXT_NAMESPACE "BodyInstance"
@@ -1372,6 +1373,12 @@ bool TInitBodiesHelperBase<InAllocatorType>::CreateShapes_AssumesLocked(FBodyIns
 
 UBodySetup* FBodyInstance::GetBodySetup() const
 {
+	if (!IsInGameThread() && !IsInParallelGameThread() && PhysicsCreationInputs.IsValid())
+	{
+		return PhysicsCreationInputs->ResolvedBodySetup;
+	}
+
+	FGCScopeGuard Guard;
 	if(UBodySetupCore* BodySetupCore = BodySetup.Get())
 	{
 		return CastChecked<UBodySetup>(BodySetupCore);
@@ -3163,6 +3170,15 @@ namespace UE::Private
 
 UPhysicalMaterial* FBodyInstance::GetSimplePhysicalMaterial(const FBodyInstance* BodyInstance, TWeakObjectPtr<UPrimitiveComponent> OwnerComp, TWeakObjectPtr<UBodySetup> BodySetupPtr, IPhysicsBodyInstanceOwner* BodyInstanceOwner)
 {
+	const FBodyInstancePhysicsCreationInputsPtr& PhysicsCreationInputs = BodyInstance->GetPhysicsCreationInputs();
+	if (!IsInGameThread() && !IsInParallelGameThread())
+	{
+		if (ensureMsgf(PhysicsCreationInputs.IsValid(), TEXT("FBodyInstance::GetSimplePhysicalMaterial called outside of game thread, expected a valid cache for %s"), *BodyInstance->GetBodyDebugName()))
+		{
+			return PhysicsCreationInputs->ResolvedSimpleMaterial;
+		}
+	}
+
 	if(!GEngine || !GEngine->DefaultPhysMaterial)
 	{
 		UE_LOG(LogPhysics, Error, TEXT("FBodyInstance::GetSimplePhysicalMaterial : GEngine not initialized! Cannot call this during native CDO construction, wrap with if(!HasAnyFlags(RF_ClassDefaultObject)) or move out of constructor, material parameters will not be correct."));
@@ -3253,6 +3269,20 @@ void FBodyInstance::GetComplexPhysicalMaterials(TArray<UPhysicalMaterial*>& OutP
 
 void FBodyInstance::GetComplexPhysicalMaterials(const FBodyInstance* BodyInstance, TWeakObjectPtr<UPrimitiveComponent> OwnerComp, TArray<UPhysicalMaterial*>& OutPhysMaterials, TArray<FPhysicalMaterialMaskParams>* OutPhysMaterialMasks, IPhysicsBodyInstanceOwner* InBodyInstanceOwner)
 {
+	const FBodyInstancePhysicsCreationInputsPtr& PhysicsCreationInputs = BodyInstance->GetPhysicsCreationInputs();
+	if (!IsInGameThread() && !IsInParallelGameThread())
+	{
+		if (ensureMsgf(PhysicsCreationInputs.IsValid(), TEXT("FBodyInstance::GetComplexPhysicalMaterials called outside of game thread, expected a valid cache for %s"), *BodyInstance->GetBodyDebugName()))
+		{
+			if (OutPhysMaterialMasks)
+			{
+				*OutPhysMaterialMasks = PhysicsCreationInputs->ResolvedComplexMasks;
+			}
+			OutPhysMaterials = PhysicsCreationInputs->ResolvedComplexMaterials;
+			return;
+		}
+	}
+
 	if(!GEngine || !GEngine->DefaultPhysMaterial)
 	{
 		UE_LOG(LogPhysics, Error, TEXT("FBodyInstance::GetComplexPhysicalMaterials : GEngine not initialized! Cannot call this during native CDO construction, wrap with if(!HasAnyFlags(RF_ClassDefaultObject)) or move out of constructor, material parameters will not be correct."));
@@ -4981,6 +5011,49 @@ FRecalculatedMassProperties& FBodyInstance::OnRecalculatedMassProperties()
 	return BodyInstanceDelegates->OnRecalculatedMassProperties;
 }
 
+const FBodyInstancePhysicsCreationInputsPtr& FBodyInstance::CachePhysicsCreationInputs(IPhysicsBodyInstanceOwner* InBodyInstanceOwner)
+{
+	check(IsInGameThread() || IsInParallelGameThread());
+	check(!PhysicsCreationInputs.IsValid());
+	check(InBodyInstanceOwner);
+
+	FBodyInstancePhysicsCreationInputsPtr NewPhysicsCreationInputs = MakeShared<FBodyInstancePhysicsCreationInputs, ESPMode::ThreadSafe>();
+	UBodySetup* OwnerBodySetup = InBodyInstanceOwner->GetPhysicsBodySetup();
+	NewPhysicsCreationInputs->ResolvedBodySetup = OwnerBodySetup;
+	NewPhysicsCreationInputs->ResolvedSimpleMaterial = FBodyInstance::GetSimplePhysicalMaterial(this, TWeakObjectPtr<UPrimitiveComponent>(), OwnerBodySetup, InBodyInstanceOwner);
+	FBodyInstance::GetComplexPhysicalMaterials(this, TWeakObjectPtr<UPrimitiveComponent>(), NewPhysicsCreationInputs->ResolvedComplexMaterials, &NewPhysicsCreationInputs->ResolvedComplexMasks, InBodyInstanceOwner);
+
+	PhysicsCreationInputs = NewPhysicsCreationInputs;
+	return PhysicsCreationInputs;
+}
+
+void FBodyInstancePhysicsCreationInputs::GetReferencedObjects(TArray<UObject*>& OutReferencedObjects)
+{
+	if (ResolvedBodySetup)
+	{
+		OutReferencedObjects.Add(ResolvedBodySetup);
+	}
+
+	if (ResolvedSimpleMaterial)
+	{
+		OutReferencedObjects.Add(ResolvedSimpleMaterial);
+	}
+
+	for (UPhysicalMaterial* ComplexMaterial : ResolvedComplexMaterials)
+	{
+		if (ComplexMaterial)
+		{
+			OutReferencedObjects.Add(ComplexMaterial);
+		}
+	}
+
+	for (const FPhysicalMaterialMaskParams& Mask : ResolvedComplexMasks)
+	{
+		OutReferencedObjects.Add(Mask.PhysicalMaterialMap);
+		OutReferencedObjects.Add(Mask.PhysicalMaterialMask);
+	}
+}
+
 bool FBodyInstanceAsyncPhysicsTickHandle::IsValid() const
 {
 	return Proxy && Proxy->GetPhysicsThreadAPI() != nullptr;
diff --git a/Engine/Source/Runtime/Engine/Private/PhysicsEngine/Experimental/PhysScene_AsyncPhysicsStateJobQueue.cpp b/Engine/Source/Runtime/Engine/Private/PhysicsEngine/Experimental/PhysScene_AsyncPhysicsStateJobQueue.cpp
--- a/Engine/Source/Runtime/Engine/Private/PhysicsEngine/Experimental/PhysScene_AsyncPhysicsStateJobQueue.cpp
+++ b/Engine/Source/Runtime/Engine/Private/PhysicsEngine/Experimental/PhysScene_AsyncPhysicsStateJobQueue.cpp
@@ -60,6 +60,9 @@ FPhysScene_AsyncPhysicsStateJobQueue::~FPhysScene_AsyncPhysicsStateJobQueue()
 		Tick(true);
 	}
 	check(IsCompleted());
+
+	ensure(RootCounts.IsEmpty());
+	RootCounts.Empty();
 }
 
 void FPhysScene_AsyncPhysicsStateJobQueue::AddJob(const FJob& Job)
@@ -67,6 +70,7 @@ void FPhysScene_AsyncPhysicsStateJobQueue::AddJob(const FJob& Job)
 	TRACE_CPUPROFILER_EVENT_SCOPE(FPhysScene_AsyncPhysicsStateJobQueue::AddJob);
 	check(IsInGameThread());
 
+	Job.Owner = this;
 	Job.OnPreExecute_GameThread();
 
 	if (DeferNewJobsScopeCount > 0)
@@ -85,6 +89,7 @@ void FPhysScene_AsyncPhysicsStateJobQueue::RemoveJob(const FJob& Job)
 	TRACE_CPUPROFILER_EVENT_SCOPE(FPhysScene_AsyncPhysicsStateJobQueue::RemoveJob);
 	check(IsInGameThread());
 	check(Job.IsValid());
+	Job.Owner = this;
 
 	bool bWaitForJobToComplete = false;
 	bool bFoundJob = false;
@@ -138,6 +143,17 @@ void FPhysScene_AsyncPhysicsStateJobQueue::RemoveJob(const FJob& Job)
 	}
 }
 
+void FPhysScene_AsyncPhysicsStateJobQueue::AddReferencedObjects(FReferenceCollector& Collector)
+{
+	for (auto& Pair : RootCounts)
+	{
+		if (TObjectPtr<UObject>& Object = Pair.Key)
+		{
+			Collector.AddReferencedObject(Object);
+		}
+	}
+}
+
 void FPhysScene_AsyncPhysicsStateJobQueue::BeginDeferringNewJobs()
 {
 	check(IsInGameThread());
@@ -315,9 +331,37 @@ void FPhysScene_AsyncPhysicsStateJobQueue::Tick(bool bWaitForCompletion)
 	}
 }
 
+void FPhysScene_AsyncPhysicsStateJobQueue::AddRootObject(UObject* Obj)
+{
+	if (ensure(Obj))
+	{
+		int32& Count = RootCounts.FindOrAdd(Obj);
+		++Count;
+	}
+}
+
+void FPhysScene_AsyncPhysicsStateJobQueue::RemoveRootObject(UObject* Obj)
+{
+	if (!ensure(Obj))
+	{
+		return;
+	}
+	int32* CountPtr = RootCounts.Find(Obj);
+	if (ensure(CountPtr))
+	{
+		--(*CountPtr);
+		if (*CountPtr <= 0)
+		{
+			RootCounts.Remove(Obj);
+		}
+	}
+}
+
 void FPhysScene_AsyncPhysicsStateJobQueue::FJob::OnPreExecute_GameThread() const
 {
 	TRACE_CPUPROFILER_EVENT_SCOPE(FPhysScene_AsyncPhysicsStateJobQueue::OnPreExecute_GameThread);
+
+	RootedInputs.Reset();
 	switch (Type)
 	{
 		case EJobType::CreatePhysicsState:
@@ -332,15 +376,25 @@ void FPhysScene_AsyncPhysicsStateJobQueue::FJob::OnPreExecute_GameThread() const
 					ParallelFor(TEXT("CreatePhysicsMeshes.PF"), BodySetups.Num(), 1, [this, &BodySetups](int32 Index) { BodySetups[Index]->CreatePhysicsMeshes(); });
 				}
 			}
-			Processor->OnAsyncCreatePhysicsStateBegin_GameThread();
+			Processor->OnAsyncCreatePhysicsStateBegin_GameThread(RootedInputs);
 		}
 		break;
 		case EJobType::DestroyPhysicsState:
 		{
-			Processor->OnAsyncDestroyPhysicsStateBegin_GameThread();
+			Processor->OnAsyncDestroyPhysicsStateBegin_GameThread(RootedInputs);
 		}
 		break;
 	}
+
+	// Root referenced objects
+	if (UObject* Obj = Processor ? Processor->GetAsyncPhysicsStateObject() : nullptr)
+	{
+		Owner->AddRootObject(Obj);
+	}
+	for (UObject* Obj : RootedInputs)
+	{
+		Owner->AddRootObject(Obj);
+	}
 }
 
 bool FPhysScene_AsyncPhysicsStateJobQueue::FJob::Execute(UE::FTimeout& Timeout) const
@@ -384,4 +438,16 @@ void FPhysScene_AsyncPhysicsStateJobQueue::FJob::OnPostExecute_GameThread() cons
 			break;
 		}
 	}
+
+	// Remove rooted objects
+	if (UObject* Obj = Processor ? Processor->GetAsyncPhysicsStateObject() : nullptr)
+	{
+		Owner->RemoveRootObject(Obj);
+	}
+
+	for (UObject* Obj : RootedInputs)
+	{
+		Owner->RemoveRootObject(Obj);
+	}
+	RootedInputs.Reset();
 }
\ No newline at end of file
diff --git a/Engine/Source/Runtime/Engine/Private/SkeletalMeshComponentPhysics.cpp b/Engine/Source/Runtime/Engine/Private/SkeletalMeshComponentPhysics.cpp
--- a/Engine/Source/Runtime/Engine/Private/SkeletalMeshComponentPhysics.cpp
+++ b/Engine/Source/Runtime/Engine/Private/SkeletalMeshComponentPhysics.cpp
@@ -1251,19 +1251,19 @@ bool USkeletalMeshComponent::AllowsAsyncPhysicsStateDestruction() const
 	return AllowsAsyncPhysicsStateCreation();
 }
 
-void USkeletalMeshComponent::OnAsyncCreatePhysicsStateBegin_GameThread()
+void USkeletalMeshComponent::OnAsyncCreatePhysicsStateBegin_GameThread(TArray<UObject*>& OutRootedObjects)
 {
 	TRACE_CPUPROFILER_EVENT_SCOPE(USkeletalMeshComponent::OnAsyncCreatePhysicsStateBegin_GameThread);
 
 	if (bEnablePerPolyCollision == false)
 	{
-		USceneComponent::OnAsyncCreatePhysicsStateBegin_GameThread();
+		USceneComponent::OnAsyncCreatePhysicsStateBegin_GameThread(OutRootedObjects);
 	}
 	else
 	{
 		CreateBodySetup();
 		BodySetup->CreatePhysicsMeshes();
-		Super::OnAsyncCreatePhysicsStateBegin_GameThread();
+		Super::OnAsyncCreatePhysicsStateBegin_GameThread(OutRootedObjects);
 	}
 }
 
@@ -1314,7 +1314,7 @@ void USkeletalMeshComponent::OnAsyncCreatePhysicsStateEnd_GameThread()
 	OnSkelMeshPhysicsCreated.Broadcast();
 }
 
-void USkeletalMeshComponent::OnAsyncDestroyPhysicsStateBegin_GameThread()
+void USkeletalMeshComponent::OnAsyncDestroyPhysicsStateBegin_GameThread(TArray<UObject*>& OutRootedObjects)
 {
 	TRACE_CPUPROFILER_EVENT_SCOPE(USkeletalMeshComponent::OnAsyncDestroyPhysicsStateBegin_GameThread);
 
@@ -1322,7 +1322,7 @@ void USkeletalMeshComponent::OnAsyncDestroyPhysicsStateBegin_GameThread()
 
 	DestroyPhysicsStateInternal();
 
-	Super::OnAsyncDestroyPhysicsStateBegin_GameThread();
+	Super::OnAsyncDestroyPhysicsStateBegin_GameThread(OutRootedObjects);
 }
 
 bool USkeletalMeshComponent::OnAsyncDestroyPhysicsState(const UE::FTimeout& Timeout)
diff --git a/Engine/Source/Runtime/Engine/Public/Physics/Experimental/AsyncPhysicsStateProcessorInterface.h b/Engine/Source/Runtime/Engine/Public/Physics/Experimental/AsyncPhysicsStateProcessorInterface.h
--- a/Engine/Source/Runtime/Engine/Public/Physics/Experimental/AsyncPhysicsStateProcessorInterface.h
+++ b/Engine/Source/Runtime/Engine/Public/Physics/Experimental/AsyncPhysicsStateProcessorInterface.h
@@ -28,21 +28,45 @@ public:
 	/** Returns body setups that needs to create their physics meshes before physics state asynchronous creation */
 	virtual void CollectBodySetupsWithPhysicsMeshesToCreate(TSet<UBodySetup*>& OutBodySetups) const {}
 
-	/** Called on the GameThread before the component's physic engine information is created. */
-	virtual void OnAsyncCreatePhysicsStateBegin_GameThread() {}
+	/** 
+	 * Called on the GameThread before the component's physics state is created.
+	 * Allows preparing caches for asynchronous physics state creation and adding
+	 * any objects that must remain rooted during the async process.
+	 * 
+	 * @param OutRootedObjects Objects that must be rooted for the duration of the async creation.
+	 */
+	virtual void OnAsyncCreatePhysicsStateBegin_GameThread(TArray<UObject*>& OutRootedObjects) {}
 
-	/** Used to create any physics engine information for this component outside of the GameThread. */
+	/** Used to create any physics state for this component outside of the GameThread. */
 	virtual bool OnAsyncCreatePhysicsState(const UE::FTimeout& TimeOut) { return true; }
 
-	/** Called on the GameThread once the component's physic engine information is created. */
+	/**
+	 * Called on the GameThread once the component's physics state is created.
+	 * If necessary, free caches.
+	 */
 	virtual void OnAsyncCreatePhysicsStateEnd_GameThread() {}
 
-	/** Called on the GameThread before the component's physic engine information is destroyed. */
-	virtual void OnAsyncDestroyPhysicsStateBegin_GameThread() {}
+	/** 
+	 * Called on the GameThread before the component's physics state is destroyed.
+	 * Allows preparing caches for asynchronous physics state destruction and adding
+	 * any objects that must remain rooted during the async process.
+	 * 
+	 * @param OutRootedObjects Objects that must be rooted for the duration of the async destruction.
+	 */
+	virtual void OnAsyncDestroyPhysicsStateBegin_GameThread(TArray<UObject*>& OutRootedObjects) {}
 
-	/** Used to destroy any physics engine information for this component outside of the GameThread. */
+	/** Used to destroy any physics state for this component outside of the GameThread. */
 	virtual bool OnAsyncDestroyPhysicsState(const UE::FTimeout& TimeOut) { return true; }
 
-	/** Called on the GameThread once the component's physic engine information is destroyed. */
+	/**
+	 * Called on the GameThread once the component's physics state is destroyed. 
+	 * If necessary, free caches.
+	 */
 	virtual void OnAsyncDestroyPhysicsStateEnd_GameThread() {}
+
+	UE_DEPRECATED(5.8, "Use OnAsyncCreatePhysicsStateBegin_GameThread with parameters instead")
+	virtual void OnAsyncCreatePhysicsStateBegin_GameThread() {}
+
+	UE_DEPRECATED(5.8, "Use OnAsyncDestroyPhysicsStateBegin_GameThread with parameters instead")
+	virtual void OnAsyncDestroyPhysicsStateBegin_GameThread() {}
 };
\ No newline at end of file
diff --git a/Engine/Source/Runtime/Engine/Public/Physics/Experimental/PhysScene_AsyncPhysicsStateJobQueue.h b/Engine/Source/Runtime/Engine/Public/Physics/Experimental/PhysScene_AsyncPhysicsStateJobQueue.h
--- a/Engine/Source/Runtime/Engine/Public/Physics/Experimental/PhysScene_AsyncPhysicsStateJobQueue.h
+++ b/Engine/Source/Runtime/Engine/Public/Physics/Experimental/PhysScene_AsyncPhysicsStateJobQueue.h
@@ -6,10 +6,11 @@
 #include "Tasks/Task.h"
 #include "Async/Fundamental/Task.h"
 #include "UObject/Object.h"
+#include "UObject/GCObject.h"
 #include "PhysicsInterfaceDeclaresCore.h"
 #include "Physics/Experimental/AsyncPhysicsStateProcessorInterface.h"
 
-struct FPhysScene_AsyncPhysicsStateJobQueue
+struct FPhysScene_AsyncPhysicsStateJobQueue : public FGCObject
 {
 	FPhysScene_AsyncPhysicsStateJobQueue(FPhysScene* InPhysScene);
 	~FPhysScene_AsyncPhysicsStateJobQueue();
@@ -29,7 +30,8 @@ struct FPhysScene_AsyncPhysicsStateJobQueue
 
 		bool IsValid() const
 		{
-			return Processor && ::IsValid(Processor->GetAsyncPhysicsStateObject());
+			// No need to test for Processor->GetAsyncPhysicsStateObject() as it's rooted by OnPreExecute_GameThread
+			return (Processor != nullptr);
 		}
 
 		bool Execute(UE::FTimeout& Timeout) const;
@@ -46,8 +48,12 @@ struct FPhysScene_AsyncPhysicsStateJobQueue
 			return GetTypeHash(TPair<IAsyncPhysicsStateProcessor*, EJobType>{InJob.Processor, InJob.Type});
 		}
 
+		mutable TArray<UObject*> RootedInputs;
 		IAsyncPhysicsStateProcessor* Processor;
 		EJobType Type;
+	private:
+		mutable FPhysScene_AsyncPhysicsStateJobQueue* Owner = nullptr;
+		friend struct FPhysScene_AsyncPhysicsStateJobQueue;
 	};
 
 	struct FScopedDeferAsyncPhysicsStateJobs
@@ -67,6 +73,11 @@ struct FPhysScene_AsyncPhysicsStateJobQueue
 	void RemoveJob(const FJob& Job);
 	bool IsCompleted() const;
 
+	//~ Begin FGCObject interface
+	virtual void AddReferencedObjects(FReferenceCollector& Collector) override;
+	virtual FString GetReferencerName() const override { return TEXT("FPhysScene_AsyncPhysicsStateJobQueue"); }
+	//~ End FGCObject interface
+
 private:
 
 	void BeginDeferringNewJobs();
@@ -76,6 +87,9 @@ private:
 	void ExecuteJobsAsync(double TimeBudgetSeconds);
 	void OnUpdateLevelStreamingDone();
 
+	void AddRootObject(UObject* Obj);
+	void RemoveRootObject(UObject* Obj);
+
 	// GameThread variables
 	FPhysScene* PhysScene = nullptr; // The Physics scene
 	UE::Tasks::FTask AsyncJobTask; // The async task
@@ -98,5 +112,8 @@ private:
 	std::atomic<bool> bIsBlocking = false; // Notifies to the async task that we are block waiting for the task to complete
 	std::atomic<int32> GameThreadEpoch = 0; // Epoch updated by the GameThread at every frame and used by the async task to reset UsedAsyncTaskTimeBudgetSec
 
+	// GT-only: tracks roots across jobs (counts to handle duplicate/additional jobs)
+	TMap<TObjectPtr<UObject>, int32> RootCounts;
+
 	friend struct FScopedDeferAsyncPhysicsStateJobs;
 };
\ No newline at end of file
diff --git a/Engine/Source/Runtime/PhysicsCore/Private/PhysicalMaterial.cpp b/Engine/Source/Runtime/PhysicsCore/Private/PhysicalMaterial.cpp
--- a/Engine/Source/Runtime/PhysicsCore/Private/PhysicalMaterial.cpp
+++ b/Engine/Source/Runtime/PhysicsCore/Private/PhysicalMaterial.cpp
@@ -7,6 +7,7 @@
 #include "PhysicalMaterials/PhysicalMaterial.h"
 #include "PhysicalMaterials/PhysicalMaterialPropertyBase.h"
 #include "UObject/UObjectIterator.h"
+#include "Chaos/AsyncInitBodyHelper.h"
 #include "Chaos/PhysicalMaterials.h"
 #include "HAL/IConsoleManager.h"
 #include "UObject/FortniteMainBranchObjectVersion.h"
@@ -128,6 +129,12 @@ void UPhysicalMaterial::PostLoad()
 		// We don't have a color yet, generate one based on a stable hash generated by the asset's path : 
 		DebugColor = FLinearColor::MakeRandomSeededColor(GetTypeHash(GetPathName()));
 	}
+
+	if (Chaos::CVars::bEnableAsyncInitBody)
+	{
+		// Prime material handle for thread-safe access in async physics path (i.e. when p.Chaos.EnableAsyncInitBody = true)
+		EnsurePhysicsMaterialCreated();
+	}
 }
 
 void UPhysicalMaterial::FinishDestroy()
@@ -139,7 +146,7 @@ void UPhysicalMaterial::FinishDestroy()
 	Super::FinishDestroy();
 }
 
-FPhysicsMaterialHandle& UPhysicalMaterial::GetPhysicsMaterial()
+void UPhysicalMaterial::EnsurePhysicsMaterialCreated()
 {
 	if(!MaterialHandle)
 	{
@@ -153,7 +160,11 @@ FPhysicsMaterialHandle& UPhysicalMaterial::GetPhysicsMaterial()
 		FChaosEngineInterface::SetUserData(*MaterialHandle, &UserData);
 		FChaosEngineInterface::UpdateMaterial(*MaterialHandle, this);
 	}
+}
 
+FPhysicsMaterialHandle& UPhysicalMaterial::GetPhysicsMaterial()
+{
+	EnsurePhysicsMaterialCreated();
 	return *MaterialHandle;
 }
 
diff --git a/Engine/Source/Runtime/PhysicsCore/Public/PhysicalMaterials/PhysicalMaterial.h b/Engine/Source/Runtime/PhysicsCore/Public/PhysicalMaterials/PhysicalMaterial.h
--- a/Engine/Source/Runtime/PhysicsCore/Public/PhysicalMaterials/PhysicalMaterial.h
+++ b/Engine/Source/Runtime/PhysicsCore/Public/PhysicalMaterials/PhysicalMaterial.h
@@ -240,6 +240,10 @@ public:
 
 	/** Determine Material Type from input PhysicalMaterial **/
 	static PHYSICSCORE_API EPhysicalSurface DetermineSurfaceType(UPhysicalMaterial const* PhysicalMaterial);
+
+private:
+
+	void EnsurePhysicsMaterialCreated();
 };
 
 
