Fixed static analysis

#Synced-CL 50924777

diff --git a/Engine/Source/Runtime/Engine/Private/PhysicsEngine/BodyInstance.cpp b/Engine/Source/Runtime/Engine/Private/PhysicsEngine/BodyInstance.cpp
--- a/Engine/Source/Runtime/Engine/Private/PhysicsEngine/BodyInstance.cpp
+++ b/Engine/Source/Runtime/Engine/Private/PhysicsEngine/BodyInstance.cpp
@@ -3269,17 +3269,20 @@ void FBodyInstance::GetComplexPhysicalMaterials(TArray<UPhysicalMaterial*>& OutP
 
 void FBodyInstance::GetComplexPhysicalMaterials(const FBodyInstance* BodyInstance, TWeakObjectPtr<UPrimitiveComponent> OwnerComp, TArray<UPhysicalMaterial*>& OutPhysMaterials, TArray<FPhysicalMaterialMaskParams>* OutPhysMaterialMasks, IPhysicsBodyInstanceOwner* InBodyInstanceOwner)
 {
-	const FBodyInstancePhysicsCreationInputsPtr& PhysicsCreationInputs = BodyInstance->GetPhysicsCreationInputs();
-	if (!IsInGameThread() && !IsInParallelGameThread())
+	if (BodyInstance)
 	{
-		if (ensureMsgf(PhysicsCreationInputs.IsValid(), TEXT("FBodyInstance::GetComplexPhysicalMaterials called outside of game thread, expected a valid cache for %s"), *BodyInstance->GetBodyDebugName()))
+		const FBodyInstancePhysicsCreationInputsPtr& PhysicsCreationInputs = BodyInstance->GetPhysicsCreationInputs();
+		if (!IsInGameThread() && !IsInParallelGameThread())
 		{
-			if (OutPhysMaterialMasks)
+			if (ensureMsgf(PhysicsCreationInputs.IsValid(), TEXT("FBodyInstance::GetComplexPhysicalMaterials called outside of game thread, expected a valid cache for %s"), *BodyInstance->GetBodyDebugName()))
 			{
-				*OutPhysMaterialMasks = PhysicsCreationInputs->ResolvedComplexMasks;
+				if (OutPhysMaterialMasks)
+				{
+					*OutPhysMaterialMasks = PhysicsCreationInputs->ResolvedComplexMasks;
+				}
+				OutPhysMaterials = PhysicsCreationInputs->ResolvedComplexMaterials;
+				return;
 			}
-			OutPhysMaterials = PhysicsCreationInputs->ResolvedComplexMaterials;
-			return;
 		}
 	}
 
diff --git a/Engine/Source/Runtime/Engine/Private/PhysicsEngine/Experimental/PhysScene_AsyncPhysicsStateJobQueue.cpp b/Engine/Source/Runtime/Engine/Private/PhysicsEngine/Experimental/PhysScene_AsyncPhysicsStateJobQueue.cpp
--- a/Engine/Source/Runtime/Engine/Private/PhysicsEngine/Experimental/PhysScene_AsyncPhysicsStateJobQueue.cpp
+++ b/Engine/Source/Runtime/Engine/Private/PhysicsEngine/Experimental/PhysScene_AsyncPhysicsStateJobQueue.cpp
@@ -387,7 +387,7 @@ void FPhysScene_AsyncPhysicsStateJobQueue::FJob::OnPreExecute_GameThread() const
 	}
 
 	// Root referenced objects
-	if (UObject* Obj = Processor ? Processor->GetAsyncPhysicsStateObject() : nullptr)
+	if (UObject* Obj = Processor->GetAsyncPhysicsStateObject())
 	{
 		Owner->AddRootObject(Obj);
 	}
@@ -437,12 +437,12 @@ void FPhysScene_AsyncPhysicsStateJobQueue::FJob::OnPostExecute_GameThread() cons
 			}
 			break;
 		}
-	}
 
-	// Remove rooted objects
-	if (UObject* Obj = Processor ? Processor->GetAsyncPhysicsStateObject() : nullptr)
-	{
-		Owner->RemoveRootObject(Obj);
+		// Remove rooted objects
+		if (UObject* Obj = Processor->GetAsyncPhysicsStateObject())
+		{
+			Owner->RemoveRootObject(Obj);
+		}
 	}
 
 	for (UObject* Obj : RootedInputs)
