Added some Insight tags in FWorldPartitionLevelHelper::MoveExternalActorsToLevel for various optional validations.

#Synced-CL 50904874

diff --git a/Engine/Source/Runtime/Engine/Private/WorldPartition/WorldPartitionLevelHelper.cpp b/Engine/Source/Runtime/Engine/Private/WorldPartition/WorldPartitionLevelHelper.cpp
--- a/Engine/Source/Runtime/Engine/Private/WorldPartition/WorldPartitionLevelHelper.cpp
+++ b/Engine/Source/Runtime/Engine/Private/WorldPartition/WorldPartitionLevelHelper.cpp
@@ -341,11 +341,15 @@ void FWorldPartitionLevelHelper::MoveExternalActorsToLevel(const TArray<FWorldPa
 
 	InLevel->AddLoadedActors(LoadedActors.Array());
 
-	for (AActor* Actor : InLevel->Actors)
 	{
-		if (IsValid(Actor) && Actor->HasAllFlags(RF_WasLoaded))
+		TRACE_CPUPROFILER_EVENT_SCOPE(ValidateLoaded);
+
+		for (AActor* Actor : InLevel->Actors)
 		{
-			checkf(LevelActors.Contains(Actor->GetFName()), TEXT("Actor %s(%s) was unexpectedly loaded when moving actors to streaming cell"), *Actor->GetActorNameOrLabel(), *Actor->GetName());
+			if (IsValid(Actor) && Actor->HasAllFlags(RF_WasLoaded))
+			{
+				checkf(LevelActors.Contains(Actor->GetFName()), TEXT("Actor %s(%s) was unexpectedly loaded when moving actors to streaming cell"), *Actor->GetActorNameOrLabel(), *Actor->GetName());
+			}
 		}
 	}
 
@@ -362,6 +366,8 @@ void FWorldPartitionLevelHelper::MoveExternalActorsToLevel(const TArray<FWorldPa
 	
 	// Perform validations and fixups to make sure the level is saveable
 	{
+		TRACE_CPUPROFILER_EVENT_SCOPE(ValidateMovedObjects);
+
 		TArray<UObject*> ObjectStack { InLevel };
 		TSet<UObject*> VisitedObjects { InLevel };
 
