World Partition Streaming Generation: make sure to generate a streaming generation log for dedicated servers.
Helps while debugging SEW issues.

#Synced-CL 50914233

diff --git a/Engine/Source/Runtime/Engine/Private/WorldPartition/WorldPartitionStreamingGeneration.cpp b/Engine/Source/Runtime/Engine/Private/WorldPartition/WorldPartitionStreamingGeneration.cpp
--- a/Engine/Source/Runtime/Engine/Private/WorldPartition/WorldPartitionStreamingGeneration.cpp
+++ b/Engine/Source/Runtime/Engine/Private/WorldPartition/WorldPartitionStreamingGeneration.cpp
@@ -2353,12 +2353,41 @@ bool UWorldPartition::GenerateContainerStreaming(const FGenerateStreamingParams&
 	TUniquePtr<FHierarchicalLogArchive> HierarchicalLogAr;
 
 	const bool bIsStreamingGenerationLogAllowed = GIsStreamingGenerationLogAllowed && (!bIsPIE || IWorldPartitionEditorModule::Get().GetEnableStreamingGenerationLogOnPIE()) && (!InParams.OutputLogType.IsSet() || !InParams.OutputLogType.GetValue().IsEmpty());
-	const bool bIsStreamingGenerationLogRelevant = IsMainWorldPartition() && (!GIsBuildMachine || GIsAutomationTesting || IsRunningCookCommandlet());
+	const bool bIsStreamingGenerationLogRelevant = !GIsBuildMachine || GIsAutomationTesting || IsRunningCookCommandlet();
 
 	if (bIsStreamingGenerationLogAllowed && bIsStreamingGenerationLogRelevant)
 	{
 		TStringBuilder<256> StateLogSuffix;
-		StateLogSuffix += bIsPIE ? TEXT("PIE") : (IsRunningGame() ? TEXT("Game") : (IsRunningCookCommandlet() ? TEXT("Cook") : (GIsAutomationTesting ? TEXT("UnitTest") : (InParams.OutputLogType.IsSet() ? *InParams.OutputLogType.GetValue() : TEXT("Manual")))));
+
+		if (bIsPIE)
+		{
+			StateLogSuffix += TEXT("PIE");
+		}
+		else if (IsRunningDedicatedServer())
+		{
+			StateLogSuffix += TEXT("Server");
+		}
+		else if (IsRunningGame())
+		{
+			StateLogSuffix += TEXT("Game");
+		}
+		else if (IsRunningCookCommandlet())
+		{
+			StateLogSuffix += TEXT("Cook");
+		}
+		else if (GIsAutomationTesting)
+		{
+			StateLogSuffix += TEXT("UnitTest");
+		}
+		else if (InParams.OutputLogType.IsSet())
+		{
+			StateLogSuffix += *InParams.OutputLogType.GetValue();
+		}
+		else
+		{
+			StateLogSuffix += TEXT("Manual");
+		}
+
 		StateLogSuffix += TEXT("_");
 		StateLogSuffix += ContainerShortName;
 		LogFileAr = FWorldPartitionStreamingGenerator::CreateDumpStateLogArchive(*StateLogSuffix);
