Editor: As with 50524767, allow cut operation to work on non-actors even when TypedElements.EnableElementsCopyAndPaste is false. Add ability to override cut action in customization so that we can issue the appropriate callback for the actor case (when using element copy/paste).

#Synced-CL 50868977

diff --git a/Engine/Source/Editor/LevelEditor/Private/Elements/Actor/ActorElementLevelEditorCommonActionsCustomization.cpp b/Engine/Source/Editor/LevelEditor/Private/Elements/Actor/ActorElementLevelEditorCommonActionsCustomization.cpp
--- a/Engine/Source/Editor/LevelEditor/Private/Elements/Actor/ActorElementLevelEditorCommonActionsCustomization.cpp
+++ b/Engine/Source/Editor/LevelEditor/Private/Elements/Actor/ActorElementLevelEditorCommonActionsCustomization.cpp
@@ -33,6 +33,26 @@ void FActorElementLevelEditorCommonActionsCustomization::CopyElements(ITypedElem
 	}
 }
 
+// We override this function because we want to issue OnEditCutActorsBegin/End instead of OnEditCopyActorsBegin, esp since
+//  the cut delegate gets fired from within a transaction, whereas the copy does not.
+void FActorElementLevelEditorCommonActionsCustomization::CutElements(ITypedElementWorldInterface* InWorldInterface, 
+	TArrayView<const FTypedElementHandle> InElementHandles, UWorld* InWorld, UTypedElementSelectionSet* InSelectionSet, 
+	const FTypedElementDeletionOptions& InDeletionOptions, FOutputDevice& Out)
+{
+	if (TypedElementCommonActionsUtils::IsElementCopyAndPasteEnabled())
+	{
+		FEditorDelegates::OnEditCutActorsBegin.Broadcast();
+
+		// Note that we qualify this call to be the base class version, not ours, so that we don't call the
+		//  copy callback.
+		FTypedElementCommonActionsCustomization::CopyElements(InWorldInterface, InElementHandles, Out);
+		
+		DeleteElements(InWorldInterface, InElementHandles, InWorld, InSelectionSet, InDeletionOptions);
+
+		FEditorDelegates::OnEditCutActorsEnd.Broadcast();
+	}
+}
+
 bool FActorElementLevelEditorCommonActionsCustomization::DeleteElements(ITypedElementWorldInterface* InWorldInterface, TArrayView<const FTypedElementHandle> InElementHandles, UWorld* InWorld, UTypedElementSelectionSet* InSelectionSet, const FTypedElementDeletionOptions& InDeletionOptions)
 {
 	if (const IToolkitHost* ToolkitHostPtr = GetToolkitHost())
diff --git a/Engine/Source/Editor/LevelEditor/Private/Elements/Actor/ActorElementLevelEditorCommonActionsCustomization.h b/Engine/Source/Editor/LevelEditor/Private/Elements/Actor/ActorElementLevelEditorCommonActionsCustomization.h
--- a/Engine/Source/Editor/LevelEditor/Private/Elements/Actor/ActorElementLevelEditorCommonActionsCustomization.h
+++ b/Engine/Source/Editor/LevelEditor/Private/Elements/Actor/ActorElementLevelEditorCommonActionsCustomization.h
@@ -12,6 +12,9 @@ public:
 	virtual bool IsCopyCapable() const override;
 	virtual bool CanCopyElements(ITypedElementWorldInterface* InWorldInterface, TArrayView<const FTypedElementHandle> InElementHandles) override;
 	virtual void CopyElements(ITypedElementWorldInterface* InWorldInterface, TArrayView<const FTypedElementHandle> InElementHandles, FOutputDevice& Out) override;
+	virtual void CutElements(ITypedElementWorldInterface* InWorldInterface, TArrayView<const FTypedElementHandle> InElementHandles,
+		UWorld* InWorld, UTypedElementSelectionSet* InSelectionSet, const FTypedElementDeletionOptions& InDeletionOptions,
+		FOutputDevice& Out) override;
 	virtual bool DeleteElements(ITypedElementWorldInterface* InWorldInterface, TArrayView<const FTypedElementHandle> InElementHandles, UWorld* InWorld, UTypedElementSelectionSet* InSelectionSet, const FTypedElementDeletionOptions& InDeletionOptions) override;
 	virtual void DuplicateElements(ITypedElementWorldInterface* InWorldInterface, TArrayView<const FTypedElementHandle> InElementHandles, UWorld* InWorld, const FVector& InLocationOffset, TArray<FTypedElementHandle>& OutNewElements) override;
 	virtual TSharedPtr<FWorldElementPasteImporter> GetPasteImporter(ITypedElementWorldInterface* InWorldInterface,
diff --git a/Engine/Source/Editor/UnrealEd/Private/UnrealEdSrv.cpp b/Engine/Source/Editor/UnrealEd/Private/UnrealEdSrv.cpp
--- a/Engine/Source/Editor/UnrealEd/Private/UnrealEdSrv.cpp
+++ b/Engine/Source/Editor/UnrealEd/Private/UnrealEdSrv.cpp
@@ -1488,6 +1488,8 @@ bool UUnrealEdEngine::DoDirtyPackagesNeedCheckout() const
 
 bool UUnrealEdEngine::Exec_Edit( UWorld* InWorld, const TCHAR* Str, FOutputDevice& Ar )
 {
+	const FText TransactionName = NSLOCTEXT("UnrealEd", "Cut", "Cut");
+
 	const bool bComponentsSelected = GetSelectedComponentCount() > 0;
 
 	if( FParse::Command(&Str,TEXT("CUT")) )
@@ -1497,46 +1499,40 @@ bool UUnrealEdEngine::Exec_Edit( UWorld* InWorld, const TCHAR* Str, FOutputDevic
 			return true;
 		}
 
-		if (TypedElementCommonActionsUtils::IsElementCopyAndPasteEnabled())
+		bool bExportedViaCommonActions = false;
+		if (GCurrentLevelEditingViewportClient)
 		{
-			if (GCurrentLevelEditingViewportClient)
+			if (TSharedPtr<ILevelEditor> LevelEditor = GCurrentLevelEditingViewportClient->ParentLevelEditor.Pin())
 			{
-				if (TSharedPtr<ILevelEditor> LevelEditor = GCurrentLevelEditingViewportClient->ParentLevelEditor.Pin())
+				if (UTypedElementCommonActions* CommonActions = LevelEditor->GetCommonActions())
 				{
-					if (UTypedElementCommonActions* CommonActions = LevelEditor->GetCommonActions())
-					{
-						// End drags to avoid deleting something out from under one.
-						FSlateApplication::Get().CancelDragDrop();
-
-						const FScopedTransaction Transaction(NSLOCTEXT("UnrealEd", "Cut", "Cut"));
-
-						UTypedElementSelectionSet* SelectionSet = LevelEditor->GetMutableElementSelectionSet();
-
-						if (!bComponentsSelected)
-						{
-							FEditorDelegates::OnEditCutActorsBegin.Broadcast();
-						}
+					// End drags to avoid deleting something out from under one.
+					FSlateApplication::Get().CancelDragDrop();
+					
+					const FScopedTransaction Transaction(TransactionName);
 
-						CommonActions->CopySelectedElements(SelectionSet);
-						const bool bCheckRef = GetDefault<ULevelEditorMiscSettings>()->bCheckReferencesOnDelete;
-						FTypedElementDeletionOptions Options;
-						Options
-							.SetWarnAboutReferences(bCheckRef)
-							.SetWarnAboutSoftReferences(bCheckRef);
-						CommonActions->DeleteSelectedElements(SelectionSet, InWorld, Options);
+					const bool bCheckRef = GetDefault<ULevelEditorMiscSettings>()->bCheckReferencesOnDelete;
+					FTypedElementDeletionOptions Options;
+					Options
+						.SetWarnAboutReferences(bCheckRef)
+						.SetWarnAboutSoftReferences(bCheckRef);
 
-						if (!bComponentsSelected)
-						{
-							FEditorDelegates::OnEditCutActorsEnd.Broadcast();
-						}
-					}
+					UTypedElementSelectionSet* SelectionSet = LevelEditor->GetMutableElementSelectionSet();
+					bExportedViaCommonActions = CommonActions->CutSelectedElements(SelectionSet, InWorld, Options);
 				}
 			}
 		}
+
+		// As with copy, bExportedViaCommonActions would be false when TypedElementCommonActionsUtils::IsElementCopyAndPasteEnabled()
+		//  returns false inside an actor/component common actions customization, falling into this legacy path.
+		if (bExportedViaCommonActions)
+		{
+			// Do nothing
+		}
 		else if (bComponentsSelected)
 		{
 			// Same transaction language used in CopySelectedActorsToClipboard below
-			const FScopedTransaction Transaction(NSLOCTEXT("UnrealEd", "Cut", "Cut"));
+			const FScopedTransaction Transaction(TransactionName);
 
 			edactCopySelected(InWorld);
 			const bool bCheckRef = GetDefault<ULevelEditorMiscSettings>()->bCheckReferencesOnDelete;
diff --git a/Engine/Source/Runtime/Engine/Private/Elements/Framework/TypedElementCommonActions.cpp b/Engine/Source/Runtime/Engine/Private/Elements/Framework/TypedElementCommonActions.cpp
--- a/Engine/Source/Runtime/Engine/Private/Elements/Framework/TypedElementCommonActions.cpp
+++ b/Engine/Source/Runtime/Engine/Private/Elements/Framework/TypedElementCommonActions.cpp
@@ -62,6 +62,14 @@ void FTypedElementCommonActionsCustomization::CopyElements(ITypedElementWorldInt
 	InWorldInterface->CopyElements(InElementHandles, Out);
 }
 
+void FTypedElementCommonActionsCustomization::CutElements(ITypedElementWorldInterface* InWorldInterface, TArrayView<const FTypedElementHandle> InElementHandles,
+	UWorld* InWorld, UTypedElementSelectionSet* InSelectionSet, const FTypedElementDeletionOptions& InDeletionOptions,
+	FOutputDevice& Out)
+{
+	CopyElements(InWorldInterface, InElementHandles, Out);
+	DeleteElements(InWorldInterface, InElementHandles, InWorld, InSelectionSet, InDeletionOptions);
+}
+
 TSharedPtr<FWorldElementPasteImporter> FTypedElementCommonActionsCustomization::GetPasteImporter(ITypedElementWorldInterface* InWorldInterface, const FTypedElementListConstPtr& InSelectedHandles, UWorld* InWorld)
 {
 	return InWorldInterface->GetPasteImporter();
@@ -211,6 +219,18 @@ bool UTypedElementCommonActions::CopySelectedElementsToString(UTypedElementSelec
 }
 
 bool UTypedElementCommonActions::CopyNormalizedElements(const FTypedElementListConstPtr& ElementListPtr, FString* OptionalOutputString)
+{
+	return CutOrCopyNormalizedElements(ElementListPtr, OptionalOutputString, []
+		(FTypedElementCommonActionsCustomization& CommonActionsCustomization, ITypedElementWorldInterface& WorldInterface, 
+			const TArray<FTypedElementHandle>& ElementsIn, FStringOutputDevice& Output)
+		{
+			CommonActionsCustomization.CopyElements(&WorldInterface, ElementsIn, Output);
+		});
+}
+
+bool UTypedElementCommonActions::CutOrCopyNormalizedElements(const FTypedElementListConstPtr& ElementListPtr, FString* OptionalOutputString,
+	TFunctionRef<void(FTypedElementCommonActionsCustomization& CommonActionsCustomization, ITypedElementWorldInterface& WorldInterface, 
+		const TArray<FTypedElementHandle>& ElementsIn, FStringOutputDevice& Output)> CutOrCopy)
 {
 	if (!ElementListPtr)
 	{
@@ -232,7 +252,7 @@ bool UTypedElementCommonActions::CopyNormalizedElements(const FTypedElementListC
 		ITypedElementWorldInterface* WorldInterface = Registry->GetElementInterface<ITypedElementWorldInterface>(ElementsByTypePair.Key);
 		if (CommonActionsCustomization && WorldInterface)
 		{
-			CommonActionsCustomization->CopyElements(WorldInterface, ElementsByTypePair.Value, ElementsOut);
+			CutOrCopy(*CommonActionsCustomization, *WorldInterface, ElementsByTypePair.Value, ElementsOut);
 		}
 		if (ElementsOut.IsEmpty())
 		{
@@ -259,6 +279,28 @@ bool UTypedElementCommonActions::CopyNormalizedElements(const FTypedElementListC
 	return bHasExportedSomething;
 }
 
+bool UTypedElementCommonActions::CutSelectedElements(UTypedElementSelectionSet* SelectionSet, UWorld* World, const FTypedElementDeletionOptions& DeletionOptions)
+{
+	if (!SelectionSet || !World)
+	{
+		return false;
+	}
+
+	FTypedElementListRef NormalizedElements = SelectionSet->GetNormalizedSelection(FTypedElementSelectionNormalizationOptions());
+	return CutNormalizedElements(NormalizedElements, World, SelectionSet, DeletionOptions);
+}
+
+bool UTypedElementCommonActions::CutNormalizedElements(const FTypedElementListConstPtr& ElementListPtr, UWorld* World, 
+	UTypedElementSelectionSet* InSelectionSet, const FTypedElementDeletionOptions& DeletionOptions)
+{
+	return CutOrCopyNormalizedElements(ElementListPtr, nullptr, [World, InSelectionSet, &DeletionOptions]
+		(FTypedElementCommonActionsCustomization& CommonActionsCustomization, ITypedElementWorldInterface& WorldInterface, 
+			const TArray<FTypedElementHandle>& ElementsIn, FStringOutputDevice& Output)
+		{
+			CommonActionsCustomization.CutElements(&WorldInterface, ElementsIn, World, InSelectionSet, DeletionOptions, Output);
+		});
+}
+
 TArray<TPair<TSharedRef<FWorldElementPasteImporter>, FStringView>> UTypedElementCommonActions::ParsePasteTextAndCreateImporters(
 	const FTypedElementPasteOptions& PasteOptions, 
 	const FString& InputString, 
diff --git a/Engine/Source/Runtime/Engine/Public/Elements/Framework/TypedElementCommonActions.h b/Engine/Source/Runtime/Engine/Public/Elements/Framework/TypedElementCommonActions.h
--- a/Engine/Source/Runtime/Engine/Public/Elements/Framework/TypedElementCommonActions.h
+++ b/Engine/Source/Runtime/Engine/Public/Elements/Framework/TypedElementCommonActions.h
@@ -11,6 +11,7 @@
 
 #include "TypedElementCommonActions.generated.h"
 
+class FStringOutputDevice;
 class UTypedElementSelectionSet;
 
 /**
@@ -28,7 +29,11 @@ public:
 	ENGINE_API virtual bool IsCopyCapable() const;
 	ENGINE_API virtual bool CanCopyElements(ITypedElementWorldInterface* InWorldInterface, TArrayView<const FTypedElementHandle> InElementHandles);
 	ENGINE_API virtual void CopyElements(ITypedElementWorldInterface* InWorldInterface,TArrayView<const FTypedElementHandle> InElementHandles, FOutputDevice& Out);
-	
+	// By default this just calls CopyElements followed by DeleteElements
+	ENGINE_API virtual void CutElements(ITypedElementWorldInterface* InWorldInterface, TArrayView<const FTypedElementHandle> InElementHandles,
+		UWorld* InWorld, UTypedElementSelectionSet* InSelectionSet, const FTypedElementDeletionOptions& InDeletionOptions,
+		FOutputDevice& Out);
+
 	ENGINE_API virtual TSharedPtr<FWorldElementPasteImporter> GetPasteImporter(ITypedElementWorldInterface* InWorldInterface, const FTypedElementListConstPtr& InSelectedHandles, UWorld* InWorld);
 };
 
@@ -168,6 +173,13 @@ public:
 	 */
 	ENGINE_API virtual bool CopyNormalizedElements(const FTypedElementListConstPtr& ElementListPtr, FString* OptionalOutputString = nullptr);
 
+	/**
+	 * Cut any elements from the given selection set. This typically amounts to copying any elements that can be copied
+	 *  into the clipboard and then deleting them, but it does go through a customization endpoint that may adjust the
+	 *  behavior (e.g. for actors, this fires a different delegate than copy/delete would).
+	 */
+	ENGINE_API bool CutSelectedElements(UTypedElementSelectionSet* SelectionSet, UWorld* World, const FTypedElementDeletionOptions& DeletionOptions);
+
 	/**
 	 * Given an input string (or use clipboard if null string), determine if it is possible to use the string to create elements.
 	 * Editor and type specific customizations registered to the SelectionSet are given a chance to use auxillary state to make this decision.
@@ -269,6 +281,15 @@ private:
 	 * Attempt to resolve the selection interface and common actions customization for the given element, if any.
 	 */
 	FTypedElementCommonActionsElement ResolveCommonActionsElement(const FTypedElementHandle& InElementHandle) const;
+
+	// Helper to share code between copy and cut. The cut/copy operation is passed in as a lambda
+	bool CutOrCopyNormalizedElements(const FTypedElementListConstPtr& ElementListPtr, FString* OptionalOutputString,
+		TFunctionRef<void(FTypedElementCommonActionsCustomization& CommonActionsCustomization, 
+			ITypedElementWorldInterface& WorldInterface, const TArray<FTypedElementHandle>& ElementsIn, FStringOutputDevice& Output)> CutOrCopy);
+
+	// We might expose this as we do for copy/delete, but we can wait to do so until needed.
+	bool CutNormalizedElements(const FTypedElementListConstPtr& ElementListPtr, 
+		UWorld* World, UTypedElementSelectionSet* InSelectionSet, const FTypedElementDeletionOptions& DeletionOptions);
 };
 
 namespace TypedElementCommonActionsUtils
