Add GetPostProcessVolumes() function and deprecate the PostProcessVolumes field, to allow building acceleration structures or sorting on demand in future versions

#Synced-CL 50811210

diff --git a/Engine/Plugins/MovieScene/MovieRenderPipeline/Source/MovieRenderPipelineCore/Private/Graph/MovieGraphCoreTimeStep.cpp b/Engine/Plugins/MovieScene/MovieRenderPipeline/Source/MovieRenderPipelineCore/Private/Graph/MovieGraphCoreTimeStep.cpp
--- a/Engine/Plugins/MovieScene/MovieRenderPipeline/Source/MovieRenderPipelineCore/Private/Graph/MovieGraphCoreTimeStep.cpp
+++ b/Engine/Plugins/MovieScene/MovieRenderPipeline/Source/MovieRenderPipelineCore/Private/Graph/MovieGraphCoreTimeStep.cpp
@@ -763,7 +763,7 @@ float UMovieGraphCoreTimeStep::GetBlendedMotionBlurAmount()
 	{
 		// Apply any motion blur settings from post process volumes in the world
 		FVector ViewLocation = PlayerCameraManager->GetCameraLocation();
-		for (IInterface_PostProcessVolume* PPVolume : GetWorld()->PostProcessVolumes)
+		for (IInterface_PostProcessVolume* PPVolume : GetWorld()->GetPostProcessVolumes())
 		{
 			const FPostProcessVolumeProperties VolumeProperties = PPVolume->GetProperties();
 
diff --git a/Engine/Plugins/MovieScene/MovieRenderPipeline/Source/MovieRenderPipelineCore/Private/MoviePipeline.cpp b/Engine/Plugins/MovieScene/MovieRenderPipeline/Source/MovieRenderPipelineCore/Private/MoviePipeline.cpp
--- a/Engine/Plugins/MovieScene/MovieRenderPipeline/Source/MovieRenderPipelineCore/Private/MoviePipeline.cpp
+++ b/Engine/Plugins/MovieScene/MovieRenderPipeline/Source/MovieRenderPipelineCore/Private/MoviePipeline.cpp
@@ -1405,7 +1405,7 @@ MoviePipeline::FFrameConstantMetrics UMoviePipeline::CalculateShotFrameMetrics(c
 	{
 		// Apply any motion blur settings from post process volumes in the world
 		FVector ViewLocation = PlayerCameraManager->GetCameraLocation();
-		for (IInterface_PostProcessVolume* PPVolume : GetWorld()->PostProcessVolumes)
+		for (IInterface_PostProcessVolume* PPVolume : GetWorld()->GetPostProcessVolumes())
 		{
 			const FPostProcessVolumeProperties VolumeProperties = PPVolume->GetProperties();
 
diff --git a/Engine/Plugins/MovieScene/MovieRenderPipeline/Source/MovieRenderPipelineCore/Private/MoviePipelineUtils.cpp b/Engine/Plugins/MovieScene/MovieRenderPipeline/Source/MovieRenderPipelineCore/Private/MoviePipelineUtils.cpp
--- a/Engine/Plugins/MovieScene/MovieRenderPipeline/Source/MovieRenderPipelineCore/Private/MoviePipelineUtils.cpp
+++ b/Engine/Plugins/MovieScene/MovieRenderPipeline/Source/MovieRenderPipelineCore/Private/MoviePipelineUtils.cpp
@@ -1260,7 +1260,7 @@ namespace UE
 
 		void DoPostProcessBlend(const FVector& InViewLocation, const UWorld* InWorld, const FMinimalViewInfo& InViewInfo, FSceneView* InOutView)
 		{
-			for (IInterface_PostProcessVolume* PPVolume : InWorld->PostProcessVolumes)
+			for (IInterface_PostProcessVolume* PPVolume : InWorld->GetPostProcessVolumes())
 			{
 				const FPostProcessVolumeProperties VolumeProperties = PPVolume->GetProperties();
 
diff --git a/Engine/Plugins/Runtime/nDisplay/Source/DisplayClusterMoviePipeline/Private/DisplayClusterMoviePipelineViewportPass.cpp b/Engine/Plugins/Runtime/nDisplay/Source/DisplayClusterMoviePipeline/Private/DisplayClusterMoviePipelineViewportPass.cpp
--- a/Engine/Plugins/Runtime/nDisplay/Source/DisplayClusterMoviePipeline/Private/DisplayClusterMoviePipelineViewportPass.cpp
+++ b/Engine/Plugins/Runtime/nDisplay/Source/DisplayClusterMoviePipeline/Private/DisplayClusterMoviePipelineViewportPass.cpp
@@ -175,7 +175,7 @@ void UDisplayClusterMoviePipelineViewportPassBase::BlendPostProcessSettings(FSce
 			// Use DCRA as a camera
 			FVector ViewLocation = DCViews[InViewportIndex].ViewLocation;
 
-			for (IInterface_PostProcessVolume* PPVolume : GetWorld()->PostProcessVolumes)
+			for (IInterface_PostProcessVolume* PPVolume : GetWorld()->GetPostProcessVolumes())
 			{
 				const FPostProcessVolumeProperties VolumeProperties = PPVolume->GetProperties();
 
diff --git a/Engine/Source/Runtime/Engine/Classes/Engine/World.h b/Engine/Source/Runtime/Engine/Classes/Engine/World.h
--- a/Engine/Source/Runtime/Engine/Classes/Engine/World.h
+++ b/Engine/Source/Runtime/Engine/Classes/Engine/World.h
@@ -1908,12 +1908,20 @@ public:
 	/** Removes a post process volume from the world.  Returns true if it was successfully removed. */
 	UE_API bool RemovePostProcessVolume(IInterface_PostProcessVolume* InVolume);
 
+	/** Return an array of post processing volumes, sorted in ascending order of priority. */
+	const TArray<IInterface_PostProcessVolume*>& GetPostProcessVolumes() const
+	{
+PRAGMA_DISABLE_DEPRECATION_WARNINGS
+		return PostProcessVolumes;
+PRAGMA_ENABLE_DEPRECATION_WARNINGS
+	}
+
+	UE_DEPRECATED(5.8, "Use GetPostProcessVolumes() instead")
+	TArray<IInterface_PostProcessVolume*> PostProcessVolumes;
+
 	/** Called when a scene view for this world needs the worlds post process settings computed */
 	UE_API void AddPostProcessingSettings(FVector ViewLocation, FSceneView* SceneView);
 
-	/** An array of post processing volumes, sorted in ascending order of priority.					*/
-	TArray< IInterface_PostProcessVolume * > PostProcessVolumes;
-
 private:
 	/** Caches the size of each volume at the time of insertion, for sort determinism in case it later changes.  One for one with PostProcessVolumes above. */
 	TArray<double> PostProcessVolumeCachedSizes;
diff --git a/Engine/Source/Runtime/Engine/Private/World.cpp b/Engine/Source/Runtime/Engine/Private/World.cpp
--- a/Engine/Source/Runtime/Engine/Private/World.cpp
+++ b/Engine/Source/Runtime/Engine/Private/World.cpp
@@ -877,6 +877,7 @@ UWorld::UWorld( const FObjectInitializer& ObjectInitializer )
 	BlockTillLevelStreamingCompletedEpoch = 0;
 }
 
+PRAGMA_DISABLE_DEPRECATION_WARNINGS
 UWorld::~UWorld()
 {
 #if WITH_STATE_STREAM
@@ -888,6 +889,7 @@ UWorld::~UWorld()
 		delete PerfTrackers;
 	}
 }
+PRAGMA_ENABLE_DEPRECATION_WARNINGS
 
 void UWorld::Serialize( FArchive& Ar )
 {
@@ -10388,6 +10390,7 @@ static bool PostProcessVolumeLess(const FPostProcessVolumeProperties& Properties
 
 void UWorld::InsertPostProcessVolume(IInterface_PostProcessVolume* InVolume)
 {
+PRAGMA_DISABLE_DEPRECATION_WARNINGS
 	const int32 NumVolumes = PostProcessVolumes.Num();
 	FPostProcessVolumeProperties TargetProperties = InVolume->GetProperties();
 	int32 InsertIndex = 0;
@@ -10420,10 +10423,12 @@ void UWorld::InsertPostProcessVolume(IInterface_PostProcessVolume* InVolume)
 	}
 	PostProcessVolumes.Insert(InVolume, InsertIndex);
 	PostProcessVolumeCachedSizes.Insert(TargetProperties.Size, InsertIndex);
+PRAGMA_ENABLE_DEPRECATION_WARNINGS
 }
 
 bool UWorld::RemovePostProcessVolume(IInterface_PostProcessVolume* InVolume)
 {
+PRAGMA_DISABLE_DEPRECATION_WARNINGS
 	int32 RemoveIndex;
 	if (PostProcessVolumes.Find(InVolume, RemoveIndex))
 	{
@@ -10432,6 +10437,7 @@ bool UWorld::RemovePostProcessVolume(IInterface_PostProcessVolume* InVolume)
 		return true;
 	}
 	return false;
+PRAGMA_ENABLE_DEPRECATION_WARNINGS
 }
 
 void UWorld::InitializeSubsystems()
@@ -10606,6 +10612,8 @@ static void DoPostProcessVolume(IInterface_PostProcessVolume* Volume, FVector Vi
 
 void UWorld::AddPostProcessingSettings(FVector ViewLocation, FSceneView* SceneView)
 {
+	SCOPED_NAMED_EVENT(UWorld_AddPostProcessingSettings, FColor::Red);
+
 	OnBeginPostProcessSettings.Broadcast(ViewLocation, SceneView);
 
 #if DEBUG_POST_PROCESS_VOLUME_ENABLE
@@ -10613,7 +10621,7 @@ void UWorld::AddPostProcessingSettings(FVector ViewLocation, FSceneView* SceneVi
 	SceneView->FinalPostProcessDebugInfo.Reset();
 #endif
 
-	for (IInterface_PostProcessVolume* PPVolume : PostProcessVolumes)
+	for (IInterface_PostProcessVolume* PPVolume : GetPostProcessVolumes())
 	{
 		DoPostProcessVolume(PPVolume, ViewLocation, SceneView);
 	}
