[HWRT] Small code cleanup.

#Synced-CL 50994078

diff --git a/Engine/Source/Runtime/Engine/Private/Rendering/RayTracingGeometryManager.cpp b/Engine/Source/Runtime/Engine/Private/Rendering/RayTracingGeometryManager.cpp
--- a/Engine/Source/Runtime/Engine/Private/Rendering/RayTracingGeometryManager.cpp
+++ b/Engine/Source/Runtime/Engine/Private/Rendering/RayTracingGeometryManager.cpp
@@ -610,7 +610,7 @@ void FRayTracingGeometryManager::RefreshRegisteredGeometry(RayTracing::FGeometry
 
 		if (GeometryInitializer.Type == ERayTracingGeometryInitializerType::StreamingDestination)
 		{
-			RegisteredGeometry.Status = FRegisteredGeometry::FStatus::StreamedOut;
+			RegisteredGeometry.Status = FRegisteredGeometry::EStatus::StreamedOut;
 		}
 	}
 }
@@ -1086,7 +1086,7 @@ bool FRayTracingGeometryManager::RequestRayTracingGeometryStreamIn(FRHICommandLi
 	FRegisteredGeometry& RegisteredGeometry = GetRegisteredGeometry(GeometryHandle);
 	FRayTracingGeometry* Geometry = RegisteredGeometry.Geometry;
 
-	if (RegisteredGeometry.Status == FRegisteredGeometry::FStatus::Streaming)
+	if (RegisteredGeometry.Status == FRegisteredGeometry::EStatus::Streaming)
 	{
 		// skip if there's already a streaming request in-flight for this geometry
 		return true;
@@ -1118,7 +1118,7 @@ bool FRayTracingGeometryManager::RequestRayTracingGeometryStreamIn(FRHICommandLi
 	{
 		// no offline data -> build from VB/IB at runtime
 
-		RegisteredGeometry.Status = FRegisteredGeometry::FStatus::StreamedIn;
+		RegisteredGeometry.Status = FRegisteredGeometry::EStatus::StreamedIn;
 	}
 	else
 	{
@@ -1157,12 +1157,12 @@ bool FRayTracingGeometryManager::RequestRayTracingGeometryStreamIn(FRHICommandLi
 
 		RegisteredGeometry.StreamableAsset->IssueRequest(StreamingRequest.Request, StreamingRequest.RequestBuffer, StreamingRequest.bBuffersOnly);
 
-		RegisteredGeometry.Status = FRegisteredGeometry::FStatus::Streaming;
+		RegisteredGeometry.Status = FRegisteredGeometry::EStatus::Streaming;
 
 		TotalStreamingSize += StreamingRequest.GeometrySize;
 	}
 
-	if (RegisteredGeometry.Status == FRegisteredGeometry::FStatus::StreamedIn)
+	if (RegisteredGeometry.Status == FRegisteredGeometry::EStatus::StreamedIn)
 	{
 		{
 			FRHIResourceReplaceBatcher Batcher(RHICmdList, 1);
@@ -1255,7 +1255,7 @@ void FRayTracingGeometryManager::ProcessCompletedStreamingRequests(FRHICommandLi
 					}
 				}
 
-				RegisteredGeometry.Status = FRegisteredGeometry::FStatus::StreamedIn;
+				RegisteredGeometry.Status = FRegisteredGeometry::EStatus::StreamedIn;
 			}
 
 			PendingRequest.Reset();
diff --git a/Engine/Source/Runtime/Engine/Public/Rendering/RayTracingGeometryManager.h b/Engine/Source/Runtime/Engine/Public/Rendering/RayTracingGeometryManager.h
--- a/Engine/Source/Runtime/Engine/Public/Rendering/RayTracingGeometryManager.h
+++ b/Engine/Source/Runtime/Engine/Public/Rendering/RayTracingGeometryManager.h
@@ -121,14 +121,14 @@ private:
 
 		int16 StreamingRequestIndex = INDEX_NONE;
 
-		enum class FStatus : uint8
+		enum class EStatus : uint8
 		{
 			StreamedOut,
 			Streaming,
 			StreamedIn,
 		};
 
-		FStatus Status = FStatus::StreamedOut;
+		EStatus Status = EStatus::StreamedOut;
 		bool bEvicted : 1 = false;
 		bool bAlwaysResident : 1 = false;
 	};
diff --git a/Engine/Source/Runtime/RenderCore/Private/RayTracingGeometry.cpp b/Engine/Source/Runtime/RenderCore/Private/RayTracingGeometry.cpp
--- a/Engine/Source/Runtime/RenderCore/Private/RayTracingGeometry.cpp
+++ b/Engine/Source/Runtime/RenderCore/Private/RayTracingGeometry.cpp
@@ -26,7 +26,7 @@ bool RayTracing::ShouldForceRuntimeBLAS()
 	return bDebugForceRuntimeBLAS;
 }
 
-PRAGMA_DISABLE_DEPRECATION_WARNINGS // Remove this in UE 5.6 once RayTracingGeometryRHI is made private
+PRAGMA_DISABLE_DEPRECATION_WARNINGS // Remove this in UE 5.10 once Initializer is made private
 
 #if RHI_RAYTRACING
 
diff --git a/Engine/Source/Runtime/Renderer/Public/RayTracingDynamicGeometryCollection.h b/Engine/Source/Runtime/Renderer/Public/RayTracingDynamicGeometryCollection.h
deleted file mode 100644
