Some Renderer API cleanup
* Removed a few previously deprecated functions.
* Deprecated & Removed implementation of IRendererModule:: / FSceneInterface::UpdateStaticDrawLists as it is no longer used.
* Removed FSceneInterface::AddLumenSceneCard etc that were never used.
* Deprecated FSceneInterface::RequiresHitProxies, it is only used internally in the renderer.
* Deprecated IRendererModule::GetNumDynamicLightsAffectingPrimitive - not used anywhere.

#Synced-CL 51006982

diff --git a/Engine/Source/Runtime/Engine/Private/SceneInterface.cpp b/Engine/Source/Runtime/Engine/Private/SceneInterface.cpp
--- a/Engine/Source/Runtime/Engine/Private/SceneInterface.cpp
+++ b/Engine/Source/Runtime/Engine/Private/SceneInterface.cpp
@@ -31,11 +31,6 @@ void FSceneInterface::UpdateAllPrimitiveSceneInfos(FRHICommandListImmediate& RHI
 	GraphBuilder.Execute();
 }
 
-void FSceneInterface::ProcessAndRenderIlluminanceMeter(FRDGBuilder& GraphBuilder, TArrayView<FViewInfo> Views, FRDGTextureRef SceneColorTexture)
-{
-	// Deprecated
-}
-
 EShaderPlatform FSceneInterface::GetShaderPlatform() const
 {
 	return GShaderPlatformForFeatureLevel[GetFeatureLevel()];
diff --git a/Engine/Source/Runtime/Engine/Public/SceneInterface.h b/Engine/Source/Runtime/Engine/Public/SceneInterface.h
--- a/Engine/Source/Runtime/Engine/Public/SceneInterface.h
+++ b/Engine/Source/Runtime/Engine/Public/SceneInterface.h
@@ -273,27 +273,15 @@ public:
 	  */
 	virtual void ResetReflectionCaptures(bool bOnlyIfOOM) {}
 
-	UE_DEPRECATED(5.1, "This method now accepts a bInsideTick parameter, which specifies whether it's called during the frame Tick.")
-	inline void AllocateReflectionCaptures(const TArray<UReflectionCaptureComponent*>& NewCaptures, const TCHAR* CaptureReason, bool bVerifyOnlyCapturing, bool bCapturingForMobile)
-	{
-		AllocateReflectionCaptures(NewCaptures, CaptureReason, bVerifyOnlyCapturing, bCapturingForMobile, false);
-	}
-
 	/** 
 	 * Updates the contents of the given sky capture by rendering the scene. 
 	 * This must be called on the game thread.
 	 */
 	virtual void UpdateSkyCaptureContents(const USkyLightComponent* CaptureComponent, bool bCaptureEmissiveOnly, UTextureCube* SourceCubemap, FTexture* OutProcessedTexture, float& OutAverageBrightness, FSHVectorRGB3& OutIrradianceEnvironmentMap, TArray<FFloat16Color>* OutRadianceMap, FLinearColor* SpecifiedCubemapColorScale) {}
 
-	UE_DEPRECATED(5.2, "AllocateAndCaptureFrameSkyEnvMap now takes an FRDGExternalAccessQueue")
-	virtual void AllocateAndCaptureFrameSkyEnvMap(FRDGBuilder& GraphBuilder, FSceneRenderer& SceneRenderer, FViewInfo& MainView, bool bShouldRenderSkyAtmosphere, bool bShouldRenderVolumetricCloud, FInstanceCullingManager& InstanceCullingManager) {}
-
 	virtual void AllocateAndCaptureFrameSkyEnvMap(FRDGBuilder& GraphBuilder, FSceneRenderer& SceneRenderer, FViewInfo& MainView, bool bShouldRenderSkyAtmosphere, bool bShouldRenderVolumetricCloud, FInstanceCullingManager& InstanceCullingManager, FRDGExternalAccessQueue& ExternalAccessQueue) {}
 	virtual void ValidateSkyLightRealTimeCapture(FRDGBuilder& GraphBuilder, const FViewInfo& View, FRDGTextureRef SceneColorTexture) {}
 
-	UE_DEPRECATED(5.4, "This method has been refactored to be a proper visualization post process. It should have never been added on the FSceneInterface in the first place.")
-	ENGINE_API virtual void ProcessAndRenderIlluminanceMeter(FRDGBuilder& GraphBuilder, TArrayView<FViewInfo> Views, FRDGTextureRef SceneColorTexture);
-
 	virtual void AddPlanarReflection(class UPlanarReflectionComponent* Component) {}
 	virtual void RemovePlanarReflection(class UPlanarReflectionComponent* Component) {}
 	virtual void UpdatePlanarReflectionTransform(UPlanarReflectionComponent* Component) {}
@@ -328,8 +316,6 @@ public:
 	 *  there are more pages being updated than can be handled in a given frame (when throttling) 
 	 */
 	virtual void InvalidateRuntimeVirtualTexture(URuntimeVirtualTextureComponent* Component, FBoxSphereBounds const& WorldBounds, EVTInvalidatePriority InvalidatePriority) {}
-	UE_DEPRECATED(5.6, "Use the alternate prototype with the Priority param")
-	virtual void InvalidateRuntimeVirtualTexture(URuntimeVirtualTextureComponent* Component, FBoxSphereBounds const& WorldBounds) {}
 
 	/** Request preload of an area of a runtime virtual texture at a given mip level. */
 	virtual void RequestPreloadRuntimeVirtualTexture(class URuntimeVirtualTextureComponent* Component, FBoxSphereBounds const& WorldBounds, int32 Level) {}
@@ -366,7 +352,7 @@ public:
 	/** Sets the precomputed volume distance field for the scene, or NULL to clear the current one. */
 	virtual void SetPrecomputedVolumeDistanceField(const class FPrecomputedVolumeDistanceField* PrecomputedVolumeDistanceField) {}
 
-	/** Updates all static draw lists. */
+	UE_DEPRECATED(5.8, "No longer implemented. For similar use-cases refer to IRendererModule::UpdateStaticDrawListsForMaterials.")
 	virtual void UpdateStaticDrawLists() {}
 
 	/** Update render states that possibly cached inside renderer, like mesh draw commands. More lightweight than re-registering the scene proxy. */
@@ -502,9 +488,6 @@ public:
 	 */
 	virtual void UpdatePhysicsField(FRDGBuilder& GraphBuilder, FViewInfo& View) = 0;
 
-	UE_DEPRECATED(5.0, "This method has been refactored to use an FRDGBuilder instead.")
-	virtual void UpdatePhysicsField(FRHICommandListImmediate& RHICmdList, FViewInfo& View) {}
-
 	/**
 	 * Returns the scene's unique info if it exists
 	 */
@@ -562,10 +545,6 @@ public:
 	 */
 	virtual FRHIUniformBuffer* GetSpeedTreeUniformBuffer(const FVertexFactory* VertexFactory) const = 0;
 
-	virtual void AddLumenSceneCard(class ULumenSceneCardComponent* LumenSceneCardComponent) {};
-	virtual void UpdateLumenSceneCardTransform(class ULumenSceneCardComponent* LumenSceneCardComponent) {};
-	virtual void RemoveLumenSceneCard(class ULumenSceneCardComponent* LumenSceneCardComponent) {};
-
 	// FPrimtiveDesc version for primitive/light scene interactions
 	virtual void AddPrimitive(FPrimitiveSceneDesc* Primitive) = 0;
 	virtual void RemovePrimitive(FPrimitiveSceneDesc* Primitive) = 0;
@@ -585,9 +564,6 @@ public:
 	UE_EXPERIMENTAL(5.6, "Experimental API, scene writing delegate may be internalized in the future.")
 	virtual void UpdatePrimitiveInstancesFromCompute(FPrimitiveSceneDesc* Primitive, FGPUSceneWriteDelegate&& DataWriterGPU) = 0;
 
-	UE_DEPRECATED(5.6, "StaticMeshBounds is no longer a parameter of this interface, please update FPrimitiveSceneDesc version.")
-	virtual void UpdatePrimitiveInstances(FInstancedStaticMeshSceneDesc* Primitive) = 0;	
-
 	UE_DEPRECATED(5.8, "Use UpdatePrimitivesIsDrawn_RenderThread instead.")
 	virtual void UpdatePrimitivesDrawnInGame_RenderThread(TArrayView<FPrimitiveSceneProxy*> PrimitiveSceneProxies, bool bDrawnInGame) final {}
 
@@ -616,6 +592,7 @@ public:
 	 *
 	 * @return true if hit proxies should be rendered in this scene.
 	 */
+	UE_DEPRECATED(5.8, "No longer exposed, only used internally.")
 	virtual bool RequiresHitProxies() const = 0;
 	/**
 	 * Get the optional UWorld that is associated with this scene
@@ -665,12 +642,6 @@ public:
 	virtual IPrimitiveTransformUpdater* CreatePrimitiveTransformUpdater(int32 MaxPrimitives) { return nullptr; }
 	virtual void UpdatePrimitiveTransforms(IPrimitiveTransformUpdater* Updater) {}
 
-	UE_DEPRECATED(5.6, "Use IPrimitiveTransformUpdater instead")
-	virtual void StartUpdatePrimitiveTransform(int32 NumPrimitives) {};
-	
-	UE_DEPRECATED(5.6, "Use IPrimitiveTransformUpdater instead")
-	virtual void FinishUpdatePrimitiveTransform() {};
-
 	/**
 	* Gets the GPU Skin Cache system associated with the scene.
 	*/
diff --git a/Engine/Source/Runtime/RenderCore/Public/RendererInterface.h b/Engine/Source/Runtime/RenderCore/Public/RendererInterface.h
--- a/Engine/Source/Runtime/RenderCore/Public/RendererInterface.h
+++ b/Engine/Source/Runtime/RenderCore/Public/RendererInterface.h
@@ -713,9 +713,7 @@ public:
 	virtual void DebugLogRenderProxies(FName Type, uint32 Num) = 0;
 #endif
 
-	/**
-	* Updates all static draw lists for each allocated scene.
-	*/
+	UE_DEPRECATED(5.8, "No longer implemented. For similar use-cases refer to UpdateStaticDrawListsForMaterials.")
 	virtual void UpdateStaticDrawLists() = 0;
 
 	/**
@@ -728,6 +726,7 @@ public:
 	virtual class FSceneViewStateInterface* AllocateViewState(ERHIFeatureLevel::Type FeatureLevel, FSceneViewStateInterface* ShareOriginTarget) = 0;
 
 	/** @return The number of lights that affect a primitive. */
+	UE_DEPRECATED(5.8, "No longer supported.")
 	virtual uint32 GetNumDynamicLightsAffectingPrimitive(const class FPrimitiveSceneInfo* PrimitiveSceneInfo,const class FLightCacheInterface* LCI) = 0;
 
 	virtual void OnWorldCleanup(UWorld* World, bool bSessionEnded, bool bCleanupResources, bool bWorldChanged) = 0;
diff --git a/Engine/Source/Runtime/Renderer/Private/RendererScene.cpp b/Engine/Source/Runtime/Renderer/Private/RendererScene.cpp
--- a/Engine/Source/Runtime/Renderer/Private/RendererScene.cpp
+++ b/Engine/Source/Runtime/Renderer/Private/RendererScene.cpp
@@ -1889,11 +1889,6 @@ void FScene::UpdatePrimitivesIsDrawn_RenderThread(TArrayView<FPrimitiveSceneProx
 	}
 }
 
-void FScene::UpdatePrimitiveInstances(FInstancedStaticMeshSceneDesc* Primitive)
-{
-	UpdatePrimitiveInstances(*Primitive);
-}
-
 void FScene::UpdatePrimitiveInstances(FPrimitiveSceneDesc* Primitive)
 {
 	SCOPE_CYCLE_COUNTER(STAT_UpdatePrimitiveInstanceGT);
@@ -4210,34 +4205,6 @@ void FScene::SetPrecomputedVisibility(const FPrecomputedVisibilityHandler* NewPr
 		});
 }
 
-void FScene::UpdateStaticDrawLists_RenderThread(FRHICommandListImmediate& RHICmdList)
-{
-	SCOPE_CYCLE_COUNTER(STAT_Scene_UpdateStaticDrawLists_RT);
-
-	const int32 NumPrimitives = Primitives.Num();
-
-	for (int32 PrimitiveIndex = 0; PrimitiveIndex < NumPrimitives; ++PrimitiveIndex)
-	{
-		FPrimitiveSceneInfo* Primitive = Primitives[PrimitiveIndex];
-
-		Primitive->RemoveStaticMeshes();
-	}
-
-	FPrimitiveSceneInfo::AddStaticMeshes(RHICmdList, this, Primitives);
-}
-
-void FScene::UpdateStaticDrawLists()
-{
-	UE::RenderCommandPipe::FSyncScope SyncScope;
-
-	FScene* Scene = this;
-	ENQUEUE_RENDER_COMMAND(FUpdateDrawLists)(
-		[Scene](FRHICommandListImmediate& RHICmdList)
-		{
-			Scene->UpdateStaticDrawLists_RenderThread(RHICmdList);
-		});
-}
-
 void FScene::UpdateCachedRenderStates(FPrimitiveSceneProxy* SceneProxy)
 {
 	if (SceneProxy->GetPrimitiveSceneInfo())
@@ -6965,7 +6932,6 @@ public:
 
 	virtual void UpdateCustomPrimitiveData(FPrimitiveSceneDesc* Primitive, const FCustomPrimitiveData&) override {}
 	virtual void UpdatePrimitiveInstances(FPrimitiveSceneDesc* Primitive) override {};
-	virtual void UpdatePrimitiveInstances(FInstancedStaticMeshSceneDesc* Primitive) override {};	
 
 	virtual void UpdatePrimitivesIsDrawn_RenderThread(TArrayView<FPrimitiveSceneProxy*> PrimitiveSceneProxies, bool bIsDrawn) override {}
 
@@ -7001,12 +6967,7 @@ void FRendererModule::RemoveScene(FSceneInterface* Scene)
 
 void FRendererModule::UpdateStaticDrawLists()
 {
-	// Update all static meshes in order to recache cached mesh draw commands.
-	check(IsInGameThread()); // AllocatedScenes is managed by the game thread
-	for (FSceneInterface* Scene : AllocatedScenes)
-	{
-		Scene->UpdateStaticDrawLists();
-	}
+	ensureMsgf(false, TEXT("UpdateStaticDrawLists is no longer implemented and will be removed, see UpdateStaticMeshesForMaterials."));
 }
 
 void UpdateStaticMeshesForMaterials(const TArray<const FMaterial*>& MaterialResourcesToUpdate)
diff --git a/Engine/Source/Runtime/Renderer/Private/ScenePrivate.h b/Engine/Source/Runtime/Renderer/Private/ScenePrivate.h
--- a/Engine/Source/Runtime/Renderer/Private/ScenePrivate.h
+++ b/Engine/Source/Runtime/Renderer/Private/ScenePrivate.h
@@ -3613,7 +3613,6 @@ public:
 		
 	virtual void UpdateCustomPrimitiveData(FPrimitiveSceneDesc* Primitive, const FCustomPrimitiveData& CustomPrimitiveData) override;
 
-	virtual void UpdatePrimitiveInstances(FInstancedStaticMeshSceneDesc* Primitive) override;
 	virtual void UpdatePrimitiveInstances(FPrimitiveSceneDesc* Primitive) override;
 
 	bool HasSkyAtmosphere() const
@@ -3649,9 +3648,6 @@ public:
 	/** Sets the precomputed visibility handler for the scene, or NULL to clear the current one. */
 	virtual void SetPrecomputedVisibility(const FPrecomputedVisibilityHandler* InPrecomputedVisibilityHandler) override;
 
-	/** Updates all static draw lists. */
-	virtual void UpdateStaticDrawLists() override;
-
 	/** Update render states that possibly cached inside renderer, like mesh draw commands. More lightweight than re-registering the scene proxy. */
 	virtual void UpdateCachedRenderStates(FPrimitiveSceneProxy* SceneProxy) override;
 
@@ -4060,9 +4056,6 @@ private:
 	/** Updates the contents of all reflection captures in the scene.  Must be called from the game thread. */
 	void UpdateAllReflectionCaptures(const TCHAR* CaptureReason, int32 ReflectionCaptureSize, bool bVerifyOnlyCapturing, bool bCapturingForMobile, bool bInsideTick);
 
-	/** Updates all static draw lists. */
-	void UpdateStaticDrawLists_RenderThread(FRHICommandListImmediate& RHICmdList);
-
 	/**
 	 * Shifts scene data by provided delta
 	 * Called on world origin changes
