Material Cache, fixed missing uniform invalidation on VT pool recreation

#Synced-CL 50905358

diff --git a/Engine/Source/Runtime/Engine/Internal/MaterialCache/IMaterialCacheTagProvider.h b/Engine/Source/Runtime/Engine/Internal/MaterialCache/IMaterialCacheTagProvider.h
--- a/Engine/Source/Runtime/Engine/Internal/MaterialCache/IMaterialCacheTagProvider.h
+++ b/Engine/Source/Runtime/Engine/Internal/MaterialCache/IMaterialCacheTagProvider.h
@@ -73,6 +73,11 @@ public:
 	 */
 	virtual void RemoveTagSceneInvalidationCallbacks(void* Baton) = 0;
 
+	/**
+	 * Invalidate all tag scenes
+	 */
+	virtual void NotifyAllTagSceneInvalidation() = 0;
+	
 	/**
 	 * Create 
 	 * @param Scene The scene the producer renders for
diff --git a/Engine/Source/Runtime/Engine/Private/MaterialCache/MaterialCacheVirtualTexture.cpp b/Engine/Source/Runtime/Engine/Private/MaterialCache/MaterialCacheVirtualTexture.cpp
--- a/Engine/Source/Runtime/Engine/Private/MaterialCache/MaterialCacheVirtualTexture.cpp
+++ b/Engine/Source/Runtime/Engine/Private/MaterialCache/MaterialCacheVirtualTexture.cpp
@@ -303,6 +303,8 @@ void UMaterialCacheVirtualTexture::Unregister()
 			delete RTDestructionBaton;
 		}
 	});
+	
+	OwningComponent = nullptr;
 }
 
 FIntPoint UMaterialCacheVirtualTexture::GetRuntimeTileCount() const
@@ -440,10 +442,10 @@ ETextureClass UMaterialCacheVirtualTexture::GetTextureClass() const
 FTextureResource* UMaterialCacheVirtualTexture::CreateResource()
 {
 	check(IsInGameThread());
-	
+
+	// Owning component may have released this texture
 	if (!OwningComponent.Get())
 	{
-		UE_LOG(LogEngine, Error, TEXT("Material Cache Virtual Texture requires an owning component"));
 		return nullptr;
 	}
 
diff --git a/Engine/Source/Runtime/Engine/Private/VT/VirtualTextureRecreate.cpp b/Engine/Source/Runtime/Engine/Private/VT/VirtualTextureRecreate.cpp
--- a/Engine/Source/Runtime/Engine/Private/VT/VirtualTextureRecreate.cpp
+++ b/Engine/Source/Runtime/Engine/Private/VT/VirtualTextureRecreate.cpp
@@ -2,6 +2,8 @@
 
 #include "VT/VirtualTextureRecreate.h"
 
+#include "EngineModule.h"
+#include "RendererInterface.h"
 #include "Components/PrimitiveComponent.h"
 #include "Components/RuntimeVirtualTextureComponent.h"
 #include "Engine/Texture.h"
@@ -10,6 +12,7 @@
 #include "UObject/UObjectIterator.h"
 #include "VirtualTexturing.h"
 #include "Engine/VirtualTextureCollection.h"
+#include "MaterialCache/IMaterialCacheTagProvider.h"
 #include "VT/RuntimeVirtualTexture.h"
 
 namespace VirtualTexture
@@ -39,7 +42,7 @@ namespace VirtualTexture
 		// Reinit material cache virtual textures.
 		for (TObjectIterator<UMaterialCacheVirtualTexture> It(RF_ClassDefaultObject, true, EInternalObjectFlags::Garbage); It; ++It)
 		{
-			It->UpdateResourceWithParams(UTexture::EUpdateResourceFlags::None);
+			It->UpdateResource();
 
 			if (UPrimitiveComponent* Owner = It->OwningComponent.Get())
 			{
@@ -51,6 +54,11 @@ namespace VirtualTexture
 		{
 			It->UpdateResource();
 		}
+
+		if (IMaterialCacheTagProvider* TagProvider = GetRendererModule().GetMaterialCacheTagProvider())
+		{
+			TagProvider->NotifyAllTagSceneInvalidation();
+		}
 	}
 
 	void Recreate(TConstArrayView< TEnumAsByte<EPixelFormat> > InFormat)
@@ -126,7 +134,7 @@ namespace VirtualTexture
 
 			if (bFormatsMatch)
 			{
-				It->UpdateResourceWithParams(UTexture::EUpdateResourceFlags::None);
+				It->UpdateResource();
 
 				if (UPrimitiveComponent* Owner = It->OwningComponent.Get())
 				{
@@ -142,5 +150,10 @@ namespace VirtualTexture
 				It->UpdateResource();
 			}
 		}
+
+		if (IMaterialCacheTagProvider* TagProvider = GetRendererModule().GetMaterialCacheTagProvider())
+		{
+			TagProvider->NotifyAllTagSceneInvalidation();
+		}
 	}
 }
diff --git a/Engine/Source/Runtime/Renderer/Internal/MaterialCache/MaterialCacheTagProvider.h b/Engine/Source/Runtime/Renderer/Internal/MaterialCache/MaterialCacheTagProvider.h
--- a/Engine/Source/Runtime/Renderer/Internal/MaterialCache/MaterialCacheTagProvider.h
+++ b/Engine/Source/Runtime/Renderer/Internal/MaterialCache/MaterialCacheTagProvider.h
@@ -35,6 +35,7 @@ public: /** IMaterialCacheTagProvider */
 	virtual IVirtualTexture* CreateProducer(FSceneInterface* Scene, FPrimitiveComponentId PrimitiveComponentId, const FMaterialCacheTagLayout& TagLayout, const FVTProducerDescription& ProducerDesc) override;
 	virtual void AddTagSceneInvalidationCallback(const FGuid& TagGuid, FMaterialCacheTagProviderSceneInvalidationDelegate Delegate, void* Baton) override;
 	virtual void RemoveTagSceneInvalidationCallbacks(void* Baton) override;
+	virtual void NotifyAllTagSceneInvalidation() override;
 
 public: /** Installation */
 	/** Initialize the global tag provider */
diff --git a/Engine/Source/Runtime/Renderer/Private/MaterialCache/MaterialCacheTagProvider.cpp b/Engine/Source/Runtime/Renderer/Private/MaterialCache/MaterialCacheTagProvider.cpp
--- a/Engine/Source/Runtime/Renderer/Private/MaterialCache/MaterialCacheTagProvider.cpp
+++ b/Engine/Source/Runtime/Renderer/Private/MaterialCache/MaterialCacheTagProvider.cpp
@@ -162,6 +162,14 @@ void FMaterialCacheTagProvider::CallPendingCallbacks()
 	}
 }
 
+void FMaterialCacheTagProvider::NotifyAllTagSceneInvalidation()
+{
+	for (auto&& [Key, _] : TagBuckets)
+	{
+		NotifyTagSceneInvalidation(Key);
+	}
+}
+
 uint32 FMaterialCacheTagProvider::AllocatePrimitiveTagOffset()
 {
 	FScopeLock Lock(&Mutex);
