Added fallback to a getter if a field override material member is not accessible or does not exist on a templated type.

#Synced-CL 51009600

diff --git a/Engine/Plugins/Experimental/FastGeoStreaming/Source/FastGeoStreaming/Internal/FastGeoMeshComponent.h b/Engine/Plugins/Experimental/FastGeoStreaming/Source/FastGeoStreaming/Internal/FastGeoMeshComponent.h
--- a/Engine/Plugins/Experimental/FastGeoStreaming/Source/FastGeoStreaming/Internal/FastGeoMeshComponent.h
+++ b/Engine/Plugins/Experimental/FastGeoStreaming/Source/FastGeoStreaming/Internal/FastGeoMeshComponent.h
@@ -34,6 +34,10 @@ public:
 	virtual const TArray<TObjectPtr<UMaterialInterface>>& GetComponentMaterialSlotsOverlayMaterial() const = 0;
 	virtual UMaterialInterface* GetOverlayMaterial() const = 0;
 
+	TConstArrayView<TObjectPtr<UMaterialInterface>> GetOverrideMaterials() const
+	{
+		return OverrideMaterials;
+	}
 protected:
 	// Persistent Data
 	TArray<TObjectPtr<UMaterialInterface>> OverrideMaterials;
diff --git a/Engine/Source/Runtime/Engine/Public/StaticMeshComponentHelper.h b/Engine/Source/Runtime/Engine/Public/StaticMeshComponentHelper.h
--- a/Engine/Source/Runtime/Engine/Public/StaticMeshComponentHelper.h
+++ b/Engine/Source/Runtime/Engine/Public/StaticMeshComponentHelper.h
@@ -18,6 +18,12 @@
 #include "Engine/World.h"
 #include "SceneInterface.h"
 
+template<typename T>
+concept HasOverridesMaterialsMember = requires(T t)
+{
+	t.OverrideMaterials;
+};
+
 /** Helper class used to share implementation for different StaticMeshComponent types */
 class FStaticMeshComponentHelper
 {
@@ -62,20 +68,38 @@ public:
 
 	template<class T, bool bAssumeRenderDataIsReady = false>
 	static FPrimitiveSceneProxy* CreateSceneProxy(T& Component, FStaticMeshComponentHelper::ESceneProxyCreationError* OutError = nullptr);
+
+	template<class T>
+	static TConstArrayView<TObjectPtr<UMaterialInterface>> GetOverrideMaterials(const T& Component)
+	{
+		if constexpr (HasOverridesMaterialsMember<std::remove_reference_t<T>>)
+		{
+			return Component.OverrideMaterials;
+		}
+		else
+		{
+			return Component.GetOverrideMaterials();
+		}
+	}
 };
 
 template<typename T>
 void FStaticMeshComponentHelper::GetUsedRayTracingOnlyMaterials(const T& Component, TArray<UMaterialInterface*>& OutMaterials)
 {
 	TArray<UMaterialInterface*> RayTracingMaterials;
-	Component.GetStaticMesh()->GetUsedMaterials(RayTracingMaterials, [&Component](int32 MaterialIndex)
+
+	const auto& OverrideMaterials = GetOverrideMaterials(Component);
+
+	if (UStaticMesh* StaticMesh = Component.GetStaticMesh())
+	{
+		StaticMesh->GetUsedMaterials(RayTracingMaterials, [&Component, &OverrideMaterials](int32 MaterialIndex)
 		{
 			UMaterialInterface* OutMaterial = nullptr;
 
 			// If we have a base materials array, use that
-			if (Component.OverrideMaterials.IsValidIndex(MaterialIndex) && Component.OverrideMaterials[MaterialIndex])
+			if (OverrideMaterials.IsValidIndex(MaterialIndex) && OverrideMaterials[MaterialIndex])
 			{
-				OutMaterial = Component.OverrideMaterials[MaterialIndex];
+				OutMaterial = OverrideMaterials[MaterialIndex];
 			}
 			// Otherwise get from static mesh
 			else if (Component.GetStaticMesh())
@@ -96,6 +120,7 @@ void FStaticMeshComponentHelper::GetUsedRayTracingOnlyMaterials(const T& Compone
 
 			return (UMaterialInterface*)nullptr;
 		});
+	}
 
 	if (!RayTracingMaterials.IsEmpty())
 	{
@@ -109,10 +134,12 @@ UMaterialInterface* FStaticMeshComponentHelper::GetMaterial(const T& Component,
 {
 	UMaterialInterface* OutMaterial = nullptr;
 
+	const auto& OverrideMaterials = GetOverrideMaterials(Component);
+
 	// If we have a base materials array, use that
-	if (Component.OverrideMaterials.IsValidIndex(MaterialIndex) && Component.OverrideMaterials[MaterialIndex])
+	if (OverrideMaterials.IsValidIndex(MaterialIndex) && OverrideMaterials[MaterialIndex])
 	{
-		OutMaterial = Component.OverrideMaterials[MaterialIndex];
+		OutMaterial = OverrideMaterials[MaterialIndex];
 	}
 	// Otherwise get from static mesh
 	else if (Component.GetStaticMesh())
@@ -192,13 +219,13 @@ void FStaticMeshComponentHelper::CollectPSOPrecacheDataImpl(const T& Component,
 		{
 			bAnySectionCastsShadows |= RenderSection.bCastShadow;
 
-			int16 MaterialIndex = RenderSection.MaterialIndex;
+			int16 MaterialIndex = (int16)RenderSection.MaterialIndex;
 			FPSOPrecacheVertexFactoryDataPerMaterialIndex* VFsPerMaterial = VFTypesPerMaterialIndex.FindByPredicate(
 				[MaterialIndex](const FPSOPrecacheVertexFactoryDataPerMaterialIndex& Other) { return Other.MaterialIndex == MaterialIndex; });
 			if (VFsPerMaterial == nullptr)
 			{
 				VFsPerMaterial = &VFTypesPerMaterialIndex.AddDefaulted_GetRef();
-				VFsPerMaterial->MaterialIndex = RenderSection.MaterialIndex;
+				VFsPerMaterial->MaterialIndex = (int16)RenderSection.MaterialIndex;
 			}
 
 			if (bSupportsManualVertexFetch)
