SpringMath: Added some more functions for estimating smoothing time.

#Synced-CL 51014848

#Synced-CL 51014848

diff --git a/Engine/Source/Runtime/Engine/Classes/Kismet/BlueprintSpringMathLibrary.h b/Engine/Source/Runtime/Engine/Classes/Kismet/BlueprintSpringMathLibrary.h
--- a/Engine/Source/Runtime/Engine/Classes/Kismet/BlueprintSpringMathLibrary.h
+++ b/Engine/Source/Runtime/Engine/Classes/Kismet/BlueprintSpringMathLibrary.h
@@ -524,6 +524,51 @@ public:
 	UFUNCTION(BlueprintPure, meta = (Category = "Math|Spring|Experimental", ReturnDisplayName = "Smoothing Time"))
 	static ENGINE_API float SpringCharacterSmoothingTimeFromTurningTime(const float TargetAngle, const float TurningTime = 1.0f, const float AngleThreshold = 5.0f);
 
+	UE_EXPERIMENTAL(5.8, "SpringMath is experimental")
+	/** Estimates an approximate maximum acceleration from a smoothing time for a character using a simple damped spring.
+	 *
+	 * @param InitialVelocity The initial velocity of the character
+	 * @param InitialAcceleration The initial acceleration of the character
+	 * @param TargetVelocity The target velocity of the character
+	 * @param SmoothingTime The time over which to smooth velocity. It takes roughly the smoothing time in order for the character to reach the target velocity.
+	 * @return An estimation of the maximum acceleration
+	 */
+	UFUNCTION(BlueprintPure, meta = (Category = "Math|Spring|Experimental", ReturnDisplayName = "Maximum Acceleration"))
+	static ENGINE_API float SpringCharacterMaximumAcceleration(const float InitialVelocity, const float InitialAcceleration, const float TargetVelocity, const float SmoothingTime = 0.2f);
+
+	UE_EXPERIMENTAL(5.8, "SpringMath is experimental")
+	/** Estimates an approximate smoothing time from a maximum acceleration for a character using by a simple damped spring.
+	 *
+	 * @param InitialVelocity The initial velocity of the character
+	 * @param TargetVelocity The target velocity of the character
+	 * @param MaximumAcceleration The maximum acceleration of the character
+	 * @return An estimation of the SmoothingTime of the spring
+	 */
+	UFUNCTION(BlueprintPure, meta = (Category = "Math|Spring|Experimental", ReturnDisplayName = "Smoothing Time"))
+	static ENGINE_API float SpringCharacterSmoothingTimeFromMaximumAcceleration(const float InitialVelocity, const float TargetVelocity, const float MaximumAcceleration);
+
+	UE_EXPERIMENTAL(5.8, "SpringMath is experimental")
+	/** Estimates an approximate maximum angular velocity from a smoothing time for a character using a simple damped spring.
+	 *
+	 * @param InitialAngle The initial angle of the character in degrees
+	 * @param InitialAngularVelocity The initial angular velocity of the character in deg/s
+	 * @param TargetAngle The target angle of the character
+	 * @param SmoothingTime The time over which to smooth angle. It takes roughly the smoothing time in order for the character to reach the target angle.
+	 * @return An estimation of the maximum angular velocity
+	 */
+	UFUNCTION(BlueprintPure, meta = (Category = "Math|Spring|Experimental", ReturnDisplayName = "Maximum Angular Velocity"))
+	static ENGINE_API float SpringCharacterMaximumAngularVelocity(const float InitialAngle, const float InitialAngularVelocity, const float TargetAngle, const float SmoothingTime = 0.2f);
+
+	UE_EXPERIMENTAL(5.8, "SpringMath is experimental")
+	/** Estimates an approximate smoothing time from a maximum angular velocity for a character using by a simple damped spring.
+	 *
+	 * @param InitialAngle The initial angle of the character in degrees
+	 * @param TargetAngle The target angle of the character
+	 * @param MaximumAngularVelocity The maximum angular velocity of the character
+	 * @return An estimation of the SmoothingTime of the spring
+	 */
+	UFUNCTION(BlueprintPure, meta = (Category = "Math|Spring|Experimental", ReturnDisplayName = "Smoothing Time"))
+	static ENGINE_API float SpringCharacterSmoothingTimeFromMaximumAngularVelocity(const float InitialAngle, const float TargetAngle, const float MaximumAngularVelocity);
 
 	UE_EXPERIMENTAL(5.7, "SpringMath is experimental")
 	/** Update a position representing a character given a target velocity using a velocity spring.
diff --git a/Engine/Source/Runtime/Engine/Private/BlueprintSpringMathLibrary.cpp b/Engine/Source/Runtime/Engine/Private/BlueprintSpringMathLibrary.cpp
--- a/Engine/Source/Runtime/Engine/Private/BlueprintSpringMathLibrary.cpp
+++ b/Engine/Source/Runtime/Engine/Private/BlueprintSpringMathLibrary.cpp
@@ -289,6 +289,26 @@ float UBlueprintSpringMathLibrary::SpringCharacterSmoothingTimeFromTurningTime(c
 	return SpringMath::SpringCharacterSmoothingTimeFromStartingTime(TargetAngle, TurningTime, AngleThreshold);
 }
 
+float UBlueprintSpringMathLibrary::SpringCharacterMaximumAcceleration(const float InitialVelocity, const float InitialAcceleration, const float TargetVelocity, const float SmoothingTime)
+{
+	return SpringMath::SpringCharacterMaximumAcceleration(InitialVelocity, InitialAcceleration, TargetVelocity, SmoothingTime);
+}
+
+float UBlueprintSpringMathLibrary::SpringCharacterSmoothingTimeFromMaximumAcceleration(const float InitialVelocity, const float TargetVelocity, const float MaximumAcceleration)
+{
+	return SpringMath::SpringCharacterSmoothingTimeFromMaximumAcceleration(InitialVelocity, TargetVelocity, MaximumAcceleration);
+}
+
+float UBlueprintSpringMathLibrary::SpringCharacterMaximumAngularVelocity(const float InitialAngle, const float InitialAngularVelocity, const float TargetAngle, const float SmoothingTime)
+{
+	return FMath::RadiansToDegrees(SpringMath::SpringCharacterMaximumAngularVelocity(FMath::DegreesToRadians(InitialAngle), FMath::DegreesToRadians(InitialAngularVelocity), FMath::DegreesToRadians(TargetAngle), SmoothingTime));
+}
+
+float UBlueprintSpringMathLibrary::SpringCharacterSmoothingTimeFromMaximumAngularVelocity(const float InitialAngle, const float TargetAngle, const float MaximumAngularVelocity)
+{
+	return SpringMath::SpringCharacterSmoothingTimeFromMaximumAngularVelocity(FMath::DegreesToRadians(InitialAngle), FMath::DegreesToRadians(TargetAngle), FMath::DegreesToRadians(MaximumAngularVelocity));
+}
+
 void UBlueprintSpringMathLibrary::VelocitySpringCharacterUpdate(FVector& InOutPosition, FVector& InOutVelocity, FVector& InOutVelocityIntermediate,
 	FVector& InOutAcceleration, const FVector& TargetVelocity, float DeltaTime, float SmoothingTime, float MaxAcceleration)
 {
diff --git a/Engine/Source/Runtime/Engine/Public/Animation/SpringMath.h b/Engine/Source/Runtime/Engine/Public/Animation/SpringMath.h
--- a/Engine/Source/Runtime/Engine/Public/Animation/SpringMath.h
+++ b/Engine/Source/Runtime/Engine/Public/Animation/SpringMath.h
@@ -12,6 +12,8 @@
 //		     https://theorangeduck.com/page/scalar-velocity
 //			 https://theorangeduck.com/page/dead-blending
 //			 https://theorangeduck.com/page/propagating-velocities-through-animation-systems
+//			 https://theorangeduck.com/page/fitting-code-driven-displacement
+//			 https://theorangeduck.com/page/fitting-code-driven-displacement-revisited
 
 struct SpringMath
 {
@@ -787,6 +789,82 @@ public:
 		return DampingToSmoothingTime((2.0f * FMath::Exp(-W) * V * (2.0f + W) + 2.0f * V * (W - 2.0f)) / FMath::Max(StartingDistance, UE_SMALL_NUMBER));
 	}
 
+	/** Estimates the time at which the maximum acceleration will occur */
+	UE_EXPERIMENTAL(5.8, "SpringMath is experimental")
+	static inline float SpringCharacterMaximumAccelerationTime(const float InitialVelocity, const float InitialAcceleration, const float TargetVelocity, const float SmoothingTime)
+	{
+		const float Y = SmoothingTimeToDamping(SmoothingTime) / 2.0f;
+		const float J0 = InitialVelocity - TargetVelocity;
+		const float J1 = InitialAcceleration + J0 * Y;
+		return (J1 + InitialAcceleration) / (J1 * Y);
+	}
+
+	/** Estimates the maximum acceleration for a character */
+	UE_EXPERIMENTAL(5.8, "SpringMath is experimental")
+	static inline float SpringCharacterMaximumAcceleration(const float InitialVelocity, const float InitialAcceleration, const float TargetVelocity, const float SmoothingTime)
+	{
+		const float Y = SmoothingTimeToDamping(SmoothingTime) / 2.0f;
+		const float J0 = InitialVelocity - TargetVelocity;
+		const float J1 = InitialAcceleration + J0 * Y;
+		PRAGMA_DISABLE_EXPERIMENTAL_WARNINGS
+		const float T = SpringCharacterMaximumAccelerationTime(InitialVelocity, InitialAcceleration, TargetVelocity, SmoothingTime);
+		PRAGMA_ENABLE_EXPERIMENTAL_WARNINGS
+
+		if (FMath::IsFinite(T) && T > 0.0f)
+		{
+			return FMath::InvExpApprox(Y * T) * (InitialAcceleration - J1 * Y * T);
+		}
+		else
+		{
+			return InitialAcceleration;
+		}
+	}
+
+	/** Estimates an approximate spring smoothing time from the maximum acceleration */
+	UE_EXPERIMENTAL(5.8, "SpringMath is experimental")
+	static inline float SpringCharacterSmoothingTimeFromMaximumAcceleration(const float InitialVelocity, const float TargetVelocity, const float MaximumAcceleration)
+	{
+		return DampingToSmoothingTime(2.0f * ((MaximumAcceleration / FMath::Max(FMath::Abs(TargetVelocity - InitialVelocity), UE_SMALL_NUMBER)) * UE_EULERS_NUMBER));
+	}
+
+	/** Estimates the time at which the maximum angular velocity will occur */
+	UE_EXPERIMENTAL(5.8, "SpringMath is experimental")
+	static inline float SpringCharacterMaximumAngularVelocityTime(const float InitialAngle, const float InitialAngularVelocity, const float TargetAngle, const float SmoothingTime)
+	{
+		const float Y = SmoothingTimeToDamping(SmoothingTime) / 2.0f;
+		const float J0 = FMath::FindDeltaAngleRadians(TargetAngle, InitialAngle);
+		const float J1 = InitialAngularVelocity + J0 * Y;
+		return (J1 + InitialAngularVelocity) / (J1 * Y);
+	}
+
+	/** Estimates the maximum angular velocity for a character */
+	UE_EXPERIMENTAL(5.8, "SpringMath is experimental")
+	static inline float SpringCharacterMaximumAngularVelocity(const float InitialAngle, const float InitialAngularVelocity, const float TargetAngle, const float SmoothingTime)
+	{
+		const float Y = SmoothingTimeToDamping(SmoothingTime) / 2.0f;
+		const float J0 = FMath::FindDeltaAngleRadians(TargetAngle, InitialAngle);
+		const float J1 = InitialAngularVelocity + J0 * Y;
+		PRAGMA_DISABLE_EXPERIMENTAL_WARNINGS
+		const float T = SpringCharacterMaximumAngularVelocityTime(InitialAngle, InitialAngularVelocity, TargetAngle, SmoothingTime);
+		PRAGMA_ENABLE_EXPERIMENTAL_WARNINGS
+
+		if (FMath::IsFinite(T) && T > 0.0f)
+		{
+			return FMath::InvExpApprox(Y * T) * (InitialAngularVelocity - J1 * Y * T);
+		}
+		else
+		{
+			return InitialAngularVelocity;
+		}
+	}
+
+	/** Estimates an approximate spring smoothing time from the maximum angular velocity */
+	UE_EXPERIMENTAL(5.8, "SpringMath is experimental")
+		static inline float SpringCharacterSmoothingTimeFromMaximumAngularVelocity(const float InitialAngle, const float TargetAngle, const float MaximumAngularVelocity)
+	{
+		return DampingToSmoothingTime(2.0f * ((MaximumAngularVelocity / FMath::Max(FMath::Abs(FMath::FindDeltaAngleRadians(InitialAngle, TargetAngle)), UE_SMALL_NUMBER)) * UE_EULERS_NUMBER));
+	}
+
 	/** Update a position representing a character given a target velocity using a velocity spring.
 	 * A velocity spring tracks an intermediate velocity which moves at a maximum acceleration linearly towards a target.
 	 * This means unlike the "SpringCharacterUpdate", it will take longer to reach a target velocity that is further away from the current velocity.
