BodySetup: Cooking: During TryGetCookedData, enter a lock around the read of CookedFormatData in addition to its write, to prevent threads crashing when deciding whether they need to enter the write lock due to another thread already in it and mutating the CookedFormatData.
#rnx

#Synced-CL 50900217

diff --git a/Engine/Source/Runtime/Engine/Classes/PhysicsEngine/BodySetup.h b/Engine/Source/Runtime/Engine/Classes/PhysicsEngine/BodySetup.h
--- a/Engine/Source/Runtime/Engine/Classes/PhysicsEngine/BodySetup.h
+++ b/Engine/Source/Runtime/Engine/Classes/PhysicsEngine/BodySetup.h
@@ -268,7 +268,7 @@ class UBodySetup : public UBodySetupCore
 	 * This mutex protects the cook path in TryGetCookedData so we don't have multiple threads attempting to
 	 * write to the format data
 	 */
-	FCriticalSection CookedFormatDataMutex;
+	FRWLock CookedFormatDataMutex;
 #endif
 
 	/** GUID used to uniquely identify this setup so it can be found in the DDC */
diff --git a/Engine/Source/Runtime/Engine/Private/PhysicsEngine/BodySetup.cpp b/Engine/Source/Runtime/Engine/Private/PhysicsEngine/BodySetup.cpp
--- a/Engine/Source/Runtime/Engine/Private/PhysicsEngine/BodySetup.cpp
+++ b/Engine/Source/Runtime/Engine/Private/PhysicsEngine/BodySetup.cpp
@@ -1471,7 +1471,16 @@ UBodySetup::FBodySetupTryCookResult UBodySetup::TryGetCookedData(FName Format)
 	const bool bIsRuntime = IsRuntime(this);
 
 #if WITH_EDITOR
-	if(!UsedFormatData->Contains(Format))
+	bool bFormatDataContainsFormat;
+	{
+		FReadScopeLock ScopeLock(CookedFormatDataMutex);
+		bFormatDataContainsFormat = UsedFormatData->Contains(Format);
+		if (bFormatDataContainsFormat)
+		{
+			Result = &UsedFormatData->GetFormat(Format);
+		}
+	}
+	if (!bFormatDataContainsFormat)
 	{
 		SCOPE_CYCLE_COUNTER(STAT_PhysXCooking);
 
@@ -1487,12 +1496,13 @@ UBodySetup::FBodySetupTryCookResult UBodySetup::TryGetCookedData(FName Format)
 
 		// Only one thread should build the physics data. Below calling GetFormat and TryGetDDCBuiltData will
 		// both mutate shared state so we protect this operation with a lock.
-		FScopeLock ScopeLock(&CookedFormatDataMutex);
+		FWriteScopeLock ScopeLock(CookedFormatDataMutex);
 
 		// Check the lock condition again, we only want the first thread in here to actually build the data,
 		// while it's building other threads should wait and proceed after the build completes.
 		// Double-checked locking avoids always taking the lock - we only do it for unbuilt data
-		if(!UsedFormatData->Contains(Format))
+		bFormatDataContainsFormat = UsedFormatData->Contains(Format);
+		if (!bFormatDataContainsFormat)
 		{
 			// We do not want a FGCObject to be created to prevent garbage collection of our own UBodySetup*
 			// because the scope and lifetime is well defined and FGCObject can't be created on other threads
@@ -1509,11 +1519,11 @@ UBodySetup::FBodySetupTryCookResult UBodySetup::TryGetCookedData(FName Format)
 			Result = &UsedFormatData->GetFormat(Format);
 		}
 	}
-	else
-#endif // #if WITH_EDITOR
+#else // #if WITH_EDITOR
 	{
 		Result = &UsedFormatData->GetFormat(Format);
 	}
+#endif // #else !WITH_EDITOR
 
 	check(Result);
 
