Fix for SpecularProfile typed UAV writes on OpenGL

#Synced-CL 51011289

diff --git a/Engine/Source/Runtime/Engine/Private/Rendering/SpecularProfile.cpp b/Engine/Source/Runtime/Engine/Private/Rendering/SpecularProfile.cpp
--- a/Engine/Source/Runtime/Engine/Private/Rendering/SpecularProfile.cpp
+++ b/Engine/Source/Runtime/Engine/Private/Rendering/SpecularProfile.cpp
@@ -363,7 +363,7 @@ IPooledRenderTarget* FSpecularProfileTextureManager::GetAtlasTexture(FRDGBuilder
 
 		// 1. Create atlas texture
 		// Check if we can do typed UAV loads on pixel formats with optional UAV load support
-		const bool bR11G11B10Supported = GPixelFormats[PF_FloatR11G11B10].Supported && UE::PixelFormat::HasCapabilities(PF_FloatR11G11B10, EPixelFormatCapabilities::TextureFilterable);
+		const bool bR11G11B10Supported = GPixelFormats[PF_FloatR11G11B10].Supported && UE::PixelFormat::HasCapabilities(PF_FloatR11G11B10, EPixelFormatCapabilities::TextureFilterable | EPixelFormatCapabilities::UAV);
 		const bool bR11G11B10TypedLoadSupported = bR11G11B10Supported && UE::PixelFormat::HasCapabilities(PF_FloatR11G11B10, EPixelFormatCapabilities::TypedUAVLoad);
 		const bool bRGBA16TypedLoadSupported = UE::PixelFormat::HasCapabilities(PF_FloatRGBA, EPixelFormatCapabilities::TypedUAVLoad);
 
