[EntityFramework] minor optimization in 'has_collision_impl::GetCollidableImpl()' and exported some BodyInstance helpers

#Synced-CL 51020371

diff --git a/Engine/Source/Runtime/Engine/Classes/PhysicsEngine/BodyInstance.h b/Engine/Source/Runtime/Engine/Classes/PhysicsEngine/BodyInstance.h
--- a/Engine/Source/Runtime/Engine/Classes/PhysicsEngine/BodyInstance.h
+++ b/Engine/Source/Runtime/Engine/Classes/PhysicsEngine/BodyInstance.h
@@ -169,7 +169,7 @@ public:
 	// Takes actor ref arrays.
 	// #PHYS2 this used to return arrays of low-level physics bodies, which would be added to scene in InitBodies. Should it still do that, rather then later iterate over BodyInstances to get phys actor refs?
 	bool CreateShapesAndActors();
-	void InitBodies();
+	UE_API void InitBodies();
 
 protected:
 	void UpdateSimulatingAndBlendWeight();
@@ -181,7 +181,7 @@ using FInitBodiesHelperBase = TInitBodiesHelperBase<FDefaultAllocator>;
 template <bool bCompileStatic, typename InAllocatorType = FDefaultAllocator>
 struct FInitBodiesHelper : public TInitBodiesHelperBase<InAllocatorType>
 {
-	FInitBodiesHelper(TArray<FBodyInstance*, InAllocatorType>& InBodies, TArray<FTransform, InAllocatorType>& InTransforms, class UBodySetup* InBodySetup,
+	UE_API FInitBodiesHelper(TArray<FBodyInstance*, InAllocatorType>& InBodies, TArray<FTransform, InAllocatorType>& InTransforms, class UBodySetup* InBodySetup,
 		class UPrimitiveComponent* InPrimitiveComp, FPhysScene* InRBScene,
 		const FInitBodySpawnParams& InSpawnParams, FPhysicsAggregateHandle InAggregate, IPhysicsBodyInstanceOwner* InBodyInstanceOwner = nullptr);
 
