PR #14425: Fix typos in World.h comments

#Synced-CL 51019638

diff --git a/Engine/Source/Runtime/Engine/Classes/Engine/World.h b/Engine/Source/Runtime/Engine/Classes/Engine/World.h
--- a/Engine/Source/Runtime/Engine/Classes/Engine/World.h
+++ b/Engine/Source/Runtime/Engine/Classes/Engine/World.h
@@ -516,7 +516,7 @@ public:
 	TFunction<void(AActor*)> CustomPreSpawnInitialization;
 };
 
-/* World actors spawmning helper functions */
+/* World actors spawning helper functions */
 struct FActorSpawnUtils
 {
 	/**
@@ -541,7 +541,7 @@ struct FActorSpawnUtils
 	static UE_API bool IsGloballyUniqueName(FName Name);
 
 	/**
-	 * Return the base ename (without any number of globally unique identifier).
+	 * Return the base name (without any number of globally unique identifier).
 	**/
 	static UE_API FName GetBaseName(FName Name);
 };
@@ -1167,7 +1167,7 @@ private:
 	/** Sets level's PendingVisibilityState */
 	UE_API void SetLevelPendingVisibilityState(ULevel* InLevel, ELevelPendingVisibilityState InState);
 
-	/** Returns wether AddToWorld should be skipped on a given level */
+	/** Returns whether AddToWorld should be skipped on a given level */
 	UE_API bool CanAddLoadedLevelToWorld(ULevel* Level) const;
 
 	/** Spawn an instance of ANoPawnPlayerController and set it as the value of Connection::PlayerController. */
@@ -1517,7 +1517,7 @@ public:
 	/** Map of blueprints that are being debugged and the object instance they are debugging. */
 	typedef TMap<TWeakObjectPtr<class UBlueprint>, TWeakObjectPtr<UObject> > FBlueprintToDebuggedObjectMap;
 
-	/** Return the array of objects currently bieng debugged. */
+	/** Return the array of objects currently being debugged. */
 	const FBlueprintToDebuggedObjectMap& GetBlueprintObjectsBeingDebugged() const{ return BlueprintObjectsBeingDebugged; };
 #endif
 
@@ -1855,7 +1855,7 @@ private:
 	/** Disables the broadcasting of level selection change. Internal use only. */
 	uint32 bBroadcastSelectionChange:1;
 
-	/** a delegate that broadcasts a notification whenever the current feautre level is changed */
+	/** a delegate that broadcasts a notification whenever the current feature level is changed */
 	FOnFeatureLevelChanged OnFeatureLevelChanged;
 #endif //WITH_EDITORONLY_DATA
 
@@ -3395,7 +3395,7 @@ public:
 	UE_API void FinishAsyncSendAllEndOfFrameUpdates();
 
 	/**
-	 * Flush any pending parameter collection updates to the render thrad.
+	 * Flush any pending parameter collection updates to the render thread.
 	 */
 	UE_API void FlushDeferredParameterCollectionInstanceUpdates();
 
@@ -4132,7 +4132,7 @@ public:
 #endif
 
 	/**
-	 * Jumps the server to new level.  If bAbsolute is true and we are using seemless traveling, we
+	 * Jumps the server to new level.  If bAbsolute is true and we are using seamless traveling, we
 	 * will do an absolute travel (URL will be flushed).
 	 *
 	 * @param URL the URL that we are traveling to
