Clear letterboxing with transparent alpha when r.TransparentViewportClear is set

#Synced-CL 51044868

diff --git a/Engine/Source/Runtime/Engine/Private/GameViewportClient.cpp b/Engine/Source/Runtime/Engine/Private/GameViewportClient.cpp
--- a/Engine/Source/Runtime/Engine/Private/GameViewportClient.cpp
+++ b/Engine/Source/Runtime/Engine/Private/GameViewportClient.cpp
@@ -1932,25 +1932,27 @@ void UGameViewportClient::Draw(FViewport* InViewport, FCanvas* SceneCanvas)
 	// Clear areas of the rendertarget (backbuffer) that aren't drawn over by the views.
 	if (!bBufferCleared)
 	{
+		const FLinearColor ClearColor = bTransparentViewportClear ? FLinearColor::Transparent : FLinearColor::Black;
+
 		// clear left
 		if( MinX > 0 )
 		{
-			SceneCanvas->DrawTile(0,0,MinX,InViewport->GetSizeXY().Y,0.0f,0.0f,1.0f,1.f,FLinearColor::Black,NULL,false);
+			SceneCanvas->DrawTile(0,0,MinX,InViewport->GetSizeXY().Y,0.0f,0.0f,1.0f,1.f,ClearColor,NULL,false);
 		}
 		// clear right
 		if( MaxX < (uint32)InViewport->GetSizeXY().X )
 		{
-			SceneCanvas->DrawTile(MaxX,0,InViewport->GetSizeXY().X,InViewport->GetSizeXY().Y,0.0f,0.0f,1.0f,1.f,FLinearColor::Black,NULL,false);
+			SceneCanvas->DrawTile(MaxX,0,InViewport->GetSizeXY().X,InViewport->GetSizeXY().Y,0.0f,0.0f,1.0f,1.f,ClearColor,NULL,false);
 		}
 		// clear top
 		if( MinY > 0 )
 		{
-			SceneCanvas->DrawTile(MinX,0,MaxX,MinY,0.0f,0.0f,1.0f,1.f,FLinearColor::Black,NULL,false);
+			SceneCanvas->DrawTile(MinX,0,MaxX,MinY,0.0f,0.0f,1.0f,1.f,ClearColor,NULL,false);
 		}
 		// clear bottom
 		if( MaxY < (uint32)InViewport->GetSizeXY().Y )
 		{
-			SceneCanvas->DrawTile(MinX,MaxY,MaxX,InViewport->GetSizeXY().Y,0.0f,0.0f,1.0f,1.f,FLinearColor::Black,NULL,false);
+			SceneCanvas->DrawTile(MinX,MaxY,MaxX,InViewport->GetSizeXY().Y,0.0f,0.0f,1.0f,1.f,ClearColor,NULL,false);
 		}
 	}
 
