CellTransformer DataAsset processor: PIE support

#Synced-CL 50931119

diff --git a/Engine/Source/Runtime/Engine/Public/WorldPartition/WorldPartitionLevelStreamingDynamic.h b/Engine/Source/Runtime/Engine/Public/WorldPartition/WorldPartitionLevelStreamingDynamic.h
--- a/Engine/Source/Runtime/Engine/Public/WorldPartition/WorldPartitionLevelStreamingDynamic.h
+++ b/Engine/Source/Runtime/Engine/Public/WorldPartition/WorldPartitionLevelStreamingDynamic.h
@@ -66,6 +66,8 @@ class UWorldPartitionLevelStreamingDynamic : public ULevelStreamingDynamic
 	UE_DEPRECATED(5.7, "UnloadFromEditor() is no longer supported and will be removed in a future version.")
 	static ENGINE_API void UnloadFromEditor(UWorldPartitionLevelStreamingDynamic* InLevelStreaming);
 
+	const TArray<FWorldPartitionRuntimeCellObjectMapping>& GetChildPackages() const { return ChildPackages; };
+
 private:
 	ENGINE_API void CreateRuntimeLevel();
 	ENGINE_API bool IssueLoadRequests();
