Add finalize animation callback on anim instances to be able to cleanup when we update without evaluating

#Synced-CL 50680057

diff --git a/Engine/Source/Editor/UnrealEd/Private/Animation/DebugSkelMeshComponent.cpp b/Engine/Source/Editor/UnrealEd/Private/Animation/DebugSkelMeshComponent.cpp
--- a/Engine/Source/Editor/UnrealEd/Private/Animation/DebugSkelMeshComponent.cpp
+++ b/Engine/Source/Editor/UnrealEd/Private/Animation/DebugSkelMeshComponent.cpp
@@ -921,6 +921,7 @@ void UDebugSkelMeshComponent::GenSpaceBases(TArray<FTransform>& OutSpaceBases)
 	DoInstancePreEvaluation();
 	PerformAnimationEvaluation(GetSkeletalMeshAsset(), AnimScriptInstance, OutSpaceBases, TempBoneSpaceTransforms, TempRootBoneTranslation, TempCurve, TempAtttributes);
 	DoInstancePostEvaluation();
+	DoInstanceFinalizeAnimation(true);
 }
 
 void UDebugSkelMeshComponent::RefreshBoneTransforms(FActorComponentTickFunction* TickFunction)
diff --git a/Engine/Source/Runtime/Engine/Classes/Animation/AnimInstance.h b/Engine/Source/Runtime/Engine/Classes/Animation/AnimInstance.h
--- a/Engine/Source/Runtime/Engine/Classes/Animation/AnimInstance.h
+++ b/Engine/Source/Runtime/Engine/Classes/Animation/AnimInstance.h
@@ -1431,7 +1431,10 @@ public:
 	UE_DEPRECATED(4.26, "Please use ParallelEvaluateAnimation with different signature.")
 	ENGINE_API void ParallelEvaluateAnimation(bool bForceRefPose, const USkeletalMesh* InSkeletalMesh, FBlendedHeapCurve& OutCurve, FCompactPose& OutPose);
 
+	// Called post evaluation if an evaluation took place
 	ENGINE_API void PostEvaluateAnimation();
+	// Called after PostEvaluateAnimation, whether or not an evaluation took place
+	ENGINE_API void FinalizeAnimation(bool bWasEvaluated);
 	ENGINE_API void UninitializeAnimation();
 
 	// the below functions are the native overrides for each phase
@@ -1445,6 +1448,8 @@ public:
 	ENGINE_API virtual void NativeThreadSafeUpdateAnimation(float DeltaSeconds);
 	// Native Post Evaluate override point
 	ENGINE_API virtual void NativePostEvaluateAnimation();
+	// Native Finalize override point
+	ENGINE_API virtual void NativeFinalizeAnimation(bool bWasEvaluated);
 	// Native Uninitialize override point
 	ENGINE_API virtual void NativeUninitializeAnimation();
 	// Native override point executed when all Linked Animation Layers are initialized
diff --git a/Engine/Source/Runtime/Engine/Classes/Components/SkeletalMeshComponent.h b/Engine/Source/Runtime/Engine/Classes/Components/SkeletalMeshComponent.h
--- a/Engine/Source/Runtime/Engine/Classes/Components/SkeletalMeshComponent.h
+++ b/Engine/Source/Runtime/Engine/Classes/Components/SkeletalMeshComponent.h
@@ -2639,6 +2639,9 @@ protected:
 	// Handle post evaluation on all anim instances
 	ENGINE_API void DoInstancePostEvaluation();
 
+	// Handle post animation on all anim instances, whether or not an evaluation took place 
+	ENGINE_API void DoInstanceFinalizeAnimation(bool bWasEvaluated);
+
 	// Returns whether we need to run the Cloth Tick or not
 	ENGINE_API virtual bool ShouldRunClothTick() const;
 
diff --git a/Engine/Source/Runtime/Engine/Private/Animation/AnimInstance.cpp b/Engine/Source/Runtime/Engine/Private/Animation/AnimInstance.cpp
--- a/Engine/Source/Runtime/Engine/Private/Animation/AnimInstance.cpp
+++ b/Engine/Source/Runtime/Engine/Private/Animation/AnimInstance.cpp
@@ -1083,6 +1083,12 @@ void UAnimInstance::PostEvaluateAnimation()
 	GetProxyOnGameThread<FAnimInstanceProxy>().PostEvaluate(this);
 }
 
+void UAnimInstance::FinalizeAnimation(bool bWasEvaluated)
+{
+	LLM_SCOPE(ELLMTag::Animation);
+	NativeFinalizeAnimation(bWasEvaluated);
+}
+
 void UAnimInstance::NativeInitializeAnimation()
 {
 }
@@ -1103,6 +1109,10 @@ void UAnimInstance::NativePostEvaluateAnimation()
 {
 }
 
+void UAnimInstance::NativeFinalizeAnimation(bool bWasEvaluated)
+{
+}
+
 void UAnimInstance::NativeUninitializeAnimation()
 {
 }
diff --git a/Engine/Source/Runtime/Engine/Private/Components/SkeletalMeshComponent.cpp b/Engine/Source/Runtime/Engine/Private/Components/SkeletalMeshComponent.cpp
--- a/Engine/Source/Runtime/Engine/Private/Components/SkeletalMeshComponent.cpp
+++ b/Engine/Source/Runtime/Engine/Private/Components/SkeletalMeshComponent.cpp
@@ -2780,6 +2780,25 @@ void USkeletalMeshComponent::DoInstancePostEvaluation()
 	}
 }
 
+void USkeletalMeshComponent::DoInstanceFinalizeAnimation(bool bWasEvaluated)
+{
+	if (AnimScriptInstance)
+	{
+		AnimScriptInstance->FinalizeAnimation(bWasEvaluated);
+
+		for (UAnimInstance* LinkedInstance : LinkedInstances)
+		{
+			LinkedInstance->FinalizeAnimation(bWasEvaluated);
+		}
+	}
+
+	if (PostProcessAnimInstance)
+	{
+		PostProcessAnimInstance->FinalizeAnimation(bWasEvaluated);
+	}
+}
+
+
 void USkeletalMeshComponent::RefreshBoneTransforms(FActorComponentTickFunction* TickFunction)
 {
 	if(!bEnableAnimation)
@@ -3243,6 +3262,8 @@ void USkeletalMeshComponent::PostAnimEvaluation(FAnimationEvaluationContext& Eva
 			DoInstancePostEvaluation();
 		}
 
+		DoInstanceFinalizeAnimation(EvaluationContext.bDoEvaluation);
+
 		bNeedToFlipSpaceBaseBuffers = true;
 
 		if (Bodies.Num() > 0 || bEnablePerPolyCollision)
@@ -3282,6 +3303,8 @@ void USkeletalMeshComponent::PostAnimEvaluation(FAnimationEvaluationContext& Eva
 	}
 	else 
 	{
+		DoInstanceFinalizeAnimation(false);
+
 		// Since we're not calling FinalizeBoneTransforms via PostBlendPhysics,
 		// make sure we call ConditionallyDispatchQueuedAnimEvents() in case we ticked, but didn't evalutate.
 
