Add the ability to invalidate cached transaction request variables to World and WorldPartition

#rnx

#Synced-CL 50820543

diff --git a/Engine/Source/Runtime/Engine/Classes/Engine/World.h b/Engine/Source/Runtime/Engine/Classes/Engine/World.h
--- a/Engine/Source/Runtime/Engine/Classes/Engine/World.h
+++ b/Engine/Source/Runtime/Engine/Classes/Engine/World.h
@@ -1029,6 +1029,10 @@ public:
 	/** Returns whether the world supports for a client to use "making invisible" transaction requests to the server. */
 	UE_API bool SupportsMakingInvisibleTransactionRequests() const;
 
+	/** Used to allow changing of visiblity transaction information in situations where overrides are applied after the world is created */
+	UE_EXPERIMENTAL(5.8, "Experimental, API may change in the future")
+	UE_API void InvalidateVisiblityTransactionRequestCachedVariables();
+
 	/** Returns the object used to query server streaming level visibility. */
 	UE_API const AServerStreamingLevelsVisibility* GetServerStreamingLevelsVisibility() const;
 
diff --git a/Engine/Source/Runtime/Engine/Private/LevelStreaming.cpp b/Engine/Source/Runtime/Engine/Private/LevelStreaming.cpp
--- a/Engine/Source/Runtime/Engine/Private/LevelStreaming.cpp
+++ b/Engine/Source/Runtime/Engine/Private/LevelStreaming.cpp
@@ -46,6 +46,21 @@ DEFINE_LOG_CATEGORY(LogLevelStreaming);
 // CVars
 namespace LevelStreamingCVars
 {
+	static void OnVisibilityTransactionCVarChanged(IConsoleVariable* CVar)
+	{
+		if (GEngine)
+		{
+			for (const FWorldContext& Context : GEngine->GetWorldContexts())
+			{
+				UWorld* World = Context.World();
+				if (World && World->IsGameWorld())
+				{
+					World->InvalidateVisiblityTransactionRequestCachedVariables();
+				}
+			}
+		}
+	}
+
 	// There are cases where we might have multiple visibility requests (and data) in flight leading to the server 
 	// starting to replicate data based on an older visibility/streamingstatus update which can lead to broken channels
 	// to mitigate this problem we assign a TransactionId to each request/update to make sure that we are acting on the correct data
@@ -56,6 +71,7 @@ namespace LevelStreamingCVars
 		TEXT("Flag combined with world support to use making invisible transaction requests to the server\n")
 		TEXT("that determines whether the client should wait for the server to acknowledge visibility update before making streaming levels invisible.\n")
 		TEXT("0: Disable, 1: Enable"),
+		FConsoleVariableDelegate::CreateStatic(&*LevelStreamingCVars::OnVisibilityTransactionCVarChanged),
 		ECVF_Default);
 
 	static bool bDefaultAllowClientUseMakingVisibleTransactionRequests = false;
@@ -65,6 +81,7 @@ namespace LevelStreamingCVars
 		TEXT("Flag combined with world support to use making visible transaction requests to the server\n")
 		TEXT("that determines whether the client should wait for the server to acknowledge visibility update before making streaming levels visible.\n")
 		TEXT("0: Disable, 1: Enable"),
+		FConsoleVariableDelegate::CreateStatic(&*LevelStreamingCVars::OnVisibilityTransactionCVarChanged),
 		ECVF_Default);
 
 	static bool bShouldServerUseMakingVisibleTransactionRequest = true;
@@ -73,6 +90,7 @@ namespace LevelStreamingCVars
 		bShouldServerUseMakingVisibleTransactionRequest,
 		TEXT("Whether server should wait for client to acknowledge visibility update before treating streaming levels as visible by the client.\n")
 		TEXT("0: Disable, 1: Enable"),
+		FConsoleVariableDelegate::CreateStatic(&*LevelStreamingCVars::OnVisibilityTransactionCVarChanged),
 		ECVF_Default);
 
 	static bool bShouldReuseUnloadedButStillAroundLevels = true;
diff --git a/Engine/Source/Runtime/Engine/Private/World.cpp b/Engine/Source/Runtime/Engine/Private/World.cpp
--- a/Engine/Source/Runtime/Engine/Private/World.cpp
+++ b/Engine/Source/Runtime/Engine/Private/World.cpp
@@ -3476,6 +3476,17 @@ bool UWorld::SupportsMakingInvisibleTransactionRequests() const
 	return false;
 }
 
+void UWorld::InvalidateVisiblityTransactionRequestCachedVariables()
+{
+	bSupportsMakingVisibleTransactionRequests.Reset();
+	bSupportsMakingInvisibleTransactionRequests.Reset();
+
+	if (IsPartitionedWorld())
+	{
+		GetWorldPartition()->InvalidateVisiblityTransactionRequestCachedVariables();
+	}
+}
+
 const AServerStreamingLevelsVisibility* UWorld::GetServerStreamingLevelsVisibility() const
 {
 	check(!IsValid(ServerStreamingLevelsVisibility) || SupportsMakingVisibleTransactionRequests());
diff --git a/Engine/Source/Runtime/Engine/Private/WorldPartition/WorldPartition.cpp b/Engine/Source/Runtime/Engine/Private/WorldPartition/WorldPartition.cpp
--- a/Engine/Source/Runtime/Engine/Private/WorldPartition/WorldPartition.cpp
+++ b/Engine/Source/Runtime/Engine/Private/WorldPartition/WorldPartition.cpp
@@ -77,6 +77,21 @@ namespace WorldPartition
 #else
 	static const EConsoleVariableFlags ECVF_Runtime_ReadOnly = ECVF_ReadOnly;
 #endif
+
+	static void OnVisibilityTransactionCVarChanged(IConsoleVariable* CVar)
+	{
+		if (GEngine)
+		{
+			for (const FWorldContext& Context : GEngine->GetWorldContexts())
+			{
+				UWorld* World = Context.World();
+				if (World && World->IsGameWorld())
+				{
+					World->InvalidateVisiblityTransactionRequestCachedVariables();
+				}
+			}
+		}
+	}
 }
 
 #if WITH_EDITOR
@@ -130,7 +145,8 @@ FAutoConsoleVariableRef UWorldPartition::CVarUseMakingVisibleTransactionRequests
 	TEXT("wp.Runtime.UseMakingVisibleTransactionRequests"),
 	UWorldPartition::bUseMakingVisibleTransactionRequests,
 	TEXT("Whether the client should wait for the server to acknowledge visibility update before making partitioned world streaming levels visible.\n")
-	TEXT("Changing the value while the game is running won't be considered."),
+	TEXT("Changing the value while the game is considered experimental."),
+	FConsoleVariableDelegate::CreateStatic(&*WorldPartition::OnVisibilityTransactionCVarChanged),
 	WorldPartition::ECVF_Runtime_ReadOnly);
 
 bool UWorldPartition::bUseMakingInvisibleTransactionRequests = false;
@@ -138,7 +154,8 @@ FAutoConsoleVariableRef UWorldPartition::CVarUseMakingInvisibleTransactionReques
 	TEXT("wp.Runtime.UseMakingInvisibleTransactionRequests"),
 	UWorldPartition::bUseMakingInvisibleTransactionRequests,
 	TEXT("Whether the client should wait for the server to acknowledge visibility update before making partitioned world streaming levels invisible.\n")
-	TEXT("Changing the value while the game is running won't be considered."),
+	TEXT("Changing the value while the game is considered experimental."),
+	FConsoleVariableDelegate::CreateStatic(&*WorldPartition::OnVisibilityTransactionCVarChanged),
 	WorldPartition::ECVF_Runtime_ReadOnly);
 
 #if WITH_EDITOR
@@ -1455,6 +1472,12 @@ bool UWorldPartition::UseMakingInvisibleTransactionRequests() const
 	return bCachedUseMakingInvisibleTransactionRequests.Get(false);
 }
 
+void UWorldPartition::InvalidateVisiblityTransactionRequestCachedVariables()
+{
+	bCachedUseMakingVisibleTransactionRequests.Reset();
+	bCachedUseMakingInvisibleTransactionRequests.Reset();
+}
+
 int32 UWorldPartition::GetStreamingStateEpoch() const
 {
 	// Merge WorldPartition's StreamingStateEpoch and AWorldDataLayers DataLayersStateEpoch
diff --git a/Engine/Source/Runtime/Engine/Public/WorldPartition/WorldPartition.h b/Engine/Source/Runtime/Engine/Public/WorldPartition/WorldPartition.h
--- a/Engine/Source/Runtime/Engine/Public/WorldPartition/WorldPartition.h
+++ b/Engine/Source/Runtime/Engine/Public/WorldPartition/WorldPartition.h
@@ -473,6 +473,7 @@ public:
 	ENGINE_API bool IsServerStreamingOutEnabled() const;
 	ENGINE_API bool UseMakingVisibleTransactionRequests() const;
 	ENGINE_API bool UseMakingInvisibleTransactionRequests() const;
+	ENGINE_API void InvalidateVisiblityTransactionRequestCachedVariables();
 
 	ENGINE_API bool IsMainWorldPartition() const;
 
