[SoundWave] Audio Derived Data Diff Tags rehydrate assets during cook

#Synced-CL 51050950

diff --git a/Engine/Source/Runtime/Engine/Private/AudioDerivedData.cpp b/Engine/Source/Runtime/Engine/Private/AudioDerivedData.cpp
--- a/Engine/Source/Runtime/Engine/Private/AudioDerivedData.cpp
+++ b/Engine/Source/Runtime/Engine/Private/AudioDerivedData.cpp
@@ -25,6 +25,7 @@
 #include "AudioDecompress.h"
 #include "ISoundWaveCloudStreaming.h"
 #include "Serialization/MemoryReader.h"
+#include "IO/IoHash.h"
 
 #if WITH_EDITORONLY_DATA
 #include "Serialization/MemoryWriter.h"
@@ -229,21 +230,10 @@ static FString GetSoundWaveHash(const USoundWave& InWave, const ITargetPlatform*
 	return SoundWaveHash;
 }
 
-static FSHAHash GetSoundWaveRawDataHash(const USoundWave& InWave)
+static FString GetSoundWaveRawDataHash(const USoundWave& InWave)
 {
-	FSHAHash Out{};
 	const USoundWave::FEditorAudioBulkData& RawData = InWave.RawData;
-
-	const FSharedBuffer& Payload = RawData.GetPayload().Consume();
-	const uint64 NumBytes = Payload.GetSize();
-
-	if (NumBytes == 0)
-	{
-		return Out;
-	}
-
-	FSHA1::HashBuffer(Payload.GetData(), NumBytes, Out.Hash);
-	return Out;
+	return LexToString(RawData.RawData.GetPayloadId());
 }
 
 static FString GetWaveformTransformationHash(const USoundWave& InWave)
@@ -1549,7 +1539,7 @@ void FStreamedAudioPlatformData::Serialize(FArchive& Ar, USoundWave* Owner)
 		const FName Format = CookingTarget->GetWaveFormat(Owner);
 
 		/* note: must have "Diff_" prefix for the diff tool to pick up the tag as important for diffing */
-		CookTags->Add(Owner, "Diff_10_RawDataHash", GetSoundWaveRawDataHash(*Owner).ToString());
+		CookTags->Add(Owner, "Diff_10_RawDataHash", GetSoundWaveRawDataHash(*Owner));
 		CookTags->Add(Owner, "Diff_20_WaveTransformHash", FString(GetWaveformTransformationHash(*Owner)));
 		CookTags->Add(Owner, "Diff_30_FormatHash", Format.ToString());
 		CookTags->Add(Owner, "Diff_40_SoundWaveHash", FString(GetSoundWaveHash(*Owner, CookingTarget)));
