Added exclusive csv stat for PostNetInit, renamed Generic_RepNotifies to RepNotifies_Generic.

#Synced-CL 51065404

diff --git a/Engine/Source/Runtime/Engine/Private/DataChannel.cpp b/Engine/Source/Runtime/Engine/Private/DataChannel.cpp
--- a/Engine/Source/Runtime/Engine/Private/DataChannel.cpp
+++ b/Engine/Source/Runtime/Engine/Private/DataChannel.cpp
@@ -42,6 +42,7 @@
 #include "Net/NetPing.h"
 #include "PacketHandler.h"
 #include "ProfilingDebugging/AssetMetadataTrace.h"
+#include "ProfilingDebugging/CsvProfiler.h"
 
 #if UE_NET_REPACTOR_NAME_DEBUG
 #include "Net/NetNameDebug.h"
@@ -3495,6 +3496,7 @@ void UActorChannel::ProcessBunchInternal(FInBunch & Bunch)
 	// After all properties have been initialized, call PostNetInit. This should call BeginPlay() so initialization can be done with proper starting values.
 	if (Actor && bSpawnedNewActor)
 	{
+		CSV_SCOPED_TIMING_STAT_EXCLUSIVE(PostNetInit);
 		SCOPE_CYCLE_COUNTER(Stat_PostNetInit);
 		LLM_SCOPE_BYNAME(TEXT("UObject/NetworkPostInit"));
 		LLM_SCOPE_DYNAMIC_STAT_OBJECTPATH(Actor->GetPackage(), ELLMTagSet::Assets);
diff --git a/Engine/Source/Runtime/Engine/Private/DemoNetDriver.cpp b/Engine/Source/Runtime/Engine/Private/DemoNetDriver.cpp
--- a/Engine/Source/Runtime/Engine/Private/DemoNetDriver.cpp
+++ b/Engine/Source/Runtime/Engine/Private/DemoNetDriver.cpp
@@ -34,6 +34,7 @@
 #include "Net/NetworkGranularMemoryLogging.h"
 #include "HAL/LowLevelMemStats.h"
 #include "ProfilingDebugging/AssetMetadataTrace.h"
+#include "ProfilingDebugging/CsvProfiler.h"
 
 #include UE_INLINE_GENERATED_CPP_BY_NAME(DemoNetDriver)
 
@@ -3516,6 +3517,8 @@ void UDemoNetDriver::RespawnNecessaryNetStartupActors(TArray<AActor*>& SpawnedAc
 
 			// BeginPlay
 			{
+				CSV_SCOPED_TIMING_STAT_EXCLUSIVE(PostNetInit);
+				QUICK_SCOPE_CYCLE_COUNTER(Stat_PostNetInit);
 				LLM_SCOPE_BYNAME(TEXT("UObject/NetworkPostInit"));
 				LLM_SCOPE_DYNAMIC_STAT_OBJECTPATH(Actor->GetPackage(), ELLMTagSet::Assets);
 				LLM_SCOPE_DYNAMIC_STAT_OBJECTPATH(Actor->GetClass(), ELLMTagSet::AssetClasses);
diff --git a/Engine/Source/Runtime/Engine/Private/Net/Iris/ReplicationSystem/NetActorFactory.cpp b/Engine/Source/Runtime/Engine/Private/Net/Iris/ReplicationSystem/NetActorFactory.cpp
--- a/Engine/Source/Runtime/Engine/Private/Net/Iris/ReplicationSystem/NetActorFactory.cpp
+++ b/Engine/Source/Runtime/Engine/Private/Net/Iris/ReplicationSystem/NetActorFactory.cpp
@@ -24,6 +24,9 @@
 #include "Net/Iris/ReplicationSystem/EngineReplicationBridge.h"
 
 #include "ProfilingDebugging/AssetMetadataTrace.h"
+#include "ProfilingDebugging/CsvProfiler.h"
+
+#include "Stats/Stats.h"
 
 #include "UObject/Package.h"
 
@@ -464,6 +467,8 @@ void UNetActorFactory::PostInit(const FPostInitContext& Context)
 		// PostNetInit is only called for dynamic actors
 		if (Context.Handle.IsDynamic())
 		{
+			CSV_SCOPED_TIMING_STAT_EXCLUSIVE(PostNetInit);
+			QUICK_SCOPE_CYCLE_COUNTER(Stat_PostNetInit);
 			LLM_SCOPE_BYNAME(TEXT("UObject/NetworkPostInit"));
 			LLM_SCOPE_DYNAMIC_STAT_OBJECTPATH(Actor->GetPackage(), ELLMTagSet::Assets);
 			LLM_SCOPE_DYNAMIC_STAT_OBJECTPATH(Actor->GetClass(), ELLMTagSet::AssetClasses);
diff --git a/Engine/Source/Runtime/Engine/Private/RepLayout.cpp b/Engine/Source/Runtime/Engine/Private/RepLayout.cpp
--- a/Engine/Source/Runtime/Engine/Private/RepLayout.cpp
+++ b/Engine/Source/Runtime/Engine/Private/RepLayout.cpp
@@ -4671,7 +4671,7 @@ void FRepLayout::CallRepNotifies(FReceivingRepState* RepState, UObject* Object)
 		return;
 	}
 
-	CSV_SCOPED_TIMING_STAT_EXCLUSIVE(Generic_RepNotifies);
+	CSV_SCOPED_TIMING_STAT_EXCLUSIVE(RepNotifies_Generic);
 
 	FRepShadowDataBuffer ShadowData(RepState->StaticBuffer.GetData());
 	FRepObjectDataBuffer ObjectData(Object);
