StaticMesh:  Fix for nondeterministic optional index buffer output when cooking, between DDC hit and miss, wasting package size.  Streaming to DDC was dead stripping the optional index buffers, so they were missing on a DDC hit.  Fix includes disabling initial generation of optional index buffers in the static mesh builder and BuildFromMeshDescription utility function when the CVars are disabled, and cleaning up incorrect logic that attempted to handle dead stripping on load.  Dead stripping on load is still useful to handle cases where the CVar was set at cook time, but disabled at runtime, to allow A/B testing without needing to cook separate packages.  All permutations of the two CVars set and disabled at both cook and runtime were tested.  For caching purposes, the optional index buffer feature flags are included in the DDC key.

Note that the r.SupportDepthOnlyIndexBuffers CVar didn't actually work!  The depth only buffers were generated without checking the CVar, and once generated, they wouldn't be dead stripped from the package.  So this change doesn't merely fix the package size nondeterminism, it also saves net package size across the board (4% for the test package used during debugging of the issue).  The depth only index buffer specifically wasn't dead stripped when writing to DDC, because the "Discard" function on the index buffer is a nop when the target platform includes editor data, which applies to DDC.  The other buffers only happened to be dead stripped on DDC write due to other logic which wrote empty buffers from a dummy structure to the DDC in that case.

#Synced-CL 51072120

diff --git a/Engine/Source/Developer/MeshBuilder/Private/StaticMeshBuilder.cpp b/Engine/Source/Developer/MeshBuilder/Private/StaticMeshBuilder.cpp
--- a/Engine/Source/Developer/MeshBuilder/Private/StaticMeshBuilder.cpp
+++ b/Engine/Source/Developer/MeshBuilder/Private/StaticMeshBuilder.cpp
@@ -1860,6 +1860,12 @@ void BuildAllBufferOptimizations(FStaticMeshLODResources& StaticMeshLOD, const F
 {
 	TRACE_CPUPROFILER_EVENT_SCOPE(BuildAllBufferOptimizations);
 
+	static const TConsoleVariableData<int32>* CVarDepthOnlyIndexBuffers = IConsoleManager::Get().FindTConsoleVariableDataInt(TEXT("r.SupportDepthOnlyIndexBuffers"));
+	static const TConsoleVariableData<int32>* CVarReversedIndexBuffers = IConsoleManager::Get().FindTConsoleVariableDataInt(TEXT("r.SupportReversedIndexBuffers"));
+
+	const bool bDepthOnlyIndexBuffers = !!CVarDepthOnlyIndexBuffers->GetValueOnAnyThread();
+	const bool bReversedIndexBuffers = !!CVarReversedIndexBuffers->GetValueOnAnyThread();
+
 	if (StaticMeshLOD.AdditionalIndexBuffers == nullptr)
 	{
 		StaticMeshLOD.AdditionalIndexBuffers = new FAdditionalStaticMeshIndexBuffers();
@@ -1868,7 +1874,7 @@ void BuildAllBufferOptimizations(FStaticMeshLODResources& StaticMeshLOD, const F
 	const EIndexBufferStride::Type IndexBufferStride = bNeeds32BitIndices ? EIndexBufferStride::Force32Bit : EIndexBufferStride::Force16Bit;
 
 	// Build the reversed index buffer.
-	if (LODBuildSettings.bBuildReversedIndexBuffer)
+	if (LODBuildSettings.bBuildReversedIndexBuffer && bReversedIndexBuffers)
 	{
 		TArray<uint32> InversedIndices;
 		const int32 IndexCount = IndexBuffer.Num();
@@ -1889,6 +1895,7 @@ void BuildAllBufferOptimizations(FStaticMeshLODResources& StaticMeshLOD, const F
 
 	// Build the depth-only index buffer.
 	TArray<uint32> DepthOnlyIndices;
+	if (bDepthOnlyIndexBuffers)
 	{
 		BuildOptimizationHelper::BuildDepthOnlyIndexBuffer(
 			DepthOnlyIndices,
@@ -1906,7 +1913,7 @@ void BuildAllBufferOptimizations(FStaticMeshLODResources& StaticMeshLOD, const F
 	}
 
 	// Build the inversed depth only index buffer.
-	if (LODBuildSettings.bBuildReversedIndexBuffer)
+	if (LODBuildSettings.bBuildReversedIndexBuffer && bReversedIndexBuffers && bDepthOnlyIndexBuffers)
 	{
 		TArray<uint32> ReversedDepthOnlyIndices;
 		const int32 IndexCount = DepthOnlyIndices.Num();
diff --git a/Engine/Source/Runtime/Engine/Private/StaticMesh.cpp b/Engine/Source/Runtime/Engine/Private/StaticMesh.cpp
--- a/Engine/Source/Runtime/Engine/Private/StaticMesh.cpp
+++ b/Engine/Source/Runtime/Engine/Private/StaticMesh.cpp
@@ -622,8 +622,8 @@ void FStaticMeshLODResources::SerializeBuffers(FArchive& Ar, UStaticMesh* OwnerS
 	Ar.UsingCustomVersion(FUE5ReleaseStreamObjectVersion::GUID);
 
 	// If the index buffers have already been initialized, do not change the behavior since the RHI resource pointer may have been cached somewhere already.
-	const bool bEnableDepthOnlyIndexBuffer = bHasDepthOnlyIndices || (CVarSupportDepthOnlyIndexBuffers.GetValueOnAnyThread() == 1);
-	const bool bEnableReversedIndexBuffer = bHasReversedIndices || bHasReversedDepthOnlyIndices || (CVarSupportReversedIndexBuffers.GetValueOnAnyThread() == 1);
+	const bool bEnableDepthOnlyIndexBuffer = bHasDepthOnlyIndices || !!CVarSupportDepthOnlyIndexBuffers.GetValueOnAnyThread();
+	const bool bEnableReversedIndexBuffer = bHasReversedIndices || bHasReversedDepthOnlyIndices || !!CVarSupportReversedIndexBuffers.GetValueOnAnyThread();
 
 	// See if the mesh wants to keep resources CPU accessible
 	bool bMeshCPUAcces = OwnerStaticMesh ? OwnerStaticMesh->bAllowCPUAccess : false;
@@ -657,27 +657,24 @@ void FStaticMeshLODResources::SerializeBuffers(FArchive& Ar, UStaticMesh* OwnerS
 	const bool bSerializeWireframeIndexBuffer = !StripFlags.IsEditorDataStripped();
 	const bool bSerializeRayTracingGeometry = !StripFlags.IsClassDataStripped(CDSF_RayTracingResources);
 
-	FAdditionalStaticMeshIndexBuffers DummyBuffers;
-	FAdditionalStaticMeshIndexBuffers* SerializedAdditionalIndexBuffers = &DummyBuffers;
-	if ((bEnableDepthOnlyIndexBuffer || bEnableReversedIndexBuffer) && (bSerializeReversedIndexBuffer || bSerializeWireframeIndexBuffer || bEnableDepthOnlyIndexBuffer))
+	if ((bSerializeReversedIndexBuffer && bEnableReversedIndexBuffer) || bSerializeWireframeIndexBuffer)
 	{
 		if (AdditionalIndexBuffers == nullptr)
 		{
 			AdditionalIndexBuffers = new FAdditionalStaticMeshIndexBuffers();
 		}
-		SerializedAdditionalIndexBuffers = AdditionalIndexBuffers;
 	}
 
+	// If the reversed index buffer feature is disabled, we serialize with dummy buffers (writes empty result if saving, reads and discards if loading).
+	// The separate structure is used to avoid the need to dynamically allocate AdditionalIndexBuffers to hold a result we are discarding.
+	FAdditionalStaticMeshIndexBuffers DummyBuffers;
+	FAdditionalStaticMeshIndexBuffers* ReversedAdditionalIndexBuffers = bEnableReversedIndexBuffer ? AdditionalIndexBuffers : &DummyBuffers;
+
 	if (bSerializeReversedIndexBuffer)
 	{
-		SerializedAdditionalIndexBuffers->ReversedIndexBuffer.Serialize(Ar, bNeedsCPUAccess);
-		AccumIndexBufferSize(SerializedAdditionalIndexBuffers->ReversedIndexBuffer, OutBuffersSize.ReversedIBsSize);
-		AccumIndexBufferSize(SerializedAdditionalIndexBuffers->ReversedIndexBuffer, OutBuffersSize.SerializedBuffersSize);
-		if (!bEnableReversedIndexBuffer)
-		{
-			SerializedAdditionalIndexBuffers->ReversedIndexBuffer.Discard();
-			SerializedAdditionalIndexBuffers->ReversedIndexBuffer.ClearMetaData();
-		}
+		ReversedAdditionalIndexBuffers->ReversedIndexBuffer.Serialize(Ar, bNeedsCPUAccess);
+		AccumIndexBufferSize(ReversedAdditionalIndexBuffers->ReversedIndexBuffer, OutBuffersSize.ReversedIBsSize);
+		AccumIndexBufferSize(ReversedAdditionalIndexBuffers->ReversedIndexBuffer, OutBuffersSize.SerializedBuffersSize);
 	}
 
 	DepthOnlyIndexBuffer.Serialize(Ar, bNeedsCPUAccess);
@@ -691,20 +688,20 @@ void FStaticMeshLODResources::SerializeBuffers(FArchive& Ar, UStaticMesh* OwnerS
 
 	if (bSerializeReversedIndexBuffer)
 	{
-		SerializedAdditionalIndexBuffers->ReversedDepthOnlyIndexBuffer.Serialize(Ar, bNeedsCPUAccess);
-		AccumIndexBufferSize(SerializedAdditionalIndexBuffers->ReversedDepthOnlyIndexBuffer, OutBuffersSize.ReversedIBsSize);
-		AccumIndexBufferSize(SerializedAdditionalIndexBuffers->ReversedDepthOnlyIndexBuffer, OutBuffersSize.SerializedBuffersSize);
-		if (!bEnableReversedIndexBuffer)
+		ReversedAdditionalIndexBuffers->ReversedDepthOnlyIndexBuffer.Serialize(Ar, bNeedsCPUAccess);
+		AccumIndexBufferSize(ReversedAdditionalIndexBuffers->ReversedDepthOnlyIndexBuffer, OutBuffersSize.ReversedIBsSize);
+		AccumIndexBufferSize(ReversedAdditionalIndexBuffers->ReversedDepthOnlyIndexBuffer, OutBuffersSize.SerializedBuffersSize);
+		if (!bEnableDepthOnlyIndexBuffer)
 		{
-			SerializedAdditionalIndexBuffers->ReversedDepthOnlyIndexBuffer.Discard();
-			SerializedAdditionalIndexBuffers->ReversedDepthOnlyIndexBuffer.ClearMetaData();
+			ReversedAdditionalIndexBuffers->ReversedDepthOnlyIndexBuffer.Discard();
+			ReversedAdditionalIndexBuffers->ReversedDepthOnlyIndexBuffer.ClearMetaData();
 		}
 	}
 
 	if (bSerializeWireframeIndexBuffer)
 	{
-		SerializedAdditionalIndexBuffers->WireframeIndexBuffer.Serialize(Ar, bNeedsCPUAccess);
-		AccumIndexBufferSize(SerializedAdditionalIndexBuffers->WireframeIndexBuffer, OutBuffersSize.SerializedBuffersSize);
+		AdditionalIndexBuffers->WireframeIndexBuffer.Serialize(Ar, bNeedsCPUAccess);
+		AccumIndexBufferSize(AdditionalIndexBuffers->WireframeIndexBuffer, OutBuffersSize.SerializedBuffersSize);
 	}
 
 	if (Ar.IsLoading() && Ar.CustomVer(FUE5ReleaseStreamObjectVersion::GUID) < FUE5ReleaseStreamObjectVersion::RemovingTessellation && !StripFlags.IsClassDataStripped(CDSF_AdjacencyData_DEPRECATED))
@@ -739,10 +736,10 @@ void FStaticMeshLODResources::SerializeBuffers(FArchive& Ar, UStaticMesh* OwnerS
 	// Update metadata but only if serialization was successful. This needs to be done now because on cooked platform, indices are discarded after RHIInit.
 	if (!Ar.IsError())
 	{
-		bHasWireframeIndices = AdditionalIndexBuffers && bSerializeWireframeIndexBuffer && SerializedAdditionalIndexBuffers->WireframeIndexBuffer.GetNumIndices() != 0;
+		bHasWireframeIndices = AdditionalIndexBuffers && bSerializeWireframeIndexBuffer && AdditionalIndexBuffers->WireframeIndexBuffer.GetNumIndices() != 0;
 		bHasDepthOnlyIndices = DepthOnlyIndexBuffer.GetNumIndices() != 0;
-		bHasReversedIndices = AdditionalIndexBuffers && bSerializeReversedIndexBuffer && SerializedAdditionalIndexBuffers->ReversedIndexBuffer.GetNumIndices() != 0;
-		bHasReversedDepthOnlyIndices = AdditionalIndexBuffers && bSerializeReversedIndexBuffer && SerializedAdditionalIndexBuffers->ReversedDepthOnlyIndexBuffer.GetNumIndices() != 0;
+		bHasReversedIndices = AdditionalIndexBuffers && bSerializeReversedIndexBuffer && AdditionalIndexBuffers->ReversedIndexBuffer.GetNumIndices() != 0;
+		bHasReversedDepthOnlyIndices = AdditionalIndexBuffers && bSerializeReversedIndexBuffer && AdditionalIndexBuffers->ReversedDepthOnlyIndexBuffer.GetNumIndices() != 0;
 		bHasColorVertexData = VertexBuffers.ColorVertexBuffer.GetNumVertices() > 0;
 		DepthOnlyNumTriangles = DepthOnlyIndexBuffer.GetNumIndices() / 3;
 	}
@@ -785,24 +782,26 @@ void FStaticMeshLODResources::SerializeAvailabilityInfo(FArchive& Ar)
 	VertexBuffers.ColorVertexBuffer.SerializeMetaData(Ar);
 	IndexBuffer.SerializeMetaData(Ar);
 
-	FAdditionalStaticMeshIndexBuffers DummyBuffers;
-	FAdditionalStaticMeshIndexBuffers* SerializedAdditionalIndexBuffers = &DummyBuffers;
-	if ((bEnableDepthOnlyIndexBuffer || bEnableReversedIndexBuffer) && (bHasReversedIndices || bHasWireframeIndices || bHasDepthOnlyIndices))
+	if ((bHasReversedIndices && bEnableReversedIndexBuffer) || bHasWireframeIndices)
 	{
 		if (AdditionalIndexBuffers == nullptr)
 		{
 			AdditionalIndexBuffers = new FAdditionalStaticMeshIndexBuffers();
 		}
-		SerializedAdditionalIndexBuffers = AdditionalIndexBuffers;
 	}
 
-	SerializedAdditionalIndexBuffers->ReversedIndexBuffer.SerializeMetaData(Ar);
-	if (!bHasReversedIndices)
+	// If the reversed index buffer or wireframe features are disabled, we serialize with dummy buffers.
+	FAdditionalStaticMeshIndexBuffers DummyBuffers;
+	FAdditionalStaticMeshIndexBuffers* ReversedAdditionalIndexBuffers = bHasReversedIndices && bEnableReversedIndexBuffer ? AdditionalIndexBuffers : &DummyBuffers;
+	FAdditionalStaticMeshIndexBuffers* WireframeAdditionalIndexBuffers = bHasWireframeIndices ? AdditionalIndexBuffers : &DummyBuffers;
+
+	ReversedAdditionalIndexBuffers->ReversedIndexBuffer.SerializeMetaData(Ar);
+	if (!bHasReversedIndices && AdditionalIndexBuffers)
 	{
 		// Reversed indices are either stripped during cook or will be stripped on load.
 		// In either case, clear CachedNumIndices to show that the buffer will be empty after actual loading
-		SerializedAdditionalIndexBuffers->ReversedIndexBuffer.Discard();
-		SerializedAdditionalIndexBuffers->ReversedIndexBuffer.ClearMetaData();
+		AdditionalIndexBuffers->ReversedIndexBuffer.Discard();
+		AdditionalIndexBuffers->ReversedIndexBuffer.ClearMetaData();
 	}
 	DepthOnlyIndexBuffer.SerializeMetaData(Ar);
 	if (!bHasDepthOnlyIndices)
@@ -810,17 +809,17 @@ void FStaticMeshLODResources::SerializeAvailabilityInfo(FArchive& Ar)
 		DepthOnlyIndexBuffer.Discard();
 		DepthOnlyIndexBuffer.ClearMetaData();
 	}
-	SerializedAdditionalIndexBuffers->ReversedDepthOnlyIndexBuffer.SerializeMetaData(Ar);
-	if (!bHasReversedDepthOnlyIndices)
+	ReversedAdditionalIndexBuffers->ReversedDepthOnlyIndexBuffer.SerializeMetaData(Ar);
+	if (!bHasReversedDepthOnlyIndices && AdditionalIndexBuffers)
 	{
-		SerializedAdditionalIndexBuffers->ReversedDepthOnlyIndexBuffer.Discard();
-		SerializedAdditionalIndexBuffers->ReversedDepthOnlyIndexBuffer.ClearMetaData();
+		AdditionalIndexBuffers->ReversedDepthOnlyIndexBuffer.Discard();
+		AdditionalIndexBuffers->ReversedDepthOnlyIndexBuffer.ClearMetaData();
 	}
-	SerializedAdditionalIndexBuffers->WireframeIndexBuffer.SerializeMetaData(Ar);
-	if (!bHasWireframeIndices)
+	WireframeAdditionalIndexBuffers->WireframeIndexBuffer.SerializeMetaData(Ar);
+	if (!bHasWireframeIndices && AdditionalIndexBuffers)
 	{
-		SerializedAdditionalIndexBuffers->WireframeIndexBuffer.Discard();
-		SerializedAdditionalIndexBuffers->WireframeIndexBuffer.ClearMetaData();
+		AdditionalIndexBuffers->WireframeIndexBuffer.Discard();
+		AdditionalIndexBuffers->WireframeIndexBuffer.ClearMetaData();
 	}
 	if (Ar.IsLoading() && Ar.CustomVer(FUE5ReleaseStreamObjectVersion::GUID) < FUE5ReleaseStreamObjectVersion::RemovingTessellation)
 	{
@@ -3948,6 +3947,12 @@ static FString BuildStaticMeshDerivedDataKeySuffix(const ITargetPlatform* Target
 	KeySuffix.Append(TEXT("_arm64"));
 #endif
 
+	// Settings for optional index buffers.
+	const int32 bDepthOnlyIndexBuffers = CVarSupportDepthOnlyIndexBuffers.GetValueOnAnyThread() ? 1 : 0;
+	const int32 bReversedIndexBuffers = CVarSupportReversedIndexBuffers.GetValueOnAnyThread() ? 1 : 0;
+
+	KeySuffix += FString::Printf(TEXT("_IB%x"), (bDepthOnlyIndexBuffers << 0) | (bReversedIndexBuffers << 1));
+
 	return KeySuffix;
 }
 
@@ -8591,41 +8596,53 @@ void UStaticMesh::BuildFromMeshDescription(const FMeshDescription& MeshDescripti
 	check(IndexBufferIndex == NumTriangles * 3);
 
 	LODResources.IndexBuffer.SetIndices(IndexBuffer, IndexBufferStride);
+	LODResources.bHasColorVertexData = true;
+
+	const bool bDepthOnlyIndexBuffers = !!CVarSupportDepthOnlyIndexBuffers.GetValueOnAnyThread();
+	const bool bReversedIndexBuffers = !!CVarSupportReversedIndexBuffers.GetValueOnAnyThread();
 
 	// Fill depth only index buffer
 
 	TArray<uint32> DepthOnlyIndexBuffer(IndexBuffer);
-	for (uint32& Index : DepthOnlyIndexBuffer)
+	if (bDepthOnlyIndexBuffers)
 	{
-		// Compress all vertex instances into the same instance for each vertex
-		Index = MeshDescription.GetVertexVertexInstanceIDs(MeshDescription.GetVertexInstanceVertex(FVertexInstanceID(Index)))[0].GetValue();
-	}
+		for (uint32& Index : DepthOnlyIndexBuffer)
+		{
+			// Compress all vertex instances into the same instance for each vertex
+			Index = MeshDescription.GetVertexVertexInstanceIDs(MeshDescription.GetVertexInstanceVertex(FVertexInstanceID(Index)))[0].GetValue();
+		}
 
-	LODResources.bHasDepthOnlyIndices = true;
-	LODResources.DepthOnlyIndexBuffer.SetIndices(DepthOnlyIndexBuffer, IndexBufferStride);
-	LODResources.DepthOnlyNumTriangles = NumTriangles;
-	LODResources.bHasColorVertexData = true;
+		LODResources.bHasDepthOnlyIndices = true;
+		LODResources.DepthOnlyIndexBuffer.SetIndices(DepthOnlyIndexBuffer, IndexBufferStride);
+		LODResources.DepthOnlyNumTriangles = DepthOnlyIndexBuffer.Num() / 3;
+	}
 
 	// Fill reversed index buffer
-	TArray<uint32> ReversedIndexBuffer(IndexBuffer);
-	for (int32 ReversedIndexBufferIndex = 0; ReversedIndexBufferIndex < IndexBuffer.Num(); ReversedIndexBufferIndex += 3)
+	if (bReversedIndexBuffers)
 	{
-		Swap(ReversedIndexBuffer[ReversedIndexBufferIndex + 0], ReversedIndexBuffer[ReversedIndexBufferIndex + 2]);
-	}
+		TArray<uint32> ReversedIndexBuffer(IndexBuffer);
+		for (int32 ReversedIndexBufferIndex = 0; ReversedIndexBufferIndex < IndexBuffer.Num(); ReversedIndexBufferIndex += 3)
+		{
+			Swap(ReversedIndexBuffer[ReversedIndexBufferIndex + 0], ReversedIndexBuffer[ReversedIndexBufferIndex + 2]);
+		}
 
-	LODResources.AdditionalIndexBuffers = new FAdditionalStaticMeshIndexBuffers();
-	LODResources.bHasReversedIndices = true;
-	LODResources.AdditionalIndexBuffers->ReversedIndexBuffer.SetIndices(ReversedIndexBuffer, IndexBufferStride);
+		LODResources.AdditionalIndexBuffers = new FAdditionalStaticMeshIndexBuffers();
+		LODResources.bHasReversedIndices = true;
+		LODResources.AdditionalIndexBuffers->ReversedIndexBuffer.SetIndices(ReversedIndexBuffer, IndexBufferStride);
 
-	// Fill reversed depth index buffer
-	TArray<uint32> ReversedDepthOnlyIndexBuffer(DepthOnlyIndexBuffer);
-	for (int32 ReversedIndexBufferIndex = 0; ReversedIndexBufferIndex < IndexBuffer.Num(); ReversedIndexBufferIndex += 3)
-	{
-		Swap(ReversedDepthOnlyIndexBuffer[ReversedIndexBufferIndex + 0], ReversedDepthOnlyIndexBuffer[ReversedIndexBufferIndex + 2]);
-	}
+		if (bDepthOnlyIndexBuffers)
+		{
+			// Fill reversed depth index buffer
+			TArray<uint32> ReversedDepthOnlyIndexBuffer(DepthOnlyIndexBuffer);
+			for (int32 ReversedIndexBufferIndex = 0; ReversedIndexBufferIndex < IndexBuffer.Num(); ReversedIndexBufferIndex += 3)
+			{
+				Swap(ReversedDepthOnlyIndexBuffer[ReversedIndexBufferIndex + 0], ReversedDepthOnlyIndexBuffer[ReversedIndexBufferIndex + 2]);
+			}
 
-	LODResources.bHasReversedDepthOnlyIndices = true;
-	LODResources.AdditionalIndexBuffers->ReversedDepthOnlyIndexBuffer.SetIndices(ReversedIndexBuffer, IndexBufferStride);
+			LODResources.bHasReversedDepthOnlyIndices = true;
+			LODResources.AdditionalIndexBuffers->ReversedDepthOnlyIndexBuffer.SetIndices(ReversedDepthOnlyIndexBuffer, IndexBufferStride);
+		}
+	}
 }
 
 void UStaticMesh::ExportStaticMeshLOD(const FStaticMeshLODResources& InLODResources, FMeshDescription& OutMeshDescription) const
