Added debug function to convert a FOverriddenPropertySet into a string.
Added debug CVar command that prints the override state of an object.

#rnx

#Synced-CL 51051084

diff --git a/Engine/Source/Runtime/CoreUObject/Private/UObject/OverriddenPropertySet.cpp b/Engine/Source/Runtime/CoreUObject/Private/UObject/OverriddenPropertySet.cpp
--- a/Engine/Source/Runtime/CoreUObject/Private/UObject/OverriddenPropertySet.cpp
+++ b/Engine/Source/Runtime/CoreUObject/Private/UObject/OverriddenPropertySet.cpp
@@ -7,6 +7,7 @@
 #include "Serialization/ArchiveSerializedPropertyChain.h"
 #include "UObject/PropertyOptional.h"
 #include "Misc/ScopeExit.h"
+#include "Misc/StringBuilder.h"
 #include "UObject/InstanceDataObjectUtils.h"
 #include "UObject/PropertyBagRepository.h"
 #include "UObject/UObjectArchetypeHelper.h"
@@ -2077,3 +2078,28 @@ void* FScopedPropertyValue::GetData() const
 {
 	return Data;
 }
+
+static void ToDebugString_Internal(const FOverriddenPropertyNode& Node, FStringBuilderBase& StrBuilder, int StackDepth = 0)
+{
+	for (int D = 0; D < StackDepth; ++D)
+	{
+		StrBuilder.Append(TEXT("  "));
+	}
+
+	StrBuilder.Append(GetOverriddenOperationString(Node.GetOperation()));
+	StrBuilder.Append(TEXT(" "));
+	StrBuilder.Append(Node.GetNodeID().ToDebugString());
+	StrBuilder.Append(TEXT("\n"));
+
+	for (const FOverriddenPropertyNode& SubNode : Node.GetSubPropertyNodes())
+	{
+		ToDebugString_Internal(SubNode, StrBuilder, StackDepth + 1);
+	}
+}
+
+FString FOverriddenPropertySet::ToDebugString() const
+{
+	TStringBuilder<32> Builder;
+	ToDebugString_Internal(RootNode, Builder);
+	return Builder.ToString();
+}
diff --git a/Engine/Source/Runtime/CoreUObject/Public/UObject/OverriddenPropertySet.h b/Engine/Source/Runtime/CoreUObject/Public/UObject/OverriddenPropertySet.h
--- a/Engine/Source/Runtime/CoreUObject/Public/UObject/OverriddenPropertySet.h
+++ b/Engine/Source/Runtime/CoreUObject/Public/UObject/OverriddenPropertySet.h
@@ -320,7 +320,7 @@ struct FOverriddenPropertyNodeID
 
 	FString ToDebugString() const
 	{
-		return Path.ToString() + FString::Printf(TEXT("(0x%p)"), Object.Get());
+		return Path.ToString() + FString::Printf(TEXT(" (%s 0x%p)"), Object.IsValid() ? *Object->GetName() : TEXT("NULL"), Object.Get());
 	}
 
 	bool IsValid() const
@@ -580,6 +580,11 @@ public:
 	 * @param FromOverriddenProperties overridable properties to restore from*/
 	UE_API void RestoreOverriddenState(const FOverriddenPropertySet& FromOverriddenProperties);
 
+	/**
+	 * Converts the state of all nodes in the property set into a string (for debug purposes).
+	 * @return a string representing all the nodes in the set */
+	UE_API FString ToDebugString() const;
+
 protected:
 
 	UE_API EOverriddenPropertyOperation GetOverriddenPropertyOperation(const FOverriddenPropertyNode* ParentPropertyNode, FPropertyVisitorPath::Iterator PropertyIterator, bool* bOutInheritedOperation, const void* Data) const;
diff --git a/Engine/Source/Runtime/Engine/Private/PrintObjectUtils.cpp b/Engine/Source/Runtime/Engine/Private/PrintObjectUtils.cpp
--- a/Engine/Source/Runtime/Engine/Private/PrintObjectUtils.cpp
+++ b/Engine/Source/Runtime/Engine/Private/PrintObjectUtils.cpp
@@ -7,6 +7,8 @@
 #include "Containers/Set.h"
 #include "HAL/IConsoleManager.h"
 #include "UObject/Object.h"
+#include "UObject/OverridableManager.h"
+#include "UObject/OverriddenPropertySet.h"
 #include "UObject/Class.h"
 #include "UObject/UnrealType.h"
 #include "UObject/UObjectHash.h"
@@ -311,6 +313,40 @@ void PrintObjectProperties(UObject* Object, EPrintObjectFlag Flags, FOutputDevic
 	}
 }
 
+void PrintObjectOverrides(UObject* Object, EPrintObjectFlag Flags, FOutputDevice* InOutputDevice)
+{
+	FOutputDevice* OutputDevice = InOutputDevice ? InOutputDevice : GLog;
+	if (!OutputDevice)
+	{
+		return;
+	}
+
+	if (Object)
+	{
+		bool bVerbose = EnumHasAnyFlags(Flags, EPrintObjectFlag::Verbose);
+
+		LogObjectState(Object, OutputDevice, bVerbose, TEXT("Printing overrides for object "));
+
+		FOverridableManager& OverridableManager = FOverridableManager::Get();
+
+		const FOverriddenPropertySet* Overrides = OverridableManager.GetOverriddenProperties(Object);
+
+		if (Overrides)
+		{
+			FString OverrideStr = Overrides->ToDebugString();
+			OutputDevice->Logf(TEXT("%s"), *OverrideStr);
+		}
+		else
+		{
+			OutputDevice->Logf(TEXT("No overrides found for object: %s"), *Object->GetPathName());
+		}
+	}
+	else
+	{
+		OutputDevice->Logf(TEXT("PrintObjectOverrides: NULL object"));
+	}
+}
+
 void PrintObjectArchetype(UObject* Object, EPrintObjectFlag Flags, FOutputDevice* InOutputDevice)
 {
 	FOutputDevice* OutputDevice = InOutputDevice ? InOutputDevice : GLog;
@@ -552,6 +588,35 @@ FAutoConsoleCommandWithArgsAndOutputDevice CVarCommandPrintObjectProperties(
 		}
 	}));
 
+FAutoConsoleCommandWithArgsAndOutputDevice CVarCommandPrintObjectOverrides(
+	TEXT("Obj.DumpOverrides"),
+	TEXT("Lists the overridable serialization overrides of an object (the object must be specified as a name or path, e.g., /MyLevel/MyLevel.MyLevel:PersistentLevel.MyActor)."),
+	FConsoleCommandWithArgsAndOutputDeviceDelegate::CreateLambda([](const TArray<FString>& Args, FOutputDevice& OutputDevice)
+	{
+		if (Args.Num() > 0)
+		{
+			FString ObjectPath;
+			EPrintObjectFlag Flags;
+			ParseObjectIdAndFlags(Args, ObjectPath, Flags);
+
+			UObject* Object = FindFirstObjectSafe<UObject>(*ObjectPath);
+
+			if (Object)
+			{
+				PrintObjectOverrides(Object, Flags, &OutputDevice);
+			}
+			else
+			{
+				OutputDevice.Logf(TEXT("PrintObjectOverrides: failed to find any objects for path: '%s'"), *ObjectPath);
+			}
+		}
+		else
+		{
+			OutputDevice.Logf(TEXT("PrintObjectOverrides: no object path specified (example usage: Obj.DumpOverrides /MyLevel/MyLevel.MyLevel:PersistentLevel.MyActor)"));
+			LogHelp(OutputDevice);
+		}
+	}));
+
 FAutoConsoleCommandWithArgsAndOutputDevice CVarCommandPrintObjectArchetype(
 	TEXT("Obj.DumpArchetype"),
 	TEXT("Outputs an object's archetype (the object must be specified as a name or path, e.g., /MyLevel/MyLevel.MyLevel:PersistentLevel.MyActor)."),
diff --git a/Engine/Source/Runtime/Engine/Public/PrintObjectUtils.h b/Engine/Source/Runtime/Engine/Public/PrintObjectUtils.h
--- a/Engine/Source/Runtime/Engine/Public/PrintObjectUtils.h
+++ b/Engine/Source/Runtime/Engine/Public/PrintObjectUtils.h
@@ -73,6 +73,15 @@ UE_INTERNAL void PrintStructProperties(UStruct* Struct, void* StructData, EPrint
  */
 UE_INTERNAL void PrintObjectProperties(UObject* Object, EPrintObjectFlag Flags = EPrintObjectFlag::None, FOutputDevice* OutputDevice = nullptr);
 
+/**
+ * Lists any overridable serialization overrides for a given object.
+ *
+ * @param Object The object to use.
+ * @param Flags For specifying optional behaviors.
+ * @param OutputDevice The output device to print to (can be NULL).
+ */
+UE_INTERNAL void PrintObjectOverrides(UObject* Object, EPrintObjectFlag Flags = EPrintObjectFlag::None, FOutputDevice* OutputDevice = nullptr);
+
 /**
  * Outputs an object's archetype.
  *
