Restore old texture coordinate input layout code again because of possible sampling out of bound issues on the GPU

#Synced-CL 51088011

diff --git a/Engine/Source/Runtime/Engine/Private/LocalVertexFactory.cpp b/Engine/Source/Runtime/Engine/Private/LocalVertexFactory.cpp
--- a/Engine/Source/Runtime/Engine/Private/LocalVertexFactory.cpp
+++ b/Engine/Source/Runtime/Engine/Private/LocalVertexFactory.cpp
@@ -558,48 +558,22 @@ void AddStaticMeshTextureCoordinateElements(
 	FVertexDeclarationElementList& Elements,
 	FVertexFactory::FVertexStreamList& InOutStreams)
 {
-	if (TextureCoordinates.IsEmpty())
+	for (int32 CoordinateIndex = 0; CoordinateIndex < TextureCoordinates.Num(); ++CoordinateIndex)
 	{
-		return;
-	}
-	
-#if DO_CHECK 
-	// Validate all texture coordinates share the same stream data
-	for (int32 TextureCoordIndex = 1; TextureCoordIndex < TextureCoordinates.Num(); ++TextureCoordIndex)
-	{
-		checkf(TextureCoordinates[0].VertexBuffer == TextureCoordinates[TextureCoordIndex].VertexBuffer, TEXT("Texture coordinate %d has a different vertex buffer than index 0"), TextureCoordIndex);
-		checkf(TextureCoordinates[0].Stride == TextureCoordinates[TextureCoordIndex].Stride, TEXT("Texture coordinate %d has a different stride than index 0"), TextureCoordIndex);
-		checkf(TextureCoordinates[0].StreamOffset == TextureCoordinates[TextureCoordIndex].StreamOffset, TEXT("Texture coordinate %d has a different stream offset than index 0"), TextureCoordIndex);
-		checkf(TextureCoordinates[0].VertexStreamUsage == TextureCoordinates[TextureCoordIndex].VertexStreamUsage, TEXT("Texture coordinate %d has a different vertex stream usage than index 0"), TextureCoordIndex);
-	}
-#endif // DO_CHECK
-	
-	EVertexElementType FirstElementType = TextureCoordinates[0].Type;
-	
-	EVertexElementType UVDoubleWideVertexElementType = VET_None;
-	uint32 DoubleUVSizeInBytes = 0;
-	if (FirstElementType == VET_Float2 || FirstElementType == VET_Float4)
-	{
-		DoubleUVSizeInBytes = sizeof(TStaticMeshVertexUVsTypeSelector<EStaticMeshVertexUVType::HighPrecision>::UVsTypeT) * 2;
-		UVDoubleWideVertexElementType = VET_Float4;
-	}
-	else
-	{
-		check(FirstElementType == VET_Half2 || FirstElementType == VET_Half4)
-		DoubleUVSizeInBytes = sizeof(TStaticMeshVertexUVsTypeSelector<EStaticMeshVertexUVType::Default>::UVsTypeT) * 2;
-		UVDoubleWideVertexElementType = VET_Half4;
+		Elements.Add(FVertexFactory::AccessStreamComponent(
+			TextureCoordinates[CoordinateIndex],
+			BaseTexCoordAttribute + CoordinateIndex,
+			InOutStreams
+		));
 	}
-	
-	uint8 TexCoordStreamIndex = FVertexFactory::AddVertexStream(TextureCoordinates[0], InOutStreams);
-	
-	for (int32 CoordinateIndex = 0; CoordinateIndex < MAX_STATIC_TEXCOORDS / 2; ++CoordinateIndex)
+
+	for (int32 CoordinateIndex = TextureCoordinates.Num(); CoordinateIndex < MAX_STATIC_TEXCOORDS / 2; ++CoordinateIndex)
 	{
-		Elements.Add(FVertexElement(
-			TexCoordStreamIndex,
-			DoubleUVSizeInBytes * CoordinateIndex,
-			UVDoubleWideVertexElementType,
+		Elements.Add(FVertexFactory::AccessStreamComponent(
+			TextureCoordinates[TextureCoordinates.Num() - 1],
 			BaseTexCoordAttribute + CoordinateIndex,
-			TextureCoordinates[0].Stride));
+			InOutStreams
+		));
 	}
 }
 
diff --git a/Engine/Source/Runtime/Engine/Private/Rendering/StaticMeshVertexBuffer.cpp b/Engine/Source/Runtime/Engine/Private/Rendering/StaticMeshVertexBuffer.cpp
--- a/Engine/Source/Runtime/Engine/Private/Rendering/StaticMeshVertexBuffer.cpp
+++ b/Engine/Source/Runtime/Engine/Private/Rendering/StaticMeshVertexBuffer.cpp
@@ -101,7 +101,7 @@ void FStaticMeshVertexBuffer::Init(uint32 InNumVertices, uint32 InNumTexCoords,
 	// Allocate the vertex data buffer.
 	TangentsData->ResizeBuffer(NumVertices);
 	TangentsDataPtr = NumVertices ? TangentsData->GetDataPointer() : nullptr;
-	TexcoordData->ResizeBuffer(NumVertices);
+	TexcoordData->ResizeBuffer(NumVertices * GetNumTexCoords());
 	TexcoordDataPtr = NumVertices ? TexcoordData->GetDataPointer() : nullptr;
 }
 
@@ -165,7 +165,7 @@ void FStaticMeshVertexBuffer::Init(const FStaticMeshVertexBuffer& InVertexBuffer
 			check(GetNumTexCoords() == InVertexBuffer.GetNumTexCoords());
 			const uint8* InData = InVertexBuffer.TexcoordDataPtr;
 
-			TexcoordData->ResizeBuffer(NumVertices);
+			TexcoordData->ResizeBuffer(NumVertices * GetNumTexCoords());
 			TexcoordDataPtr = TexcoordData->GetDataPointer();
 			FMemory::Memcpy(TexcoordDataPtr, InData, TexcoordData->GetStride() * NumVertices * GetNumTexCoords());
 		}
@@ -191,7 +191,7 @@ void FStaticMeshVertexBuffer::AppendVertices( const FStaticMeshBuildVertex* Vert
 		NumVertices += NumVerticesToAppend;
 
 		TangentsData->ResizeBuffer(NumVertices);
-		TexcoordData->ResizeBuffer(NumVertices);
+		TexcoordData->ResizeBuffer(NumVertices * GetNumTexCoords());
 
 		if( NumVertices > 0 )
 		{
@@ -306,7 +306,7 @@ FBufferRHIRef FStaticMeshVertexBuffer::CreateTangentsRHIBuffer(FRHICommandListBa
 
 FBufferRHIRef FStaticMeshVertexBuffer::CreateTexCoordRHIBuffer(FRHICommandListBase& RHICmdList)
 {
-	uint32 ResourceCount = GetNumTexCoordInVertexBuffer(GetNumTexCoords(), GetNumVertices());
+	uint32 ResourceCount = GetNumVertices() * GetNumTexCoords();
 	return CreateRHIBuffer(RHICmdList, TexcoordData, ResourceCount, BUF_Static | BUF_ShaderResource, TEXT("TexCoordRHIBuffer"));
 }
 
@@ -406,13 +406,13 @@ void FStaticMeshVertexBuffer::AllocateData(bool bNeedsCPUAccess /*= true*/)
 	{
 		typedef TStaticMeshVertexUVsDatum<typename TStaticMeshVertexUVsTypeSelector<EStaticMeshVertexUVType::HighPrecision>::UVsTypeT> UVType;
 		TexcoordStride = sizeof(UVType);
-		TexcoordData = new TStaticTexCoordVertexData<UVType>(NumTexCoords, bNeedsCPUAccess);
+		TexcoordData = new TStaticMeshVertexData<UVType>(bNeedsCPUAccess);
 	}
 	else
 	{
 		typedef TStaticMeshVertexUVsDatum<typename TStaticMeshVertexUVsTypeSelector<EStaticMeshVertexUVType::Default>::UVsTypeT> UVType;
 		TexcoordStride = sizeof(UVType);
-		TexcoordData = new TStaticTexCoordVertexData<UVType>(NumTexCoords, bNeedsCPUAccess);
+		TexcoordData = new TStaticMeshVertexData<UVType>(bNeedsCPUAccess);
 	}
 }
 
