Fix crash when attempting to preview TextureObject material graph node.
Crash happens because when compiling for preview, Substrate tries to compile input into EmissiveColor and cast it to float3.
Because TextureObject compiles to a uniform, cast fails and a check gets triggered, which crashes the editor.

[RN] minor | Rendering Architecture
#Synced-CL 51008691

diff --git a/Engine/Source/Runtime/Engine/Private/Materials/HLSLMaterialTranslator.cpp b/Engine/Source/Runtime/Engine/Private/Materials/HLSLMaterialTranslator.cpp
--- a/Engine/Source/Runtime/Engine/Private/Materials/HLSLMaterialTranslator.cpp
+++ b/Engine/Source/Runtime/Engine/Private/Materials/HLSLMaterialTranslator.cpp
@@ -12900,7 +12900,14 @@ int32 FHLSLMaterialTranslator::SubstrateCreateAndRegisterNullMaterial()
 
 FString FHLSLMaterialTranslator::SubstrateGetCastParameterCode(int32 Index, EMaterialValueType DestType)
 {
-	int32 CastParameter = ForceCast(Index, DestType);
+	const int32 CastParameter = ForceCast(Index, DestType);
+
+	// When casting a Texture Object, it will fail to assign a texture uniform to float3
+	if (CastParameter == INDEX_NONE)
+	{
+		return FString();
+	}
+
 	return GetParameterCode(CastParameter);
 }
 
@@ -13306,14 +13313,18 @@ int32 FHLSLMaterialTranslator::SubstrateVolumetricFogCloudBSDF(int32 Albedo, int
 
 int32 FHLSLMaterialTranslator::SubstrateUnlitBSDF(int32 EmissiveColor, int32 TransmittanceColor, int32 Normal, FSubstrateOperator* PromoteToOperator)
 {
+	const FString EmissiveColorCode = SubstrateGetCastParameterCode(EmissiveColor, MCT_Float3);
+	const FString TransmittanceColorCode = SubstrateGetCastParameterCode(TransmittanceColor, MCT_Float3);
+	const FString NormalCode = SubstrateGetCastParameterCode(Normal, MCT_Float3);
+
 	if (PromoteToOperator)
 	{
 		return AddCodeChunk(
 			MCT_Substrate, TEXT("Parameters.%s.PromoteParameterBlendedBSDFToOperator(GetSubstrateUnlitBSDF(%s, %s, %s), %u, %u, %u, %u)"),
 			*Substrate::GetParametersSubstrateTreeName(SubstrateTranslatorData.CurrentSubstrateCompilationContext),
-			*SubstrateGetCastParameterCode(EmissiveColor,		MCT_Float3),
-			*SubstrateGetCastParameterCode(TransmittanceColor,	MCT_Float3),
-			*SubstrateGetCastParameterCode(Normal,				MCT_Float3),
+			*EmissiveColorCode,
+			*TransmittanceColorCode,
+			*NormalCode,
 			PromoteToOperator->Index,
 			PromoteToOperator->BSDFIndex,
 			PromoteToOperator->LayerDepth,
diff --git a/Engine/Source/Runtime/Engine/Private/Materials/MaterialExpressionSubstrate.cpp b/Engine/Source/Runtime/Engine/Private/Materials/MaterialExpressionSubstrate.cpp
--- a/Engine/Source/Runtime/Engine/Private/Materials/MaterialExpressionSubstrate.cpp
+++ b/Engine/Source/Runtime/Engine/Private/Materials/MaterialExpressionSubstrate.cpp
@@ -2054,10 +2054,14 @@ UMaterialExpressionSubstrateUnlitBSDF::UMaterialExpressionSubstrateUnlitBSDF(con
 int32 UMaterialExpressionSubstrateUnlitBSDF::Compile(class FMaterialCompiler* Compiler, int32 OutputIndex)
 {
 	FSubstrateOperator& SubstrateOperator = Compiler->SubstrateCompilationGetOperator(Compiler->SubstrateTreeStackGetPathUniqueId());
+	const int32 EmissiveColorIdx = CompileWithDefaultFloat3(Compiler, EmissiveColor, 0.0f, 0.0f, 0.0f);
+	const int32 TransmittanceColorIdx = CompileWithDefaultFloat3(Compiler, TransmittanceColor, 1.0f, 1.0f, 1.0f);
+	const int32 NormalIdx = CompileWithDefaultNormalWS(Compiler, Normal);
+
 	int32 OutputCodeChunk = Compiler->SubstrateUnlitBSDF(
-		CompileWithDefaultFloat3(Compiler, EmissiveColor, 0.0f, 0.0f, 0.0f),
-		CompileWithDefaultFloat3(Compiler, TransmittanceColor, 1.0f, 1.0f, 1.0f),
-		CompileWithDefaultNormalWS(Compiler, Normal),
+		EmissiveColorIdx,
+		TransmittanceColorIdx,
+		NormalIdx,
 		&SubstrateOperator);
 	return OutputCodeChunk;
 }
