Fix for PoseableMeshComponent for when Skeleton changes during its lifetime.

#Synced-CL 51066804

diff --git a/Engine/Source/Runtime/Engine/Private/Components/PoseableMeshComponent.cpp b/Engine/Source/Runtime/Engine/Private/Components/PoseableMeshComponent.cpp
--- a/Engine/Source/Runtime/Engine/Private/Components/PoseableMeshComponent.cpp
+++ b/Engine/Source/Runtime/Engine/Private/Components/PoseableMeshComponent.cpp
@@ -87,6 +87,12 @@ void UPoseableMeshComponent::FillComponentSpaceTransforms()
 		return;
 	}
 
+	// If skeleton has changed then attempt to reallocate internal transform data
+	if (BoneSpaceTransforms.Num() != GetSkinnedAsset()->GetRefSkeleton().GetNum())
+	{
+		verify(AllocateTransformData());
+	}
+
 	// right now all this does is to convert to SpaceBases
 	check( GetSkinnedAsset()->GetRefSkeleton().GetNum() == BoneSpaceTransforms.Num() );
 	check( GetSkinnedAsset()->GetRefSkeleton().GetNum() == GetNumComponentSpaceTransforms());
