New material translator:
- Extern instruction introduced. This new instruction is intended to unify ExternInputs and inline HLSL and provide a layered, robust way to extend MIR functionality without directly modifying the translator code itself. A set of extern definitions provided that cover common use cases.
- Removed ExternalInput MIR value. Most of the external inputs have been translated to the new Extern instruction, now living entirely within the MaterialExpressionsIR.cpp file. Only two have remained, but these have been renamed to "Builtin" value (rather than external input), as they represent two constants that are known by the translator but provided by the backend.
- SceneTexture instructions removed. Replaced with Extern instruction emissions
- All calls to InlineHLSL have been replaced to Extern, and the InlineHLSL instruction/emitter API has been removed.

#rnx

#Synced-CL 51020780

diff --git a/Engine/Source/Runtime/Engine/Private/Materials/MaterialExpressionsIR.cpp b/Engine/Source/Runtime/Engine/Private/Materials/MaterialExpressionsIR.cpp
--- a/Engine/Source/Runtime/Engine/Private/Materials/MaterialExpressionsIR.cpp
+++ b/Engine/Source/Runtime/Engine/Private/Materials/MaterialExpressionsIR.cpp
@@ -259,6 +259,9 @@
 #include "Materials/MaterialExpressionBindlessSwitch.h"
 #include "Materials/MaterialExpressionSubstrate.h"
 #include "MaterialShared.h"
+#include "Materials/MaterialIREmitter.h"
+#include "Materials/MaterialIRModule.h"
+#include "Materials/MaterialIRExtern.h"
 #include "Misc/MemStackUtility.h"
 #include "RenderUtils.h"
 #include "Shader/Preshader2.h"
@@ -267,50 +270,123 @@
 #include "ColorManagement/ColorSpace.h"
 #include "SubstrateTranslatorCommon.h"
 #include "Engine/SubsurfaceProfile.h"
-#include "Materials/MaterialIREmitter.h"
-#include "Materials/MaterialIRModule.h"
+
+#include "PostProcess/PostProcessMaterialInputs.h"
 
 using FValueRef = MIR::FValueRef;
 
+/* Material builtin values emission */
+
+static FValueRef EmitTexCoord(MIR::FEmitter& Em, int32 Index)
+{
+	check(Index >= 0);
+	if (Index >= 8)
+	{
+		Em.Errorf(TEXT("Texture coordinate %d out of bounds (max is 8)"), Index);
+		return Em.Poison();
+	}
+
+	struct FTexCoord
+	{
+		int32 Index;
+
+		MIR::FExternInfo GetInfo() const
+		{
+			return
+			{
+				.Name  = TEXTVIEW("TexCoord"),
+				.Type  = MIR::FType::MakeFloatVector(2),
+				.Flags = MIR::EExternFlags::Inline,
+			};
+		}
+
+		void AnalyzeInStage(MIR::FExternAnalysisContext& Context, MIR::EStage Stage)
+		{
+			// Texture coordinates read in the vertex stage don't need to be interpolated.
+			if (Stage == MIR::EStage::Stage_Vertex)
+			{
+				return;
+			}
+
+			// @massimo.tristano this follows the logic in the old translator. A better way would instead be to **allocate** an index
+			// to a texture coordinate based on usage and map that index here. 
+			FMaterialIRModule::FStatistics& Statistics = Context.Module->GetStatistics();
+			Statistics.NumInterpolatedTexCoords = FMath::Max(Statistics.NumInterpolatedTexCoords, Index + 1);
+			Statistics.InterpolatedTexCoordsMask |= 1 << Index;
+		}
+	
+		void ToHLSL(MIR::FExternPrinterHLSL& Printer) const
+		{
+			// Parameters.TexCoords<_DD?>[<Index>]
+			Printer << TEXTVIEW("Parameters.TexCoords")
+					<< MIR::FExternTag_IfDifferential(TEXTVIEW("_"))
+				    << MIR::ExternTag_DD
+					<< TEXT('[') << Index << TEXT(']');
+		}
+	};
+
+	return Em.Extern<FTexCoord>({ Index });
+}
+
+static FValueRef EmitPixelNormalWS(MIR::FEmitter& Em)
+{
+	static const MIR::FExternInlineDeclaration Decl{ .Type =MIR::FType::MakeFloatVector(3), .CodeHLSL = TEXT("Parameters.WorldNormal"), .GraphProperties = MIR::EGraphProperties::ReadsPixelNormal };
+	return Em.Extern<MIR::FExternFromInlineDecl>(&Decl);
+}
+
+static FValueRef EmitVertexNormal(MIR::FEmitter& Em)
+{
+	static const FName NAME_VertexNormal("VertexNormal");
+	return Em.Extern<MIR::FExternFromMaterialDecl>(NAME_VertexNormal);
+}
+
+static FValueRef EmitVertexTangent(MIR::FEmitter& Em)
+{
+	static const FName NAME_VertexTangent("VertexTangent");
+	return Em.Extern<MIR::FExternFromMaterialDecl>(NAME_VertexTangent);
+}
+
 /* Helper functions to emit inline HLSL instructions */
 
-static FValueRef MaterialExpressionExternalCodeBase_EmitInlineHLSL(MIR::FEmitter& Em, const UMaterialExpressionExternalCodeBase& InExpression, int32 InExternalCodeIdentifierIndex, TConstArrayView<FValueRef> InArguments = {}, MIR::EValueFlags ValueFlags = MIR::EValueFlags::None)
+static void CheckExternalCodeIdentifierIsValid(const UMaterialExpressionExternalCodeBase& InExpression, int32 InExternalCodeIdentifierIndex)
 {
 	checkf(InExternalCodeIdentifierIndex >= 0 && InExternalCodeIdentifierIndex < InExpression.ExternalCodeIdentifiers.Num(),
 		   TEXT("External code identifier index (%d) out of bounds; Upper bound is %d"), InExternalCodeIdentifierIndex, InExpression.ExternalCodeIdentifiers.Num());
+}
 
-	const FName ExternalCodeIdentifier = InExpression.ExternalCodeIdentifiers[InExternalCodeIdentifierIndex];
-	const FMaterialExternalCodeDeclaration* ExternalCodeDeclaration = MaterialExternalCodeRegistry::Get().FindExternalCode(ExternalCodeIdentifier);
-	if (!ExternalCodeDeclaration)
-	{
-		Em.Errorf(TEXT("Missing external code declaration for '%s' [Index=%d]"), *ExternalCodeIdentifier.ToString(), InExternalCodeIdentifierIndex);
-		return {};
-	}
+template <typename... TArguments>
+requires (std::same_as<TArguments, FValueRef> && ...)
+static FValueRef MaterialExpressionExternalCodeBase_EmitExtern(MIR::FEmitter& Em, const UMaterialExpressionExternalCodeBase& InExpression, int32 InExternalCodeIdentifierIndex, TArguments... InArguments)
+{
+	CheckExternalCodeIdentifierIsValid(InExpression, InExternalCodeIdentifierIndex);
 
-	return Em.InlineHLSL(ExternalCodeDeclaration, InArguments, ValueFlags);
+	const FName ExternalCodeIdentifier = InExpression.ExternalCodeIdentifiers[InExternalCodeIdentifierIndex];
+	return Em.Extern<MIR::FExternFromMaterialDecl>({ ExternalCodeIdentifier }, InArguments...);
 }
 
-static void BuildInlineHLSLOutput(MIR::FEmitter& Em, UMaterialExpressionExternalCodeBase& InExpression, TConstArrayView<FValueRef> InArguments)
+template <typename... TArguments>
+requires (std::same_as<TArguments, FValueRef> && ...)
+static void BuildMaterialExpressionExternalCodeBase(MIR::FEmitter& Em, UMaterialExpressionExternalCodeBase& InExpression, TArguments... InArguments)
 {
 	// If the external code expression has only one identifier, evaluate it and flow its result to all expression outputs.
 	if (InExpression.ExternalCodeIdentifiers.Num() == 1)
 	{
-		Em.OutputsWithComponentMask(MaterialExpressionExternalCodeBase_EmitInlineHLSL(Em, InExpression, 0, InArguments));
+		Em.OutputsWithComponentMask(MaterialExpressionExternalCodeBase_EmitExtern(Em, InExpression, 0, InArguments...));
 	}
-	else
+	else 
 	{
 		// Make sure the number of identifiers perfectly matches the number of expression outputs.
 		if (InExpression.ExternalCodeIdentifiers.Num() != InExpression.GetOutputs().Num())
 		{
-			Em.Errorf(TEXT("Internal Error: The number of external code identifiers (%d) does not match the number of expression outputs (%d) in expression '%s'.")
-					  , InExpression.ExternalCodeIdentifiers.Num(), InExpression.GetOutputs().Num(), *InExpression.GetName());
+			Em.Errorf(TEXT("Internal Error: The number of external code identifiers (%d) does not match the number of expression outputs (%d) in expression '%s'."),
+					  InExpression.ExternalCodeIdentifiers.Num(), InExpression.GetOutputs().Num(), *InExpression.GetName());
 			return;
 		}
 
 		// Evaluate each external code identifier separately and flow it to its own matching expression output only.
 		for (int32 OutputIndex = 0; OutputIndex < InExpression.ExternalCodeIdentifiers.Num(); ++OutputIndex)
 		{
-			Em.Output(OutputIndex, MaterialExpressionExternalCodeBase_EmitInlineHLSL(Em, InExpression, OutputIndex, InArguments));
+			Em.Output(OutputIndex, MaterialExpressionExternalCodeBase_EmitExtern(Em, InExpression, OutputIndex, InArguments...));
 		}
 	}
 }
@@ -654,7 +730,7 @@ void UMaterialExpressionVirtualTextureFeatureSwitch::Build(MIR::FEmitter& Em)
 	Em.Output(0, UseVirtualTexturing(Em.GetShaderPlatform()) ? YesValue : NoValue);
 }
 
-static void EmitExternalCodeConditionalReplace(MIR::FEmitter& Em, FName ExternalCodeIdentifier, FExpressionInput& Default, FExpressionInput& Replace, const TCHAR* DefaultDesc, const TCHAR* ReplaceDesc)
+static void EmitExternalCodeConditionalReplace(MIR::FEmitter& Em, FName ExternalCodeIdentifier, FExpressionInput& Default, FExpressionInput& Replace)
 {
 	FValueRef DefaultValue = Em.Input(&Default);
 	FValueRef ReplaceValue = Em.Input(&Replace);
@@ -666,50 +742,35 @@ static void EmitExternalCodeConditionalReplace(MIR::FEmitter& Em, FName External
 	DefaultValue = Em.Cast(DefaultValue, CommonType);
 	ReplaceValue = Em.Cast(ReplaceValue, CommonType);
 
-	Em.Output(0, Em.Branch(Em.InlineHLSL(ExternalCodeIdentifier), ReplaceValue, DefaultValue));
+	Em.Output(0, Em.Branch(Em.Extern<MIR::FExternFromMaterialDecl>(ExternalCodeIdentifier), ReplaceValue, DefaultValue));
 }
 
-static FValueRef EmitVertexNormal(MIR::FEmitter& Em)
-{
-	static const FName NAME_VertexNormal("VertexNormal");
-	return Em.InlineHLSL(NAME_VertexNormal);
-}
-
-static FValueRef EmitVertexTangent(MIR::FEmitter& Em)
-{
-	static const FName NAME_VertexTangent("VertexTangent");
-	return Em.InlineHLSL(NAME_VertexTangent);
-}
-
-#define REPLACE_INPUTS(DEFAULT, REPLACE) DEFAULT, REPLACE, TEXT(#DEFAULT), TEXT(#REPLACE)
-
 void UMaterialExpressionLightmassReplace::Build(MIR::FEmitter& Em)
 {
-	EmitExternalCodeConditionalReplace(Em, FName("LightmassReplace"), REPLACE_INPUTS(Realtime, Lightmass));
+	// @massimo.tristano Missing check whether this is a "lightmass compiler" and optimize out to Lightmass
+	EmitExternalCodeConditionalReplace(Em, FName("LightmassReplace"), Realtime, Lightmass);
 }
 
 void UMaterialExpressionMeshPaintTextureReplace::Build(MIR::FEmitter& Em)
 {
-	EmitExternalCodeConditionalReplace(Em, FName("MeshPaintTextureReplace"), REPLACE_INPUTS(Default, MeshPaintTexture));
+	EmitExternalCodeConditionalReplace(Em, FName("MeshPaintTextureReplace"), Default, MeshPaintTexture);
 }
 
 void UMaterialExpressionNaniteReplace::Build(MIR::FEmitter& Em)
 {
-	EmitExternalCodeConditionalReplace(Em, FName("NaniteReplace"), REPLACE_INPUTS(Default, Nanite));
+	EmitExternalCodeConditionalReplace(Em, FName("NaniteReplace"), Default, Nanite);
 }
 
 void UMaterialExpressionReflectionCapturePassSwitch::Build(MIR::FEmitter& Em)
 {
-	EmitExternalCodeConditionalReplace(Em, FName("ReflectionCapturePassSwitch"), REPLACE_INPUTS(Default, Reflection));
+	EmitExternalCodeConditionalReplace(Em, FName("ReflectionCapturePassSwitch"), Default, Reflection);
 }
 
 void UMaterialExpressionShadowReplace::Build(MIR::FEmitter& Em)
 {
-	EmitExternalCodeConditionalReplace(Em, FName("ShadowReplace"), REPLACE_INPUTS(Default, Shadow));
+	EmitExternalCodeConditionalReplace(Em, FName("ShadowReplace"), Default, Shadow);
 }
 
-#undef REPLACE_INPUTS
-
 void UMaterialExpressionAppendVector::Build(MIR::FEmitter& Em)
 {
 	FValueRef AVal = Em.CheckIsScalarOrVector(Em.Input(&A));
@@ -962,7 +1023,7 @@ void UMaterialExpressionSquareRoot::Build(MIR::FEmitter& Em)
 
 void UMaterialExpressionExternalCodeBase::Build(MIR::FEmitter& Em)
 {
-	BuildInlineHLSLOutput(Em, *this, {});
+	BuildMaterialExpressionExternalCodeBase(Em, *this);
 }
 
 /* Binary Operators */
@@ -1061,7 +1122,7 @@ void UMaterialExpressionEyeAdaptationInverse::Build(MIR::FEmitter& Em)
 	check(ExternalCodeIdentifiers.Num() == 1);
 	FValueRef LightValue = Em.CastToFloat(Em.InputDefaultFloat(&LightValueInput, 1.0f), 3);
 	FValueRef AlphaValue = Em.CastToFloat(Em.InputDefaultFloat(&AlphaInput, 1.0f), 1);
-	FValueRef MultiplierValue = MaterialExpressionExternalCodeBase_EmitInlineHLSL(Em, *this, 0, { AlphaValue });
+	FValueRef MultiplierValue = MaterialExpressionExternalCodeBase_EmitExtern(Em, *this, 0, AlphaValue);
 	Em.Output(0, Em.Multiply(LightValue, MultiplierValue));
 }
 
@@ -1163,6 +1224,312 @@ void UMaterialExpressionIf::Build(MIR::FEmitter& Em)
 	Em.Output(0, OutputValue);
 }
 
+static void ErrorUnlessFeatureLevelSupported(FMaterialIRModule* Module, ERHIFeatureLevel::Type RequiredFeatureLevel, const TCHAR* Message)
+{
+	ERHIFeatureLevel::Type FeatureLevel = Module->GetFeatureLevel();
+	if (FeatureLevel < RequiredFeatureLevel)
+	{
+		FString FeatureLevelName;
+		GetFeatureLevelName(FeatureLevel, FeatureLevelName);
+		Module->AddError(nullptr, FString::Printf(TEXT("%s  Current feature level is %s."), Message, *FeatureLevelName));
+	}
+}
+
+enum class EScreenTexture : uint8
+{
+	SceneTexture,
+	UserSceneTexture,
+	SceneColor,
+	SceneDepth,
+	SceneDepthWithoutWater,
+	DBufferTexture,
+};
+
+// Reference to various types of full screen textures accessible by expressions.  These
+// generally involve custom code to set data in the compilation output, extensive validation
+// logic, and HLSL generation deferred to the analyzer for the case of UserSceneTextures.
+struct FScreenTextureExtern
+{
+	EScreenTexture Kind;
+
+	union
+	{
+		ESceneTextureId SceneTextureId;
+		EDBufferTextureId DBufferTextureId;
+	};
+
+	FName UserSceneTexture;
+
+	FScreenTextureExtern(ESceneTextureId InId)
+	: Kind{ EScreenTexture::SceneTexture }
+	, SceneTextureId{ InId }
+	{
+	}
+
+	FScreenTextureExtern(FName InUserSceneTexture)
+	: Kind{ EScreenTexture::UserSceneTexture }
+	, UserSceneTexture{ InUserSceneTexture }
+	{
+	}
+
+	FScreenTextureExtern(EDBufferTextureId DBufferTextureId)
+	: Kind{ EScreenTexture::DBufferTexture }
+	, DBufferTextureId { DBufferTextureId }
+	{
+	}
+
+	FScreenTextureExtern(EScreenTexture TextureKind)
+	: Kind(TextureKind)
+	{
+	}
+
+	MIR::FExternInfo GetInfo() const
+	{
+		return {
+			.Name = TEXTVIEW("ScreenTexture"),
+			.Type = MIR::FType::MakeIntScalar(),
+			.Flags = MIR::EExternFlags::Inline,
+		};
+	}
+
+	void Analyze(MIR::FExternAnalysisContext& Context)
+	{
+		const UMaterial* Material = Context.MaterialInterface->GetMaterial();
+		FMaterialCompilationOutput& CompilationOutput = Context.Module->GetCompilationOutput();
+		const EMaterialDomain MaterialDomain = Material->MaterialDomain;
+
+		switch (Kind)
+		{
+			case EScreenTexture::SceneTexture:
+			case EScreenTexture::UserSceneTexture:
+			{
+				// TODO: If referenced from custom HLSL, this can be false. Revisit when custom HLSL support is added to the new translator.
+				const bool bTextureLookup = true;
+
+				// Add defines and compilation outputs.
+				Context.Module->AddIntegerEnvironmentDefine(FName("NEEDS_SCENE_TEXTURES"), 1);
+
+				if (Kind == EScreenTexture::UserSceneTexture)
+				{
+					if (UserSceneTexture.IsNone())
+					{
+						Context.Module->AddError(nullptr, TEXT("UserSceneTexture missing name. Value must be set to something other than None."));
+						return;
+					}
+					else
+					{
+						// Allocate value during Analyze and check for failure. FindUserSceneTexture is later called during HLSL generation.
+						SceneTextureId = (ESceneTextureId)CompilationOutput.FindOrAddUserSceneTexture(UserSceneTexture);
+						if ((int32)SceneTextureId == INDEX_NONE)
+						{
+							Context.Module->AddError(nullptr, FString::Printf(TEXT("Too many unique UserSceneTexture inputs in the post process material. Max allowed is %d."), kPostProcessMaterialInputCountMax));
+							return;
+						}
+					}
+				}
+
+				CompilationOutput.bNeedsSceneTextures = true;
+				CompilationOutput.SetIsSceneTextureUsed(SceneTextureId);
+
+				// Substrate TODO:
+				// When SceneTexture lookup is used, simple and single paths are disabled to ensure correct decoding.
+				// Reading SceneTexture with Substrate implies unpacking material buffer data.
+				// To avoid compiling out unpacking paths, force simple and single versions to be disabled.
+				// FSubstrateCompilationContext& SubstrateContext = SubstrateCompilationContext[CurrentSubstrateCompilationContext];
+				// SubstrateContext.SubstrateMaterialComplexity.bIsSimple = false;
+				// SubstrateContext.SubstrateMaterialComplexity.bIsSingle = false;
+
+				const EShaderPlatform Platform = Context.Module->GetShaderPlatform();
+
+				const bool bHasSingleLayerWaterSM =
+					Context.MaterialInterface->GetShadingModels().HasShadingModel(MSM_SingleLayerWater);
+
+				if (bHasSingleLayerWaterSM && SceneTextureId != PPI_CustomDepth && SceneTextureId != PPI_CustomStencil)
+				{
+					Context.Module->AddError(nullptr, TEXT("Only CustomDepth and CustomStencil can be sampled with SceneTexture when using the Single Layer Water shading model."));
+				}
+
+				if (SceneTextureId == PPI_DecalMask)
+				{
+					Context.Module->AddError(nullptr, TEXT("Decal mask bit was moved from the GBuffer to the stencil buffer and is no longer available."));
+				}
+
+				if (MaterialDomain == MD_DeferredDecal)
+				{
+					const bool bSceneTextureSupportsDecal = SceneTextureId == PPI_SceneDepth || SceneTextureId == PPI_WorldNormal || SceneTextureId == PPI_CustomDepth || SceneTextureId == PPI_CustomStencil;
+
+					if (!bSceneTextureSupportsDecal)
+					{
+						Context.Module->AddError(nullptr, TEXT("Decals can only access SceneDepth, CustomDepth, CustomStencil, and WorldNormal."));
+					}
+
+					if (SceneTextureId == PPI_WorldNormal)
+					{
+						ErrorUnlessFeatureLevelSupported(Context.Module, ERHIFeatureLevel::SM5, TEXT("Deferred decals require SM5 for world normal access."));
+					}
+
+					if (SceneTextureId == PPI_WorldNormal && !IsUsingDBuffers(Platform))
+					{
+						const bool bHasNormalConnected = Substrate::IsSubstrateEnabled() ? Material->IsPropertyConnected(MP_Normal) : Material->HasNormalConnected();
+						if (bHasNormalConnected)
+						{
+							Context.Module->AddError(nullptr, TEXT("Decals that read WorldNormal cannot output to Normal at the same time. Enable DBuffer to support this."));
+						}
+					}
+				}
+
+				bool bNeedsSceneTexturePostProcessInputs = false;
+				if (bTextureLookup)
+				{
+					bNeedsSceneTexturePostProcessInputs =
+						(SceneTextureId >= PPI_PostProcessInput0 && SceneTextureId <= PPI_PostProcessInput6)
+						|| (SceneTextureId >= PPI_UserSceneTexture0 && SceneTextureId <= PPI_UserSceneTexture6)
+						|| SceneTextureId == PPI_Velocity
+						|| SceneTextureId == PPI_SceneColor;
+				}
+
+				if (SceneTextureId == PPI_Velocity && MaterialDomain != MD_PostProcess)
+				{
+					Context.Module->AddError(nullptr, TEXT("Velocity scene textures are only available in post process materials."));
+				}
+
+				if (MaterialDomain != MD_DeferredDecal && MaterialDomain != MD_PostProcess)
+				{
+					if (!bHasSingleLayerWaterSM && IsOpaqueOrMaskedBlendMode(Material->BlendMode))
+					{
+						Context.Module->AddError(nullptr, TEXT("SceneTexture expressions cannot be used in opaque materials except with Single Layer Water."));
+					}
+					else if (bNeedsSceneTexturePostProcessInputs)
+					{
+						Context.Module->AddError(nullptr, TEXT("SceneTexture expressions cannot use post process inputs or SceneColor outside post process materials."));
+					}
+				}
+
+				if (SceneTextureId == PPI_SceneDepth && bTextureLookup && MaterialDomain != MD_PostProcess && !IsTranslucentBlendMode(Material->BlendMode))
+				{
+					Context.Module->AddError(nullptr, TEXT("Only transparent or post process materials can read from SceneDepth."));
+				}
+				break;
+			}
+
+			case EScreenTexture::SceneColor:
+			{
+				CompilationOutput.SetIsSceneTextureUsed(PPI_SceneColor);
+
+				if (MaterialDomain != MD_Surface)
+				{
+					Context.Module->AddError(nullptr, TEXT("SceneColor lookups are only available when MaterialDomain is Surface."));
+				}
+
+				ErrorUnlessFeatureLevelSupported(Context.Module, ERHIFeatureLevel::SM5, TEXT("SceneColor access requires SM5."));
+				break;
+			}
+
+			case EScreenTexture::SceneDepth:
+			{
+				CompilationOutput.SetIsSceneTextureUsed(PPI_SceneDepth);
+
+				if (Context.MaterialInterface->IsTranslucencyWritingVelocity())
+				{
+					Context.Module->AddError(nullptr, TEXT("Translucent materials writing velocity cannot read from SceneDepth."));
+				}
+
+				if (MaterialDomain != MD_PostProcess && !IsTranslucentBlendMode(Material->BlendMode))
+				{
+					Context.Module->AddError(nullptr, TEXT("Only transparent or post process materials can read from SceneDepth."));
+				}
+				break;
+			}
+
+			case EScreenTexture::SceneDepthWithoutWater:
+			{
+				if (MaterialDomain != MD_PostProcess)
+				{
+					if (!Context.MaterialInterface->GetShadingModels().HasShadingModel(MSM_SingleLayerWater))
+					{
+						Context.Module->AddError(nullptr, TEXT("Reading scene depth below water requires Single Layer Water or PostProcess domain."));
+					}
+
+					if (MaterialDomain != MD_Surface)
+					{
+						Context.Module->AddError(nullptr, TEXT("SceneDepthWithoutWater requires Surface or PostProcess material domain."));
+					}
+
+					if (IsTranslucentBlendMode(Context.Module->GetBlendMode()))
+					{
+						Context.Module->AddError(nullptr, TEXT("SceneDepthWithoutWater cannot be used with translucent blend mode."));
+					}
+				}
+				break;
+			}
+
+			case EScreenTexture::DBufferTexture:
+			{
+				Context.Module->AddEnvironmentDefine(FName("Material_Uses_Decal_Lookup"));
+
+				CompilationOutput.SetIsDBufferTextureUsed(DBufferTextureId);
+				CompilationOutput.SetIsDBufferTextureLookupUsed(true);
+
+				if (MaterialDomain != MD_Surface || IsTranslucentBlendMode(Context.Module->GetBlendMode()))
+				{
+					Context.Module->AddError(nullptr, TEXT("DBuffer scene textures are only available on opaque or masked surfaces."));
+				}
+				break;
+			}
+		}
+	}
+
+	void AnalyzeInStage(MIR::FExternAnalysisContext& Context, MIR::EStage Stage)
+	{
+		switch (Kind)
+		{
+			case EScreenTexture::SceneColor:
+				if (Stage == MIR::Stage_Vertex)
+				{
+					Context.Module->AddError(nullptr, TEXT("SceneColor is only supported in pixel shader input."));
+				}
+				break;
+
+			case EScreenTexture::SceneDepth:
+				if (Stage == MIR::Stage_Vertex)
+				{
+					ErrorUnlessFeatureLevelSupported(Context.Module, ERHIFeatureLevel::SM5, TEXT("Reading SceneDepth from the vertex shader requires SM5."));
+				}
+				break;
+
+			case EScreenTexture::SceneDepthWithoutWater:
+				if (Stage == MIR::Stage_Vertex)
+				{
+					Context.Module->AddError(nullptr, TEXT("Cannot read SceneDepthWithoutWater from the vertex shader."));
+				}
+				break;
+
+			default:
+				break;
+		}
+	}
+
+	void ToHLSL(MIR::FExternPrinterHLSL& Printer) const
+	{
+		check(Printer.Differential == MIR::EExternDifferential::None);
+
+		switch (Kind)
+		{
+			case EScreenTexture::SceneTexture:
+			case EScreenTexture::UserSceneTexture:
+				Printer << UE::MaterialTranslatorUtils::SceneTextureIdToHLSLString(SceneTextureId);
+				break;
+
+			case EScreenTexture::DBufferTexture:
+				Printer << (int32)DBufferTextureId;
+				break;
+
+			default:
+				break;
+		}
+	}
+};
+
 // If DefaultOffset is not null, Coordinates are treated as an offset (or DefaultOffset if unset), rather than absolute coordinates.  Clamping is
 // automatically applied for custom or offset fetches -- the "bClamped" parameter only controls clamping for default texture coordinate fetches,
 // and is only needed when fetching from lower resolution User Scene Textures.  A zero constant can be passed in for "SceneTextureInput" for
@@ -1172,25 +1539,21 @@ static FValueRef SceneTextureExpressionTexCoords(MIR::FEmitter& Em, FValueRef Sc
 	FValueRef TexCoords;
 	if (DefaultOffset)
 	{
-		TexCoords = Em.InputDefaultFloat2(&Coordinates, { (float)DefaultOffset->X, (float)DefaultOffset->Y });
-		TexCoords = Em.InlineHLSL(MIR::FType::MakeFloatVector(2), TEXTVIEW("CalcScreenUVFromOffsetFraction(GetScreenPosition(Parameters), $0)"), { TexCoords });
+		TexCoords = Em.Cast(Em.InputDefaultFloat2(&Coordinates, { (float)DefaultOffset->X, (float)DefaultOffset->Y }), MIR::FType::MakeFloatVector(2));
+		TexCoords = Em.Extern<MIR::FExternSimpleHLSL>({ MIR::EExternSimpleType::Float2, TEXT("CalcScreenUVFromOffsetFraction(GetScreenPosition(Parameters), $0)") }, TexCoords);
 	}
 	else
 	{
-		TexCoords = Em.TryInput(&Coordinates);
+		TexCoords = Em.Cast(Em.TryInput(&Coordinates), MIR::FType::MakeFloatVector(2));
 		if (TexCoords)
 		{
-			// Convert raw TexCoords expression to Float2, then convert from viewport to scene texture space
-			TexCoords = Em.Cast(TexCoords, MIR::FType::MakeFloatVector(2));
-			TexCoords = Em.InlineHLSL(MIR::FType::MakeFloatVector(2), TEXTVIEW("ClampSceneTextureUV(ViewportUVToSceneTextureUV($0, $1), $1)"), { TexCoords, SceneTextureInput });
+			// Convert from viewport to scene texture space
+			TexCoords = Em.Extern<MIR::FExternSimpleHLSL>({ MIR::EExternSimpleType::Float2, TEXT("ClampSceneTextureUV(ViewportUVToSceneTextureUV($0, $1), $1)") }, TexCoords, SceneTextureInput);
 		}
 		else
 		{
-			FStringView TexCoordsCode = bClamped ?
-				TEXTVIEW("ClampSceneTextureUV(GetDefaultSceneTextureUV(Parameters, $0), $0)") :
-				TEXTVIEW("GetDefaultSceneTextureUV(Parameters, $0)");
-
-			TexCoords = Em.InlineHLSL(MIR::FType::MakeFloatVector(2), TexCoordsCode, { SceneTextureInput });
+			const TCHAR* Code = bClamped ? TEXT("ClampSceneTextureUV(GetDefaultSceneTextureUV(Parameters, $0), $0)") : TEXT("GetDefaultSceneTextureUV(Parameters, $0)");
+			TexCoords = Em.Extern<MIR::FExternSimpleHLSL>({ MIR::EExternSimpleType::Float2, Code }, SceneTextureInput);
 		}
 	}
 	return TexCoords;
@@ -1199,25 +1562,22 @@ static FValueRef SceneTextureExpressionTexCoords(MIR::FEmitter& Em, FValueRef Sc
 static void SceneTextureExpressionBuild(MIR::FEmitter& Em, FValueRef SceneTextureInput, const FExpressionInput& Coordinates, const FVector2D* DefaultOffset, bool bClamped, bool bFiltered)
 {
 	FValueRef TexCoords = SceneTextureExpressionTexCoords(Em, SceneTextureInput, Coordinates, DefaultOffset, bClamped);
+	Em.Output(0, Em.Extern<MIR::FExternSimpleHLSL>({ MIR::EExternSimpleType::Float4, TEXT("SceneTextureLookup(Parameters, $0, $1, $2)") }, TexCoords, SceneTextureInput, Em.ConstantBool(bFiltered)));
 
-	FStringView SceneTextureLookupCode = bFiltered ?
-		TEXTVIEW("SceneTextureLookup(Parameters, $0, $1, true)") :
-		TEXTVIEW("SceneTextureLookup(Parameters, $0, $1, false)");
-
-	Em.Output(0, Em.InlineHLSL(MIR::FType::MakeFloatVector(4), SceneTextureLookupCode, { TexCoords, SceneTextureInput }));
-	Em.Output(1, Em.InlineHLSL(MIR::FType::MakeFloatVector(2), TEXTVIEW("GetSceneTextureViewSize($0).xy"), { SceneTextureInput }));
-	Em.Output(2, Em.InlineHLSL(MIR::FType::MakeFloatVector(2), TEXTVIEW("GetSceneTextureViewSize($0).zw"), { SceneTextureInput }));
+	FValueRef SceneTextureViewSize = Em.Extern<MIR::FExternSimpleHLSL>({ MIR::EExternSimpleType::Float4, TEXT("GetSceneTextureViewSize($0)") }, SceneTextureInput);
+	Em.Output(1, Em.Swizzle(SceneTextureViewSize, MIR::FSwizzleMask::XY()));
+	Em.Output(2, Em.Swizzle(SceneTextureViewSize, MIR::FSwizzleMask::ZW()));
 }
 
 void UMaterialExpressionSceneTexture::Build(MIR::FEmitter& Em)
 {
 	const bool bClamped = false;
-	SceneTextureExpressionBuild(Em, Em.SceneTexture(SceneTextureId), Coordinates, nullptr, bClamped, bFiltered);
+	SceneTextureExpressionBuild(Em, Em.Extern<FScreenTextureExtern>({ SceneTextureId.GetValue() }), Coordinates, nullptr, bClamped, bFiltered);
 }
 
 void UMaterialExpressionUserSceneTexture::Build(MIR::FEmitter& Em)
 {
-	SceneTextureExpressionBuild(Em, Em.UserSceneTexture(UserSceneTexture), Coordinates, nullptr, bClamped, bFiltered);
+	SceneTextureExpressionBuild(Em, Em.Extern<FScreenTextureExtern>({ UserSceneTexture }), Coordinates, nullptr, bClamped, bFiltered);
 }
 
 void UMaterialExpressionSceneColor::Build(MIR::FEmitter& Em)
@@ -1225,8 +1585,8 @@ void UMaterialExpressionSceneColor::Build(MIR::FEmitter& Em)
 	FValueRef TexCoords = SceneTextureExpressionTexCoords(Em, Em.ConstantInt(0), Input, this->InputMode == EMaterialSceneAttributeInputMode::OffsetFraction ? &ConstInput : nullptr);
 
 	// We need a dependency on ScreenTexture as a second argument, so the value analyzer can see it, even though it's technically not used in the code.
-	FValueRef ScreenTexture = Em.ScreenTexture(MIR::EScreenTexture::SceneColor);
-	Em.Output(0, Em.InlineHLSL(MIR::FType::MakeFloatVector(4), TEXTVIEW("DecodeSceneColorAndAlpharForMaterialNode($0)"), { TexCoords, ScreenTexture }));;
+	FValueRef ScreenTexture = Em.Extern<FScreenTextureExtern>({ EScreenTexture::SceneColor });
+	Em.Output(0, Em.Extern<MIR::FExternSimpleHLSL>({ MIR::EExternSimpleType::Float4, TEXT("DecodeSceneColorAndAlpharForMaterialNode($0)") }, TexCoords, ScreenTexture));
 }
 
 void UMaterialExpressionSceneDepth::Build(MIR::FEmitter& Em)
@@ -1234,8 +1594,8 @@ void UMaterialExpressionSceneDepth::Build(MIR::FEmitter& Em)
 	FValueRef TexCoords = SceneTextureExpressionTexCoords(Em, Em.ConstantInt(0), Input, this->InputMode == EMaterialSceneAttributeInputMode::OffsetFraction ? &ConstInput : nullptr);
 
 	// We need a dependency on ScreenTexture as a second argument, so the value analyzer can see it, even though it's technically not used in the code.
-	FValueRef ScreenTexture = Em.ScreenTexture(MIR::EScreenTexture::SceneDepth);
-	Em.Output(0, Em.InlineHLSL(MIR::FType::MakeFloatScalar(), TEXTVIEW("CalcSceneDepth($0)"), { TexCoords, ScreenTexture }));;
+	FValueRef ScreenTexture = Em.Extern<FScreenTextureExtern>({ EScreenTexture::SceneDepth });
+	Em.Output(0, Em.Extern<MIR::FExternSimpleHLSL>({ MIR::EExternSimpleType::Float1, TEXT("CalcSceneDepth($0)") }, TexCoords, ScreenTexture));
 }
 
 void UMaterialExpressionSceneDepthWithoutWater::Build(MIR::FEmitter& Em)
@@ -1243,27 +1603,52 @@ void UMaterialExpressionSceneDepthWithoutWater::Build(MIR::FEmitter& Em)
 	FValueRef TexCoords = SceneTextureExpressionTexCoords(Em, Em.ConstantInt(0), Input, this->InputMode == EMaterialSceneAttributeInputMode::OffsetFraction ? &ConstInput : nullptr);
 
 	// We need a dependency on ScreenTexture as a third argument, so the value analyzer can see it, even though it's technically not used in the code.
-	FValueRef ScreenTexture = Em.ScreenTexture(MIR::EScreenTexture::SceneDepthWithoutWater);
-	Em.Output(0, Em.InlineHLSL(MIR::FType::MakeFloatScalar(), TEXTVIEW("MaterialExpressionSceneDepthWithoutWater($0, $1)"), { TexCoords, Em.ConstantFloat(FallbackDepth), ScreenTexture }));;
+	FValueRef ScreenTexture = Em.Extern<FScreenTextureExtern>({ EScreenTexture::SceneDepthWithoutWater });
+	Em.Output(0, Em.Extern<MIR::FExternSimpleHLSL>({ MIR::EExternSimpleType::Float1, TEXT("MaterialExpressionSceneDepthWithoutWater($0, $1)")}, TexCoords, Em.ConstantFloat(FallbackDepth), ScreenTexture));
 }
 
 void UMaterialExpressionDBufferTexture::Build(MIR::FEmitter& Em)
 {
 	FValueRef TexCoords = SceneTextureExpressionTexCoords(Em, Em.ConstantInt(0), Coordinates);
 
-	FValueRef ScreenTexture = Em.DBufferTexture(DBufferTextureId);
-	Em.Output(0, Em.InlineHLSL(MIR::FType::MakeFloatVector(4), TEXTVIEW("MaterialExpressionDBufferTextureLookup(Parameters, $0, $1)"), { TexCoords, ScreenTexture }));;
+	FValueRef ScreenTexture = Em.Extern<FScreenTextureExtern>({ DBufferTextureId.GetValue() });
+	Em.Output(0, Em.Extern<MIR::FExternSimpleHLSL>({ MIR::EExternSimpleType::Float4, TEXT("MaterialExpressionDBufferTextureLookup(Parameters, $0, $1)") }, TexCoords, ScreenTexture));
 }
 
 void UMaterialExpressionSphericalParticleOpacity::Build(MIR::FEmitter& Em)
 {
 	FValueRef DensityValue = Em.InputDefaultFloat(&Density, ConstantDensity);
 	UE_MIR_CHECKPOINT(Em); // Early out in case of errors
-	BuildInlineHLSLOutput(Em, *this, { DensityValue });
+	BuildMaterialExpressionExternalCodeBase(Em, *this, DensityValue);
 }
 
 void UMaterialExpressionShadingModel::Build(MIR::FEmitter& Em)
 {
+	struct FShadingModel
+	{
+		EMaterialShadingModel Id;
+
+		MIR::FExternInfo GetInfo() const
+		{
+			return {
+				.Name  = TEXTVIEW("ShadingModel"),
+				.Type  = MIR::FType::MakeIntScalar(),
+				.Flags = MIR::EExternFlags::Inline,
+			};
+		}
+
+		void Analyze(MIR::FExternAnalysisContext& Context)
+		{
+			check(Id < MSM_NUM);
+			Context.Module->AddShadingModel(Id);
+		}
+	
+		void ToHLSL(MIR::FExternPrinterHLSL& Printer) const
+		{
+			Printer << (uint32)Id;
+		}
+	};
+
 	EMaterialShadingModel Id = ShadingModel;
 
 	// If the shading model is masked out, fallback to default shading model
@@ -1273,7 +1658,7 @@ void UMaterialExpressionShadingModel::Build(MIR::FEmitter& Em)
 		Id = MSM_DefaultLit;
 	}
 
-	Em.Output(0, Em.ConstantInt(Id));
+	Em.Output(0, Em.Extern<FShadingModel>({ Id }));
 }
 
 void UMaterialExpressionTextureObject::Build(MIR::FEmitter& Em)
@@ -1401,7 +1786,7 @@ static void BuildTextureSampleExpression(MIR::FEmitter& Em, UMaterialExpressionT
 	FValueRef TexCoords = Em.TryInput(&Expr->Coordinates);
 	if (!TexCoords)
 	{
-		TexCoords = Em.ExternalInput(MIR::TexCoordIndexToExternalInput(Expr->ConstCoordinate));
+		TexCoords = EmitTexCoord(Em, Expr->ConstCoordinate);
 	}
 
 	// Determine if automatic view mip bias should be used, by trying to acquire its input as a static boolean.
@@ -1414,7 +1799,7 @@ static void BuildTextureSampleExpression(MIR::FEmitter& Em, UMaterialExpressionT
 	Em.Output(2, Em.Subscript(TextureRead, 1));
 	Em.Output(3, Em.Subscript(TextureRead, 2));
 	Em.Output(4, Em.Subscript(TextureRead, 3));
-	Em.Output(5, TextureRead); 
+	Em.Output(5, TextureRead);
 }
 
 void UMaterialExpressionTextureSample::Build(MIR::FEmitter& Em)
@@ -1455,16 +1840,16 @@ static void EmitParticleSubUV(MIR::FEmitter& Em, UMaterialExpressionTextureSampl
 	static FName NAME_ParticleSubUVLerp("ParticleSubUVLerp");
 
 	FValueRef TexCoords0 = DummyDependency
-		? Em.InlineHLSL(NAME_ParticleSubUVCoords0, { DummyDependency })
-		: Em.InlineHLSL(NAME_ParticleSubUVCoords0);
+		? Em.Extern<MIR::FExternFromMaterialDecl>(NAME_ParticleSubUVCoords0, DummyDependency)
+		: Em.Extern<MIR::FExternFromMaterialDecl>(NAME_ParticleSubUVCoords0);
 	FValueRef Sample0 = BuildTextureSample(Em, Expr, TextureValue, TextureType, TexCoords0, bAutomaticViewMipBias);
 
 	if (bBlend)
 	{
-		FValueRef TexCoords1 = Em.InlineHLSL(NAME_ParticleSubUVCoords1);
+		FValueRef TexCoords1 = Em.Extern<MIR::FExternFromMaterialDecl>(NAME_ParticleSubUVCoords1);
 		FValueRef Sample1 = BuildTextureSample(Em, Expr, TextureValue, TextureType, TexCoords1, bAutomaticViewMipBias);
 
-		FValueRef SubImageLerp = Em.InlineHLSL(NAME_ParticleSubUVLerp);
+		FValueRef SubImageLerp = Em.Extern<MIR::FExternFromMaterialDecl>(NAME_ParticleSubUVLerp);
 
 		Sample0 = Em.Lerp(Sample0, Sample1, SubImageLerp);
 	}
@@ -1496,7 +1881,7 @@ void UMaterialExpressionTextureSampleParameterSubUV::Build(MIR::FEmitter& Em)
 	// while this expression does provide a TextureCoordinate input pin, it is, and has always been, ignored.  And only
 	// supports using UV0.  Further, in order to support non-vertex fetch implementations we need to be sure to register
 	// the use of the first texture slot
-	FValueRef DummyDependency = Em.ExternalInput(MIR::TexCoordIndexToExternalInput(0));
+	FValueRef DummyDependency = EmitTexCoord(Em, 0);
 
 	EmitParticleSubUV(Em, this, ParameterValue, bBlend, DummyDependency);
 }
@@ -1517,7 +1902,7 @@ void UMaterialExpressionAntialiasedTextureMask::Build(MIR::FEmitter& Em)
 	FValueRef TexCoords = Em.TryInput(&Coordinates);
 	if (!TexCoords)
 	{
-		TexCoords = Em.ExternalInput(MIR::TexCoordIndexToExternalInput(ConstCoordinate));
+		TexCoords = EmitTexCoord(Em, ConstCoordinate);
 	}
 
 	FValueRef TextureValue = BuildTextureValue(Em, this);
@@ -1587,7 +1972,7 @@ void UMaterialExpressionTextureCoordinate::Build(MIR::FEmitter& Em)
 		return;
 	}
 
-	FValueRef OutputValue = Em.ExternalInput(MIR::TexCoordIndexToExternalInput(CoordinateIndex));
+	FValueRef OutputValue = EmitTexCoord(Em, CoordinateIndex);
 	
 	// Multiply the UV input by the UV tiling constants
 	OutputValue = Em.Multiply(OutputValue, Em.ConstantFloat2({ UTiling, VTiling }));
@@ -1660,7 +2045,7 @@ void UMaterialExpressionFontSample::Build(MIR::FEmitter& Em)
 
 		FValueRef OutputValue = Em.TextureSample(
 			Em.TextureObject(Texture, ExpectedSamplerType),
-			Em.ExternalInput(MIR::TexCoordIndexToExternalInput(0)),
+			EmitTexCoord(Em, 0),
 			false,
 			{ SSM_FromTextureAsset, ExpectedSamplerType }
 		);
@@ -1701,7 +2086,7 @@ void UMaterialExpressionFontSampleParameter::Build(MIR::FEmitter& Em)
 
 		FValueRef OutputValue = Em.TextureSample(
 			TextureParameter,
-			Em.ExternalInput(MIR::TexCoordIndexToExternalInput(0)),
+			EmitTexCoord(Em, 0),
 			false,
 			{ SSM_FromTextureAsset, ExpectedSamplerType }
 		);
@@ -1715,27 +2100,165 @@ void UMaterialExpressionFontSampleParameter::Build(MIR::FEmitter& Em)
 
 static MIR::FValueRef BuildVirtualTextureWorldToUV(MIR::FEmitter& Em, MIR::FValueRef WorldPositionValue, MIR::FValueRef P0, MIR::FValueRef P1, MIR::FValueRef P2, EPositionOrigin PositionOrigin)
 {
-	return Em.InlineHLSL(MIR::FType::MakeFloatVector(2), TEXTVIEW("VirtualTextureWorldToUV($0, $1, $2, $3)"), { WorldPositionValue, P0, P1, P2 });
+	return Em.Extern<MIR::FExternSimpleHLSL>({ MIR::EExternSimpleType::Float2, TEXT("VirtualTextureWorldToUV($0, $1, $2, $3)") }, WorldPositionValue, P0, P1, P2);
 }
 
 static MIR::FValueRef BuildConstantVector(MIR::FEmitter& Em, TArrayView<float> InConstants)
 {
 	switch (InConstants.Num())
 	{
-	case 1:
-		return Em.ConstantFloat(InConstants[0]);
-	case 2:
-		return Em.ConstantFloat2(FVector2f{ InConstants[0], InConstants[1] });
-	case 3:
-		return Em.ConstantFloat3(FVector3f{ InConstants[0], InConstants[1], InConstants[2] });
-	case 4:
-		return Em.ConstantFloat4(FVector4f{ InConstants[0], InConstants[1], InConstants[2], InConstants[3] });
-	default:
-		UE_MIR_UNREACHABLE();
+		case 1: return Em.ConstantFloat(InConstants[0]);
+		case 2: return Em.ConstantFloat2(FVector2f{ InConstants[0], InConstants[1] });
+		case 3: return Em.ConstantFloat3(FVector3f{ InConstants[0], InConstants[1], InConstants[2] });
+		case 4: return Em.ConstantFloat4(FVector4f{ InConstants[0], InConstants[1], InConstants[2], InConstants[3] });
+		default: UE_MIR_UNREACHABLE();
 	}
 }
 
-static FValueRef EmitWorldPosition(MIR::FEmitter& Em, EWorldPositionIncludedOffsets WorldPositionShaderOffset);
+
+static FValueRef EmitWorldPosition(MIR::FEmitter& Em, EWorldPositionIncludedOffsets WorldPositionShaderOffset)
+{
+	struct FWorldPosition
+	{
+		EWorldPositionIncludedOffsets WorldPositionShaderOffset;
+
+		MIR::FExternInfo GetInfo() const
+		{
+			return {
+				.Name = TEXTVIEW("WorldPosition"),
+				.Type = (WorldPositionShaderOffset == WPT_Default || WorldPositionShaderOffset == WPT_ExcludeAllShaderOffsets)
+				            ? MIR::FType::MakeDoubleVector(3)
+				            : MIR::FType::MakeFloatVector(3),
+			};
+		}
+
+		void Analyze(MIR::FExternAnalysisContext& Context)
+		{
+			Context.Module->GetCompilationOutput().SetIsSceneTextureUsed(PPI_SceneDepth);
+		}
+
+		void ToHLSL(MIR::FExternPrinterHLSL& Printer) const
+		{
+			switch (Printer.Differential)
+			{
+				case MIR::EExternDifferential::None:
+					// Various function permutations exist to fetch world position.  "Prev" permutations are only available in the vertex shader, while
+					// "MoMaterialOffsets" permutations are only available in the pixel shader, so we need to cobble together the permutation string
+					// factoring in those limitations.  Logic adapted from FHLSLMaterialTranslator::WorldPosition.
+					// 
+					// Format:  Get[Prev][Translated]WorldPosition[_NoMaterialOffsets](Parameters)
+					//
+					Printer << TEXTVIEW("Get") << MIR::ExternTag_Prev;
+
+					if (WorldPositionShaderOffset == WPT_CameraRelative || WorldPositionShaderOffset == WPT_CameraRelativeNoOffsets)
+					{
+						Printer << TEXTVIEW("Translated");
+					}
+
+					Printer << TEXTVIEW("WorldPosition");
+
+					if (Printer.Stage == MIR::Stage_Pixel && (WorldPositionShaderOffset == WPT_ExcludeAllShaderOffsets || WorldPositionShaderOffset == WPT_CameraRelativeNoOffsets))
+					{
+						Printer << TEXTVIEW("_NoMaterialOffsets");
+					}
+
+					Printer << TEXTVIEW("(Parameters)");
+					break;
+
+				default:
+					// Expression emitter assumes these to be LWC for analytic derivative evaluation, so we promote them on load
+					Printer << TEXTVIEW("WSPromote(Parameters.WorldPosition_)") << MIR::ExternTag_DD;
+					break;
+			}
+		}
+	};
+
+	FValueRef WorldPosition = Em.Extern<FWorldPosition>({ WorldPositionShaderOffset });
+
+	// CastToNonLWCIfDisabled
+	if (!UE::MaterialTranslatorUtils::IsLWCEnabled())
+	{
+		WorldPosition = Em.CastToFloatKind(WorldPosition);
+	}
+
+	return WorldPosition;
+}
+
+static FValueRef EmitLocalPosition(MIR::FEmitter& Em, ELocalPositionOrigin LocalOrigin, EPositionIncludedOffsets LocalShaderOffset)
+{
+	struct FLocalPosition
+	{
+		bool bLocalOriginIsInstance;
+		bool bIncludeOffsets;
+
+		MIR::FExternInfo GetInfo() const
+		{
+			return {
+				.Name = TEXTVIEW("LocalPosition"),
+				.Type = MIR::FType::MakeFloatVector(3),
+				.Flags = MIR::EExternFlags::Inline,
+			};
+		}
+
+		void ToHLSL(MIR::FExternPrinterHLSL& Printer) const
+		{
+			switch (Printer.Differential)
+			{
+				case MIR::EExternDifferential::None:
+					// Various function permutations exist to fetch local position.  "Prev" permutations are only available in the vertex shader, while
+					// "MoMaterialOffsets" permutations are only available in the pixel shader, so we need to cobble together the permutation string
+					// factoring in those limitations.  Logic adapted from FHLSLMaterialTranslator::LocalPosition (and similar to WorldPosition above).
+					// 
+					// Format:  Get[Prev]Position[Instance|Primitive]Space[_NoMaterialOffsets](Parameters)
+					//
+					Printer << TEXTVIEW("Get") << MIR::ExternTag_Prev << TEXTVIEW("Position");
+					Printer << (bLocalOriginIsInstance ? TEXTVIEW("Instance") : TEXTVIEW("Primitive"));
+					Printer << TEXTVIEW("Space");
+
+					if (Printer.Stage == MIR::Stage_Pixel && !bIncludeOffsets)
+					{
+						Printer << TEXTVIEW("_NoMaterialOffsets");
+					}
+
+					Printer << TEXTVIEW("(Parameters)");
+					break;
+
+				default:
+					Printer << TEXTVIEW("mul(Parameters.WorldPosition_") << MIR::ExternTag_DD;
+					// Local position derivatives reuse the world position derivatives, transforming them to local space
+					if (bLocalOriginIsInstance)
+					{
+						Printer << TEXTVIEW(", DFToFloat3x3(GetWorldToInstanceDF(Parameters)))");
+					}
+					else
+					{
+						Printer << (Printer.bIsPreviousFrame ? TEXTVIEW(", DFToFloat3x3(GetPrevWorldToLocalDF(Parameters)))") : TEXTVIEW(", DFToFloat3x3(GetWorldToLocalDF(Parameters)))"));
+					}
+					break;
+			}
+		}
+	};
+
+	// ELocalPositionOrigin::InstancePreSkinning just uses an external code declaration, and doesn't have variations for offsets.
+	if (LocalOrigin == ELocalPositionOrigin::InstancePreSkinning)
+	{
+		return Em.Extern<MIR::FExternFromMaterialDecl>({ TEXT("PreSkinnedPosition") });
+	}
+
+	// Make sure ELocalPositionOrigin / EPositionIncludedOffsets have the number of entries we expect.
+	if ((int)LocalOrigin >= 3)
+	{
+		Em.Errorf(TEXT("Unsupported LocalPositionOrigin with value '%s'."), *UEnum::GetValueAsString(LocalOrigin));
+		return Em.Poison();
+	}
+	else if ((int)LocalShaderOffset >= 2)
+	{
+		Em.Errorf(TEXT("Unsupported PositionIncludedOffsets with value '%s'."), *UEnum::GetValueAsString(LocalShaderOffset));
+		return Em.Poison();
+	}
+
+	return Em.Extern<FLocalPosition>({ LocalOrigin == ELocalPositionOrigin::Instance, LocalShaderOffset == EPositionIncludedOffsets::IncludeOffsets });
+}
 
 static FValueRef EmitTransformVector(
 	MIR::FEmitter& Em, FValueRef InputValue, EMaterialCommonBasis TransformSourceBasis, EMaterialCommonBasis TransformDestBasis,
@@ -1745,16 +2268,16 @@ static FValueRef BuildVirtualTextureUnpack(MIR::FEmitter& Em, FValueRef SampleCo
 {
 	switch (UnpackType)
 	{
-	case EVirtualTextureUnpackType::BaseColorYCoCg:		return Em.InlineHLSL(MIR::FType::MakeFloatVector(3), TEXTVIEW("VirtualTextureUnpackBaseColorYCoCg($0)"), { SampleCode0 });
-	case EVirtualTextureUnpackType::NormalBC3:			return Em.InlineHLSL(MIR::FType::MakeFloatVector(3), TEXTVIEW("VirtualTextureUnpackNormalBC3($0)"), { SampleCode1 });
-	case EVirtualTextureUnpackType::NormalBC5:			return Em.InlineHLSL(MIR::FType::MakeFloatVector(3), TEXTVIEW("VirtualTextureUnpackNormalBC5($0)"), { SampleCode1 });
-	case EVirtualTextureUnpackType::NormalBC3BC3:		return Em.InlineHLSL(MIR::FType::MakeFloatVector(3), TEXTVIEW("VirtualTextureUnpackNormalBC3BC3($0, $1)"), { SampleCode0, SampleCode1 });
-	case EVirtualTextureUnpackType::NormalBC5BC1:		return Em.InlineHLSL(MIR::FType::MakeFloatVector(3), TEXTVIEW("VirtualTextureUnpackNormalBC5BC1($0, $1)"), { SampleCode1, SampleCode2 });
-	case EVirtualTextureUnpackType::HeightR16:			return Em.InlineHLSL(MIR::FType::MakeFloatScalar(), TEXTVIEW("VirtualTextureUnpackHeight($0, $1)"), { SampleCode0, P0 });
-	case EVirtualTextureUnpackType::DisplacementR16:	return Em.Swizzle(SampleCode0, MIR::EVectorComponent::X);
-	case EVirtualTextureUnpackType::NormalBGR565:		return Em.InlineHLSL(MIR::FType::MakeFloatVector(3), TEXTVIEW("VirtualTextureUnpackNormalBGR565($0)"), { SampleCode1 });
-	case EVirtualTextureUnpackType::BaseColorSRGB:		return Em.InlineHLSL(MIR::FType::MakeFloatVector(3), TEXTVIEW("VirtualTextureUnpackBaseColorSRGB($0)"), { SampleCode0 });
-	default:											UE_MIR_UNREACHABLE();
+		case EVirtualTextureUnpackType::BaseColorYCoCg:		return Em.Extern<MIR::FExternSimpleHLSL>({ MIR::EExternSimpleType::Float3, TEXT("VirtualTextureUnpackBaseColorYCoCg($0)") }, SampleCode0);
+		case EVirtualTextureUnpackType::NormalBC3:			return Em.Extern<MIR::FExternSimpleHLSL>({ MIR::EExternSimpleType::Float3, TEXT("VirtualTextureUnpackNormalBC3($0)") }, SampleCode1);
+		case EVirtualTextureUnpackType::NormalBC5:			return Em.Extern<MIR::FExternSimpleHLSL>({ MIR::EExternSimpleType::Float3, TEXT("VirtualTextureUnpackNormalBC5($0)") }, SampleCode1);
+		case EVirtualTextureUnpackType::NormalBC3BC3:		return Em.Extern<MIR::FExternSimpleHLSL>({ MIR::EExternSimpleType::Float3, TEXT("VirtualTextureUnpackNormalBC3BC3($0, $1)") }, SampleCode0, SampleCode1);
+		case EVirtualTextureUnpackType::NormalBC5BC1:		return Em.Extern<MIR::FExternSimpleHLSL>({ MIR::EExternSimpleType::Float3, TEXT("VirtualTextureUnpackNormalBC5BC1($0, $1)") }, SampleCode1, SampleCode2);
+		case EVirtualTextureUnpackType::HeightR16:			return Em.Extern<MIR::FExternSimpleHLSL>({ MIR::EExternSimpleType::Float1, TEXT("VirtualTextureUnpackHeight($0, $1)") }, SampleCode0, P0);
+		case EVirtualTextureUnpackType::DisplacementR16:	return Em.Swizzle(SampleCode0, MIR::EVectorComponent::X);
+		case EVirtualTextureUnpackType::NormalBGR565:		return Em.Extern<MIR::FExternSimpleHLSL>({ MIR::EExternSimpleType::Float3, TEXT("VirtualTextureUnpackNormalBGR565($0)") }, SampleCode1);
+		case EVirtualTextureUnpackType::BaseColorSRGB:		return Em.Extern<MIR::FExternSimpleHLSL>({ MIR::EExternSimpleType::Float3, TEXT("VirtualTextureUnpackBaseColorSRGB($0)") }, SampleCode0);
+		default:											UE_MIR_UNREACHABLE();
 	}
 }
 
@@ -1769,7 +2292,7 @@ void UMaterialExpressionRuntimeVirtualTextureSample::Build(MIR::FEmitter& Em)
 	// Is this a valid UMaterialExpressionRuntimeVirtualTextureSampleParameter?
 	const bool bIsParameter = IsParameter();
 
-	// Check validity of current virtual texture
+	// Check validity of current texture
 	FString TextureValidityError;
 	bool bIsVirtualTextureValid = false;
 	if (!ValidateVirtualTextureParameters(TextureValidityError, bIsVirtualTextureValid))
@@ -1805,7 +2328,7 @@ void UMaterialExpressionRuntimeVirtualTextureSample::Build(MIR::FEmitter& Em)
 
 	UE_MIR_CHECKPOINT(Em);
 
-	// Compile the runtime virtual texture uniforms
+	// Compile the runtime texture uniforms
 	MIR::FValueRef Uniforms[ERuntimeVirtualTextureShaderUniform_Count];
 
 	for (int32 UniformIndex = 0; UniformIndex < ERuntimeVirtualTextureShaderUniform_Count; ++UniformIndex)
@@ -1817,7 +2340,7 @@ void UMaterialExpressionRuntimeVirtualTextureSample::Build(MIR::FEmitter& Em)
 	}
 
 	// Compile the coordinates
-	// We use the virtual texture world space transform by default
+	// We use the texture world space transform by default
 	if (Coordinates.GetTracedInput().Expression != nullptr && WorldPosition.GetTracedInput().Expression != nullptr)
 	{
 		Em.Error(TEXT("Only one of 'Coordinates' and 'WorldPosition' can be used"));
@@ -1913,7 +2436,7 @@ void UMaterialExpressionRuntimeVirtualTextureSample::Build(MIR::FEmitter& Em)
 		// This is legacy/hidden, and is better implemented in the material graph using RVTMVM_DerivativeWorld.
 		TextureMipLevelMode = TMVM_Derivative;
 
-		MIR::FValueRef WorldPos = Em.ExternalInput(MIR::EExternalInput::WorldPosition_CameraRelative);
+		MIR::FValueRef WorldPos = EmitWorldPosition(Em, WPT_CameraRelative);
 		MIR::FValueRef WorldPositionDdx = Em.AnalyticalPartialDerivative(WorldPos, MIR::EDerivativeAxis::X);
 		MIR::FValueRef UDdx = Em.Dot(WorldPositionDdx, Uniforms[ERuntimeVirtualTextureShaderUniform_WorldToUVTransform1]);
 		MIR::FValueRef VDdx = Em.Dot(WorldPositionDdx, Uniforms[ERuntimeVirtualTextureShaderUniform_WorldToUVTransform2]);
@@ -1931,7 +2454,7 @@ void UMaterialExpressionRuntimeVirtualTextureSample::Build(MIR::FEmitter& Em)
 	// Compile the texture sample code
 	constexpr bool bAutomaticMipViewBias = true;
 
-	const MIR::FTextureSampleBaseAttributes SampleAttributes
+	const MIR::FTextureSampleAttributes SampleAttributes
 	{
 		.SamplerSourceMode = GetSamplerSourceMode(),
 		.SamplerType = SAMPLERTYPE_VirtualMasks,
@@ -1970,8 +2493,8 @@ void UMaterialExpressionRuntimeVirtualTextureSample::Build(MIR::FEmitter& Em)
 	// Compile unpacking code
 	for (int32 OutputIndex = 0; OutputIndex < 8; ++OutputIndex)
 	{
-		// Calculate the virtual texture layer and sampling/unpacking functions for this output
-		// Fallback to a sensible default value if the output isn't valid for the bound virtual texture
+		// Calculate the texture layer and sampling/unpacking functions for this output
+		// Fallback to a sensible default value if the output isn't valid for the bound texture
 		FRuntimeVirtualTextureUnpackProperties UnpackProperties;
 		if (!GetRVTUnpackProperties(OutputIndex, bIsVirtualTextureValid, UnpackProperties))
 		{
@@ -2008,13 +2531,10 @@ void UMaterialExpressionRuntimeVirtualTextureSample::Build(MIR::FEmitter& Em)
 
 void UMaterialExpressionTime::Build(MIR::FEmitter& Em)
 {
-	MIR::FType ScalarFloatType = MIR::FType::MakeScalar(MIR::EScalarKind::Float);
-
 	// When pausing the game is ignored for this time expression, use real-time instead of game-time.
 	if (!bOverride_Period)
 	{
-		FStringView InlinedCode = bIgnorePause ? TEXTVIEW("View.<PREVFRAME>RealTime") : TEXTVIEW("View.<PREVFRAME>GameTime");
-		Em.Output(0, Em.InlineHLSL(ScalarFloatType, InlinedCode, {}, MIR::EValueFlags::SubstituteTags));
+		Em.Output(0, Em.Extern<MIR::FExternSimpleHLSL>({ MIR::EExternSimpleType::Float1, bIgnorePause ? TEXT("View.<PREVFRAME>RealTime") : TEXT("View.<PREVFRAME>GameTime") }));
 	}
 	else if (Period == 0.0f)
 	{
@@ -2025,8 +2545,7 @@ void UMaterialExpressionTime::Build(MIR::FEmitter& Em)
 		// Note: Don't use IR intrinsic for Fmod() here to avoid conversion to fp16 on mobile.
 		// We want full 32 bit float precision until the fmod when using a period.
 		FValueRef PeriodValue = Em.ConstantFloat(Period);
-		FStringView InlinedCode = bIgnorePause ? TEXTVIEW("fmod(View.<PREVFRAME>RealTime, $0)") : TEXTVIEW("fmod(View.<PREVFRAME>GameTime, $0)");
-		Em.Output(0, Em.InlineHLSL(ScalarFloatType, InlinedCode, { PeriodValue }, MIR::EValueFlags::SubstituteTags));
+		Em.Output(0, Em.Extern<MIR::FExternSimpleHLSL>({ MIR::EExternSimpleType::Float1, bIgnorePause ? TEXT("fmod(View.<PREVFRAME>RealTime, $0)") : TEXT("fmod(View.<PREVFRAME>GameTime, $0)") }, PeriodValue));
 	}
 }
 
@@ -2055,18 +2574,17 @@ static bool IsConstFloatOfPow2Expression(FValueRef TileScaleIndexValue)
 
 static FValueRef EmitPeriodicWorldPositionOrigin(MIR::FEmitter& Em, FValueRef TileScaleIndexValue)
 {
-	FStringView PeriodicWorldOriginFunctionName = IsConstFloatOfPow2Expression(TileScaleIndexValue) ? TEXTVIEW("GetPeriodicWorldOrigin_Pow2($0)") : TEXTVIEW("GetPeriodicWorldOrigin($0)");
-	return Em.InlineHLSL(TileScaleIndexValue->Type, PeriodicWorldOriginFunctionName, { TileScaleIndexValue });
+	return Em.Extern<MIR::FExternSimpleHLSL>({ MIR::EExternSimpleType::Float3, IsConstFloatOfPow2Expression(TileScaleIndexValue) ? TEXT("GetPeriodicWorldOrigin_Pow2($0)") : TEXT("GetPeriodicWorldOrigin($0)") }, TileScaleIndexValue);
 }
 
 static FValueRef EmitFixedExternalCode(MIR::FEmitter& Em, FName InExternalCodeIdentifier)
 {
-	return Em.InlineHLSL(InExternalCodeIdentifier, {}, MIR::EValueFlags::SubstituteTags);
+	return Em.Extern<MIR::FExternFromMaterialDecl>(InExternalCodeIdentifier);
 }
 
 static FValueRef EmitMatrixCastTo3x3(MIR::FEmitter& Em, FValueRef MatrixValue)
 {
-	return Em.InlineHLSL(MIR::FType::MakeFloat(3, 3), MatrixValue->Type.IsDouble() ? TEXTVIEW("DFToFloat3x3($0)") : TEXTVIEW("(float3x3)$0"), { MatrixValue });
+	return Em.Extern<MIR::FExternSimpleHLSL>({ MIR::EExternSimpleType::Float3x3, MatrixValue->Type.IsDouble() ? TEXT("DFToFloat3x3($0)") : TEXT("(float3x3)$0") }, MatrixValue);
 }
 
 static FValueRef EmitMatrixMultiply(MIR::FEmitter& Em, FValueRef VectorValue, FValueRef MatrixValue, bool bHasWComponent)
@@ -2089,7 +2607,7 @@ static FValueRef EmitMultiplyTranslatedMatrix(MIR::FEmitter& Em, FValueRef Vecto
 	if (bHasWComponent)
 	{
 		// mul(Float4(V, 1.0), DFFastToTranslatedWorld(M, ResolvedView.PreViewTranslation))
-		MatrixPreTranslation = Em.InlineHLSL(MIR::FType::MakeFloat(4, 4), TEXTVIEW("DFFastToTranslatedWorld($0, ResolvedView.<PREV>PreViewTranslation)"), { MatrixPreTranslation }, MIR::EValueFlags::SubstituteTags);
+		MatrixPreTranslation = Em.Extern<MIR::FExternSimpleHLSL>({ MIR::EExternSimpleType::Float4x4, TEXT("DFFastToTranslatedWorld($0, ResolvedView.<PREV>PreViewTranslation)") }, MatrixPreTranslation);
 		return Em.MatrixMultiply(Em.Vector4(VectorValue, Em.ConstantOne(MIR::EScalarKind::Float)), MatrixPreTranslation);
 	}
 	else
@@ -2101,9 +2619,15 @@ static FValueRef EmitMultiplyTranslatedMatrix(MIR::FEmitter& Em, FValueRef Vecto
 
 static FValueRef EmitMultiplyLWCMatrix(MIR::FEmitter& Em, FValueRef VectorValue, FValueRef MatrixValue, bool bHasWComponent, bool bDemote)
 {
-	return bHasWComponent
-		? Em.InlineHLSL(VectorValue->Type, bDemote ? TEXTVIEW("WSMultiplyDemote($0, $1)") : TEXTVIEW("WSMultiply($0, $1)"), { VectorValue, MatrixValue })
-		: Em.InlineHLSL(VectorValue->Type, TEXTVIEW("WSMultiplyVector($0, $1)"), { VectorValue, MatrixValue });
+	static const MIR::FExternInlineDeclaration DeclA{ .Type = MIR::FType::MakeFloatVector(3), .CodeHLSL = TEXT("WSMultiplyDemote($0, $1)") };
+	static const MIR::FExternInlineDeclaration DeclB{ .Type = MIR::FType::MakeDoubleVector(3), .CodeHLSL = TEXT("WSMultiply($0, $1)") };
+	static const MIR::FExternInlineDeclaration DeclC{ .Type = MIR::FType::MakeFloatVector(3), .CodeHLSL = TEXT("WSMultiplyVector($0, $1)") };
+
+	const MIR::FExternInlineDeclaration& Decl = bHasWComponent
+		? ( bDemote ? DeclA : DeclB )
+		: DeclC;
+
+	return Em.Extern<MIR::FExternFromInlineDecl>(&Decl, VectorValue, MatrixValue);
 }
 
 static FValueRef EmitTransformVector(MIR::FEmitter& Em, FValueRef InputValue, EMaterialCommonBasis TransformSourceBasis, EMaterialCommonBasis TransformDestBasis,
@@ -2136,7 +2660,7 @@ static FValueRef EmitTransformVector(MIR::FEmitter& Em, FValueRef InputValue, EM
 			check(!bIsPositionTransform);
 			if (TransformDestBasis == MCB_World)
 			{
-				return Em.MatrixMultiply(InputValue, Em.InlineHLSL(TEXT("TangentToWorld")));
+				return Em.MatrixMultiply(InputValue, Em.Extern<MIR::FExternFromMaterialDecl>({ TEXT("TangentToWorld") }));
 			}
 			// else use MCB_World as intermediary basis
 			break;
@@ -2146,8 +2670,11 @@ static FValueRef EmitTransformVector(MIR::FEmitter& Em, FValueRef InputValue, EM
 			{
 				case MCB_World:
 				{
-					FStringView Code = bIsPositionTransform ? TEXTVIEW("TransformLocalPositionTo<PREV>World(Parameters, $0)") : TEXTVIEW("TransformLocalVectorTo<PREV>World(Parameters, $0)");
-					return Em.InlineHLSL(ResultType, Code, { InputValue }, MIR::EValueFlags::SubstituteTags);
+					return Em.Extern<MIR::FExternSimpleHLSL>(
+						{
+							MIR::EExternSimpleType::Float3,
+							bIsPositionTransform ? TEXT("TransformLocalPositionTo<PREV>World(Parameters, $0)") : TEXT("TransformLocalVectorTo<PREV>World(Parameters, $0)")
+						}, InputValue);
 				}
 
 				case MCB_TranslatedWorld:
@@ -2211,7 +2738,7 @@ static FValueRef EmitTransformVector(MIR::FEmitter& Em, FValueRef InputValue, EM
 					}
 					// The first person transform is actually a 3x3 matrix and can therefore be used for derivatives as well.
 					MIR::FValueRef LerpAlphaClampedIndexValue = Em.Saturate(Em.CastToFloat(FirstPersonInterpolationAlphaValue, 1));
-					return Em.InlineHLSL(ResultType, TEXT("TransformTo<PREVIOUS>FirstPerson($0, $1)"), { InputValue, LerpAlphaClampedIndexValue }, MIR::EValueFlags::SubstituteTags);
+					return Em.Extern<MIR::FExternSimpleHLSL>({ MIR::EExternSimpleType::Float3, TEXT("TransformTo<PREVIOUS>FirstPerson($0, $1)") }, InputValue, LerpAlphaClampedIndexValue);
 				}
 
 				default:
@@ -2694,20 +3221,37 @@ void UMaterialExpressionConvert::Build(MIR::FEmitter& Em)
 	}
 }
 
-static FValueRef BuildViewProperty(MIR::FEmitter& Em, EMaterialExposedViewProperty InProperty, bool bInvProperty = false)
-{
-	check(InProperty < MEVP_MAX);
+static FValueRef BuildViewProperty(MIR::FEmitter& Em, EMaterialExposedViewProperty InProperty, bool bInvProperty = false)
+{
+	check(InProperty < MEVP_MAX);
+
+	struct FViewPropertyExtern
+	{
+		const FMaterialExposedViewPropertyMeta* PropertyMeta;
+		bool bInvProperty;
+
+		MIR::FExternInfo GetInfo() const
+		{
+			return {
+				.Name = TEXT("ViewProperty"),
+				.Type = MIR::FType::FromMaterialValueType(PropertyMeta->Type),
+				.Flags = MIR::EExternFlags::Inline | MIR::EExternFlags::ZeroDifferentials,
+			};
+		}
+
+		void ToHLSL(MIR::FExternPrinterHLSL& Printer) const
+		{
+			Printer << (bInvProperty ? PropertyMeta->InvPropertyCode : PropertyMeta->PropertyCode);
+		}
+	};
 
 	const FMaterialExposedViewPropertyMeta& PropertyMeta = MaterialExternalCodeRegistry::Get().GetExternalViewPropertyCode(InProperty);
 	const bool bHasCustomInverseCode = PropertyMeta.InvPropertyCode != nullptr;
 
-	FStringView HLSLCode = bInvProperty && bHasCustomInverseCode ? PropertyMeta.InvPropertyCode : PropertyMeta.PropertyCode;
-	MIR::FType HLSLCodeType = MIR::FType::FromMaterialValueType(PropertyMeta.Type);
-
-	FValueRef Result = Em.InlineHLSL(HLSLCodeType, HLSLCode, {}, MIR::EValueFlags::SubstituteTags);
+	FValueRef Result = Em.Extern<FViewPropertyExtern>({ &PropertyMeta, bInvProperty && bHasCustomInverseCode });
 
 	// CastToNonLWCIfDisabled
-	TOptional<MIR::FPrimitive> PrimitiveType = HLSLCodeType.AsPrimitive();
+	TOptional<MIR::FPrimitive> PrimitiveType = MIR::FType::FromMaterialValueType(PropertyMeta.Type).AsPrimitive();
 	if (PrimitiveType && PrimitiveType->IsDouble() && !UE::MaterialTranslatorUtils::IsLWCEnabled())
 	{
 		Result = Em.CastToFloatKind(Result);
@@ -2750,8 +3294,7 @@ void UMaterialExpressionCameraPositionWS::Build(MIR::FEmitter& Em)
 
 void UMaterialExpressionPixelNormalWS::Build(MIR::FEmitter& Em)
 {
-	FValueRef Output = Em.InlineHLSL(MIR::FType::MakeFloatVector(3), TEXTVIEW("Parameters.WorldNormal"), {}, MIR::EValueFlags::None, MIR::EGraphProperties::ReadsPixelNormal);
-	Em.Output(0, Output);
+	Em.Output(0, EmitPixelNormalWS(Em));
 }
 
 void UMaterialExpressionDDX::Build(MIR::FEmitter& Em)
@@ -2837,81 +3380,76 @@ void UMaterialExpressionTruncateLWC::Build(MIR::FEmitter& Em)
 	Em.Output(0, InputValue);
 }
 
-void UMaterialExpressionActorPositionWS::Build(MIR::FEmitter& Em)
+enum class EActorOrObjectPosition
 {
-	if (OriginType == EPositionOrigin::CameraRelative)
-	{
-		Em.Output(0, Em.ExternalInput(MIR::EExternalInput::ActorPosition_CameraRelative));
-	}
-	else if (!UE::MaterialTranslatorUtils::IsLWCEnabled())
-	{
-		// CastToNonLWCIfDisabled
-		Em.Output(0, Em.CastToFloatKind(Em.ExternalInput(MIR::EExternalInput::ActorPosition_Absolute)));
-	}
-	else
-	{
-		Em.Output(0, Em.ExternalInput(MIR::EExternalInput::ActorPosition_Absolute));
-	}
-}
+	Actor, Object,
+};
 
-void UMaterialExpressionObjectPositionWS::Build(MIR::FEmitter& Em)
+static FValueRef EmitActorOrObjectPosition(MIR::FEmitter& Em, EActorOrObjectPosition Mode, EPositionOrigin Origin)
 {
-	if (OriginType == EPositionOrigin::CameraRelative)
-	{
-		Em.Output(0, Em.ExternalInput(MIR::EExternalInput::ObjectPosition_CameraRelative));
-	}
-	else if (!UE::MaterialTranslatorUtils::IsLWCEnabled())
-	{
-		// CastToNonLWCIfDisabled
-		Em.Output(0, Em.CastToFloatKind(Em.ExternalInput(MIR::EExternalInput::ObjectPosition_Absolute)));
-	}
-	else
+	struct FActorAndObjectPositionExtern
 	{
-		Em.Output(0, Em.ExternalInput(MIR::EExternalInput::ObjectPosition_Absolute));
-	}
-}
+		EActorOrObjectPosition Mode;
+		bool bIsAbsolute;
 
-static FValueRef EmitWorldPosition(MIR::FEmitter& Em, EWorldPositionIncludedOffsets WorldPositionShaderOffset)
-{
-	// Make sure EWorldPositionIncludedOffsets and corresponding elements of EExternalInput stay in sync, so this enum addition is valid
-	static_assert(WPT_MAX == 4 && WPT_Default == 0);
-	static_assert((int32)MIR::EExternalInput::WorldPosition_AbsoluteNoOffsets - (int32)MIR::EExternalInput::WorldPosition_Absolute == WPT_ExcludeAllShaderOffsets);
-	static_assert((int32)MIR::EExternalInput::WorldPosition_CameraRelative - (int32)MIR::EExternalInput::WorldPosition_Absolute == WPT_CameraRelative);
-	static_assert((int32)MIR::EExternalInput::WorldPosition_CameraRelativeNoOffsets - (int32)MIR::EExternalInput::WorldPosition_Absolute == WPT_CameraRelativeNoOffsets);
+		FActorAndObjectPositionExtern(EActorOrObjectPosition Mode, bool bIsAbsolute)
+		: Mode { Mode }
+		, bIsAbsolute { bIsAbsolute }
+		{}
 
-	FValueRef WorldPosition = Em.ExternalInput((MIR::EExternalInput)((int32)MIR::EExternalInput::WorldPosition_Absolute + FMath::Clamp(WorldPositionShaderOffset, 0, WPT_MAX - 1)));
+		MIR::FExternInfo GetInfo() const
+		{
+			return {
+				.Name = TEXTVIEW("ActorAndObjectPosition"),
+				.Type = bIsAbsolute ? MIR::FType::MakeDoubleVector(3) : MIR::FType::MakeFloatVector(3),
+			};
+		}
 
-	// CastToNonLWCIfDisabled
-	if (!UE::MaterialTranslatorUtils::IsLWCEnabled())
+		void ToHLSL(MIR::FExternPrinterHLSL& Printer) const
+		{
+			if (Printer.Differential != MIR::EExternDifferential::None)
+			{
+				Printer << (bIsAbsolute ? TEXTVIEW("WSPromote(MaterialFloat3(0))") : TEXTVIEW("MaterialFloat3(0)"));
+				return;
+			}
+			// Get<Previous><Actor|Object><Translated>WorldPosition(Parameters)
+			Printer << TEXTVIEW("Get");
+			if (Mode == EActorOrObjectPosition::Actor && Printer.bIsPreviousFrame)
+			{
+				// Object doesn't have previous
+				Printer << TEXTVIEW("Previous");
+			}
+			Printer << (Mode == EActorOrObjectPosition::Actor ? TEXTVIEW("Actor") : TEXTVIEW("Object"));
+			if (!bIsAbsolute)
+			{
+				Printer << TEXTVIEW("Translated");
+			}
+			Printer << TEXTVIEW("WorldPosition(Parameters)");
+		}
+
+		void EmitDebugInfo(FString& Out) const
+		{
+			Out.Appendf(TEXT("IsActorPosition=%d, bIsAbsolute=%d"), Mode == EActorOrObjectPosition::Actor, bIsAbsolute);
+		}
+	};
+
+	FValueRef Value = Em.Extern<FActorAndObjectPositionExtern>({ Mode, Origin == EPositionOrigin::Absolute });
+	if (Origin == EPositionOrigin::Absolute && !UE::MaterialTranslatorUtils::IsLWCEnabled())
 	{
-		WorldPosition = Em.CastToFloatKind(WorldPosition);
+		Value = Em.CastToFloatKind(Value);
 	}
 
-	return WorldPosition;
+	return Value;
 }
 
-static FValueRef EmitLocalPosition(MIR::FEmitter& Em, ELocalPositionOrigin LocalOrigin, EPositionIncludedOffsets LocalShaderOffset)
+void UMaterialExpressionActorPositionWS::Build(MIR::FEmitter& Em)
 {
-	// Make sure ELocalPositionOrigin / EPositionIncludedOffsets and corresponding elements of EExternalInput stay in sync, so math using enum is valid.
-	static_assert((int32)ELocalPositionOrigin::Instance == 0);
-	static_assert((int32)ELocalPositionOrigin::Primitive == 2);
-	static_assert((int32)EPositionIncludedOffsets::IncludeOffsets == 0);
-	static_assert((int32)EPositionIncludedOffsets::ExcludeOffsets == 1);
-	static_assert((int32)MIR::EExternalInput::LocalPosition_InstanceNoOffsets - (int32)MIR::EExternalInput::LocalPosition_Instance == (int32)EPositionIncludedOffsets::ExcludeOffsets);
-	static_assert((int32)MIR::EExternalInput::LocalPosition_Primitive - (int32)MIR::EExternalInput::LocalPosition_Instance == (int32)ELocalPositionOrigin::Primitive);
-
-	// ELocalPositionOrigin::InstancePreSkinning just uses an external code declaration, and doesn't have variations for offsets.
-	if (LocalOrigin == ELocalPositionOrigin::InstancePreSkinning)
-	{
-		static FName NAME_PreSkinnedPosition("PreSkinnedPosition");
-		return Em.InlineHLSL(NAME_PreSkinnedPosition);
-	}
+	Em.Output(0, EmitActorOrObjectPosition(Em, EActorOrObjectPosition::Actor, OriginType));
+}
 
-	// Otherwise, there's a 2x2 configuration of Origin and Offset type.  Given that ELocalPositionOrigin::InstancePreSkinning is unused, valid origin values are 0 or 2,
-	// so we can just add the values together to get a unique index from [0..3]
-	int32 OriginIndex = FMath::Clamp((int32)LocalOrigin, 0, 2);
-	int32 OffsetIndex = FMath::Clamp((int32)LocalShaderOffset, 0, 1);
-	return Em.ExternalInput((MIR::EExternalInput)((int32)MIR::EExternalInput::LocalPosition_Instance + OriginIndex + OffsetIndex));
+void UMaterialExpressionObjectPositionWS::Build(MIR::FEmitter& Em)
+{
+	Em.Output(0, EmitActorOrObjectPosition(Em, EActorOrObjectPosition::Object, OriginType));
 }
 
 // Adapted from FHLSLMaterialTranslator::GetWorldPositionOrDefault.  WorldPosition input is optional, a default is provided if not set.
@@ -2939,13 +3477,17 @@ void UMaterialExpressionWorldPosition::Build(MIR::FEmitter& Em)
 	FValueRef WorldPosition = EmitWorldPosition(Em, WorldPositionShaderOffset);
 
 	Em.Output(0, WorldPosition);
-	Em.Output(1, Em.Swizzle(WorldPosition, MIR::FSwizzleMask(MIR::EVectorComponent::X, MIR::EVectorComponent::Y)));
+	Em.Output(1, Em.Swizzle(WorldPosition, MIR::FSwizzleMask::XY()));
 	Em.Output(2, Em.Subscript(WorldPosition, 2));
 }
 
 void UMaterialExpressionLocalPosition::Build(MIR::FEmitter& Em)
 {
-	Em.Output(0, EmitLocalPosition(Em, LocalOrigin, IncludedOffsets));
+	FValueRef LocalPosition = EmitLocalPosition(Em, LocalOrigin, IncludedOffsets);
+
+	Em.Output(0, LocalPosition);
+	Em.Output(1, Em.Swizzle(LocalPosition, MIR::FSwizzleMask::XY()));
+	Em.Output(2, Em.Subscript(LocalPosition, 2));
 }
 
 void UMaterialExpressionMakeMaterialAttributes::Build(MIR::FEmitter& Em)
@@ -3289,9 +3831,9 @@ void UMaterialExpressionBlend::Build(MIR::FEmitter& Em)
 	Em.Output(0, Result);
 }
 
-static FValueRef EmitParameterCollectionVectorInlineHLSL(MIR::FEmitter& Em, MIR::FValueRef CollectionValue, int32 ParameterIndex)
+static FValueRef EmitParameterCollectionVectorExtern(MIR::FEmitter& Em, MIR::FValueRef CollectionValue, int32 ParameterIndex)
 {
-	return Em.InlineHLSL(MIR::FType::MakeFloatVector(4), TEXTVIEW("MaterialCollection$0.Vectors[$1]"), { CollectionValue, Em.ConstantInt(ParameterIndex) });
+	return Em.Extern<MIR::FExternSimpleHLSL>({ MIR::EExternSimpleType::Float4, TEXT("MaterialCollection$0.Vectors[$1]") }, CollectionValue, Em.ConstantInt(ParameterIndex));
 }
 
 static bool GetExpressionCollectionParameter(MIR::FEmitter& Em, UMaterialParameterCollection* Collection, FName ParameterName, const FGuid& ParameterId, int32& OutParamIndex, int32& OutComponentIndex)
@@ -3321,7 +3863,7 @@ void UMaterialExpressionCollectionParameter::Build(MIR::FEmitter& Em)
 		return;
 	}
 
-	FValueRef Result = EmitParameterCollectionVectorInlineHLSL(Em, Em.MaterialParameterCollection(Collection), ParameterIndex);
+	FValueRef Result = EmitParameterCollectionVectorExtern(Em, Em.MaterialParameterCollection(Collection), ParameterIndex);
 	if (ComponentIndex != INDEX_NONE)
 	{
 		Result = Em.Subscript(Result, ComponentIndex % 4);
@@ -3389,7 +3931,7 @@ void UMaterialExpressionCollectionTransform::Build(MIR::FEmitter& Em)
 
 	for (int32 i = 0; i < NumVectors; i++)
 	{
-		CollectionParameters.Add(EmitParameterCollectionVectorInlineHLSL(Em, CollectionValue, ParameterIndex + i));
+		CollectionParameters.Add(EmitParameterCollectionVectorExtern(Em, CollectionValue, ParameterIndex + i));
 	}
 
 	MIR::FValueRef Result;
@@ -3463,12 +4005,13 @@ void UMaterialExpressionAtmosphericFogColor::Build(MIR::FEmitter& Em)
 {
 	// This node is deprecated in favor of UMaterialExpressionSkyAtmosphereAerialPerspective, and falls through to the newer expression
 	FValueRef PositionValue = EmitWorldPositionOrDefault(Em, Em.TryInput(&WorldPosition), WorldPositionOriginType);
-	Em.Output(0, MaterialExpressionExternalCodeBase_EmitInlineHLSL(Em, *this, 0, { PositionValue }));
+	Em.Output(0, MaterialExpressionExternalCodeBase_EmitExtern(Em, *this, 0, PositionValue));
 }
 
 void UMaterialExpressionBlackBody::Build(MIR::FEmitter& Em)
 {
-	Em.Output(0, Em.InlineHLSL(MIR::FType::MakeFloatVector(3), TEXTVIEW("MaterialExpressionBlackBody($0)"), { Em.Input(&Temp) }));
+	FValueRef TempValue = Em.CastToFloat(Em.InputDefaultFloat(&Temp, 0.f), 1);
+	Em.Output(0, Em.Extern<MIR::FExternSimpleHLSL>({ MIR::EExternSimpleType::Float3, TEXT("MaterialExpressionBlackBody($0)") }, TempValue));
 }
 
 void UMaterialExpressionDepthFade::Build(MIR::FEmitter& Em)
@@ -3479,7 +4022,7 @@ void UMaterialExpressionDepthFade::Build(MIR::FEmitter& Em)
 	FValueRef FadeDistanceValue = Em.Max(Em.InputDefaultFloat(&FadeDistance, FadeDistanceDefault), Em.ConstantFloat(UE_DELTA));
 
 	static FName NAME_PixelDepth("PixelDepth");
-	FValueRef PixelDepth = Em.InlineHLSL(NAME_PixelDepth);
+	FValueRef PixelDepth = Em.Extern<MIR::FExternFromMaterialDecl>(NAME_PixelDepth);
 	// On mobile scene depth is limited to 65500 
 	// to avoid false fading on objects that are close or exceed this limit we clamp pixel depth to (65500 - FadeDistance)
 	if (Em.GetFeatureLevel() <= ERHIFeatureLevel::ES3_1)
@@ -3488,7 +4031,8 @@ void UMaterialExpressionDepthFade::Build(MIR::FEmitter& Em)
 	}
 
 	// We need a dependency on EScreenTexture::SceneDepth, so the value analyzer can see it, even though it's technically not used in the code.
-	FValueRef SceneDepth = Em.InlineHLSL(MIR::FType::MakeFloatScalar(), TEXTVIEW("CalcSceneDepth(ScreenAlignedPosition(GetScreenPosition(Parameters)))"), { Em.ScreenTexture(MIR::EScreenTexture::SceneDepth) });
+	FValueRef SceneDepth = Em.Nop(Em.Extern<FScreenTextureExtern>({ EScreenTexture::SceneDepth }));
+	SceneDepth = Em.Extern<MIR::FExternSimpleHLSL>({ MIR::EExternSimpleType::Float1, TEXT("CalcSceneDepth(ScreenAlignedPosition(GetScreenPosition(Parameters)))") }, SceneDepth);
 
 	Em.Output(0, Em.Multiply(OpacityValue, Em.Saturate(Em.Divide(Em.Subtract(SceneDepth, PixelDepth), FadeDistanceValue))));
 }
@@ -3504,7 +4048,7 @@ void UMaterialExpressionDeriveNormalZ::Build(MIR::FEmitter& Em)
 	
 	Em.Output(0, Em.Vector3(Em.Subscript(InputVector, 0), Em.Subscript(InputVector, 1), DerivedZ));
 }
-
+	
 void UMaterialExpressionDistanceFieldApproxAO::Build(MIR::FEmitter& Em)
 {
 	FValueRef PositionValue = EmitWorldPositionOrDefault(Em, Em.TryInput(&Position), WorldPositionOriginType);
@@ -3513,7 +4057,7 @@ void UMaterialExpressionDistanceFieldApproxAO::Build(MIR::FEmitter& Em)
 	if (!NormalValue)
 	{
 		static FName NAME_VertexNormal("VertexNormal");
-		NormalValue = Em.InlineHLSL(NAME_VertexNormal);
+		NormalValue = Em.Extern<MIR::FExternFromMaterialDecl>(NAME_VertexNormal);
 	}
 
 	FValueRef BaseDistanceValue = Em.InputDefaultFloat(&BaseDistance, BaseDistanceDefault);
@@ -3544,32 +4088,94 @@ void UMaterialExpressionDistanceFieldApproxAO::Build(MIR::FEmitter& Em)
 		MaxDistance = RadiusValue;
 	}
 
+	struct FGlobalDistanceFieldExtern
+	{
+		MIR::FExternInfo GetInfo() const
+		{
+			return {
+				.Name  = TEXTVIEW("GlobalDistanceField"),
+				.Type  = MIR::FType::MakeFloatScalar(),
+				.Flags = MIR::EExternFlags::ZeroDifferentials,
+			};
+		}
+	
+		void Analyze(MIR::FExternAnalysisContext& Context)
+		{
+			if (!FDataDrivenShaderPlatformInfo::GetSupportsDistanceFields(Context.Module->GetShaderPlatform()))
+			{
+				FString ShaderPlatformName = FDataDrivenShaderPlatformInfo::GetName(Context.Module->GetShaderPlatform()).ToString();
+				Context.Module->AddError(nullptr, FString::Printf(TEXT("Node not supported in shader platform %s. The node requires DistanceField support."), *ShaderPlatformName));
+				return;
+			}
+
+			Context.Module->GetCompilationOutput().bUsesGlobalDistanceField = true;
+		}
+
+		void ToHLSL(MIR::FExternPrinterHLSL& Printer) const
+		{
+			Printer << TEXT("CalculateDistanceFieldApproxAO($0, $1, $2, $3, $4, $5, $6)");
+		}
+	};
+
 	// Last input tells value analyzer that this expression uses the global distance field
-	Em.Output(0, Em.InlineHLSL(MIR::FType::MakeFloatScalar(), TEXTVIEW("CalculateDistanceFieldApproxAO($0, $1, $2, $3, $4, $5, $6)"),
+	Em.Output(0, Em.Extern<FGlobalDistanceFieldExtern>({},
+		PositionValue,
+		Em.Cast(NormalValue, MIR::FType::MakeFloatVector(3)),
+		NumStepsConst,
+		Em.Cast(StepDistance, MIR::FType::MakeFloatScalar()),
+		StepScaleConst,
+		Em.Cast(DistanceBias, MIR::FType::MakeFloatScalar()),
+		Em.Cast(MaxDistance, MIR::FType::MakeFloatScalar())
+	));
+}
+
+static void EmitDistanceFieldExpression(MIR::FEmitter& Em, FExpressionInput& Position, EPositionOrigin WorldPositionOriginType, EMaterialValueType Type,  const TCHAR* FunctionName)
+{
+	struct FDistanceFieldHelper
+	{
+		EMaterialValueType Type;
+		const TCHAR* FunctionName;
+
+		MIR::FExternInfo GetInfo() const
+		{
+			return {
+				.Name = TEXTVIEW("DistanceFieldHelper"),
+				.Type = MIR::FType::FromMaterialValueType(Type),
+				.Flags = MIR::EExternFlags::Inline,
+			};
+		}
+
+		void Analyze(MIR::FExternAnalysisContext& Context)
+		{
+			if (!FDataDrivenShaderPlatformInfo::GetSupportsDistanceFields(Context.Module->GetShaderPlatform()))
+			{
+				FString ShaderPlatformName = FDataDrivenShaderPlatformInfo::GetName(Context.Module->GetShaderPlatform()).ToString();
+				Context.Module->AddError(nullptr, FString::Printf(TEXT("Node not supported in shader platform %s. The node requires DistanceField support."), *ShaderPlatformName));
+				return;
+			}
+
+			Context.Module->GetCompilationOutput().bUsesGlobalDistanceField = true;
+		}
+
+		void ToHLSL(MIR::FExternPrinterHLSL& Printer) const
 		{
-			PositionValue,
-			Em.Cast(NormalValue, MIR::FType::MakeFloatVector(3)),
-			NumStepsConst,
-			Em.Cast(StepDistance, MIR::FType::MakeFloatScalar()),
-			StepScaleConst,
-			Em.Cast(DistanceBias, MIR::FType::MakeFloatScalar()),
-			Em.Cast(MaxDistance, MIR::FType::MakeFloatScalar()),
-			Em.ExternalInput(MIR::EExternalInput::GlobalDistanceField)
-		}));
+			Printer << FunctionName << TEXT('(') << Printer.Arg(0) << TEXT(')');
+		}
+	};
+
+	FValueRef PositionValue = EmitWorldPositionOrDefault(Em, Em.TryInput(&Position), WorldPositionOriginType);
+	FValueRef Extern = Em.Extern<FDistanceFieldHelper>({ Type, FunctionName }, PositionValue);
+	Em.Output(0, Extern);
 }
 
 void UMaterialExpressionDistanceFieldGradient::Build(MIR::FEmitter& Em)
 {
-	// Last input tells value analyzer that this expression uses the global distance field
-	FValueRef PositionValue = EmitWorldPositionOrDefault(Em, Em.TryInput(&Position), WorldPositionOriginType);
-	Em.Output(0, Em.InlineHLSL(MIR::FType::MakeFloatVector(3), TEXTVIEW("GetDistanceFieldGradientGlobal($0)"), { PositionValue, Em.ExternalInput(MIR::EExternalInput::GlobalDistanceField) }));
+	EmitDistanceFieldExpression(Em, Position, WorldPositionOriginType, MCT_Float3, TEXT("GetDistanceFieldGradientGlobal"));
 }
 
 void UMaterialExpressionDistanceToNearestSurface::Build(MIR::FEmitter& Em)
 {
-	// Last input tells value analyzer that this expression uses the global distance field
-	FValueRef PositionValue = EmitWorldPositionOrDefault(Em, Em.TryInput(&Position), WorldPositionOriginType);
-	Em.Output(0, Em.InlineHLSL(MIR::FType::MakeFloatScalar(), TEXTVIEW("GetDistanceToNearestSurfaceGlobal($0)"), { PositionValue, Em.ExternalInput(MIR::EExternalInput::GlobalDistanceField) }));
+	EmitDistanceFieldExpression(Em, Position, WorldPositionOriginType, MCT_Float, TEXT("GetDistanceToNearestSurfaceGlobal"));
 }
 
 void UMaterialExpressionFresnel::Build(MIR::FEmitter& Em)
@@ -3583,13 +4189,14 @@ void UMaterialExpressionFresnel::Build(MIR::FEmitter& Em)
 	}
 	else
 	{
-		NormalArg = Em.InlineHLSL(MIR::FType::MakeFloatVector(3), TEXTVIEW("Parameters.WorldNormal"), {}, MIR::EValueFlags::None, MIR::EGraphProperties::ReadsPixelNormal);
+		NormalArg = EmitPixelNormalWS(Em);
 	}
 
-	FValueRef DotArg = Em.Dot(NormalArg, Em.InlineHLSL(NAME_CameraVector));
+	FValueRef DotArg = Em.Dot(NormalArg, Em.Extern<MIR::FExternFromMaterialDecl>({ NAME_CameraVector }));
 	FValueRef MaxArg = Em.Max(Em.ConstantFloat(0.f), DotArg);
 	FValueRef MinusArg = Em.Subtract(Em.ConstantFloat(1.f), MaxArg);
 	FValueRef ExponentArg = Em.InputDefaultFloat(&ExponentIn, Exponent);
+	
 	// Compiler->Power got changed to call PositiveClampedPow instead of ClampedPow
 	// Manually implement ClampedPow to maintain backwards compatibility in the case where the input normal is not normalized (length > 1)
 	FValueRef AbsBaseArg = Em.Max(Em.Abs(MinusArg), Em.ConstantFloat(UE_KINDA_SMALL_NUMBER));
@@ -3615,13 +4222,13 @@ void UMaterialExpressionReflectionVectorWS::Build(MIR::FEmitter& Em)
 		}
 
 		// Normal * dot(Normal, CameraVector) * 2.0 - CameraVector;
-		FValueRef CameraVector = Em.InlineHLSL(NAME_CameraVector);
+		FValueRef CameraVector = Em.Extern<MIR::FExternFromMaterialDecl>({ NAME_CameraVector });
 		Em.Output(0, Em.Subtract(Em.Multiply(NormalValue, Em.Multiply(Em.Dot(NormalValue, CameraVector), Em.ConstantFloat(2.0f))), CameraVector));
 	}
 	else
 	{
 		static FName NAME_ReflectionVector("ReflectionVector");
-		Em.Output(0, Em.InlineHLSL(NAME_ReflectionVector));
+		Em.Output(0, Em.Extern<MIR::FExternFromMaterialDecl>({ NAME_ReflectionVector }));
 	}
 }
 
@@ -3657,14 +4264,14 @@ void UMaterialExpressionRotator::Build(MIR::FEmitter& Em)
 	FValueRef TimeValue = Em.TryInput(&Time);
 	if (!TimeValue)
 	{
-		TimeValue = Em.InlineHLSL(MIR::FType::MakeFloatScalar(), TEXTVIEW("View.<PREVFRAME>GameTime"), {}, MIR::EValueFlags::SubstituteTags);
+		TimeValue = Em.Extern<MIR::FExternSimpleHLSL>({ MIR::EExternSimpleType::Float1, TEXT("View.<PREVFRAME>GameTime") });
 	}
 	TimeValue = Em.Multiply(TimeValue, Em.ConstantFloat(Speed));
 
 	FValueRef BaseCoordinate = Em.TryInput(&Coordinate);
 	if (!BaseCoordinate)
 	{
-		BaseCoordinate = Em.ExternalInput(MIR::TexCoordIndexToExternalInput(ConstCoordinate));
+		BaseCoordinate = EmitTexCoord(Em, ConstCoordinate);
 	}
 	BaseCoordinate = Em.Subtract(BaseCoordinate, Em.ConstantFloat2({ CenterX, CenterY }));
 
@@ -3688,12 +4295,12 @@ void UMaterialExpressionRotator::Build(MIR::FEmitter& Em)
 void UMaterialExpressionSkyAtmosphereAerialPerspective::Build(MIR::FEmitter& Em)
 {
 	FValueRef PositionValue = EmitWorldPositionOrDefault(Em, Em.TryInput(&WorldPosition), WorldPositionOriginType);
-	Em.Output(0, MaterialExpressionExternalCodeBase_EmitInlineHLSL(Em, *this, 0, { PositionValue }));
+	Em.Output(0, MaterialExpressionExternalCodeBase_EmitExtern(Em, *this, 0, PositionValue));
 }
 
 void UMaterialExpressionSkyAtmosphereLightDirection::Build(MIR::FEmitter& Em)
 {
-	Em.Output(0, MaterialExpressionExternalCodeBase_EmitInlineHLSL(Em, *this, 0, { Em.ConstantInt(LightIndex) }));
+	Em.Output(0, MaterialExpressionExternalCodeBase_EmitExtern(Em, *this, 0, Em.ConstantInt(LightIndex)));
 }
 
 void UMaterialExpressionSkyAtmosphereLightDiskLuminance::Build(MIR::FEmitter& Em)
@@ -3708,18 +4315,18 @@ void UMaterialExpressionSkyAtmosphereLightDiskLuminance::Build(MIR::FEmitter& Em
 	{
 		CosHalfDiskRadius = Em.ConstantFloat(-1.0f);
 	}
-	Em.Output(0, MaterialExpressionExternalCodeBase_EmitInlineHLSL(Em, *this, 0, { Em.ConstantInt(LightIndex), CosHalfDiskRadius }));
+	Em.Output(0, MaterialExpressionExternalCodeBase_EmitExtern(Em, *this, 0, Em.ConstantInt(LightIndex), CosHalfDiskRadius));
 }
 
 void UMaterialExpressionSkyAtmosphereLightIlluminance::Build(MIR::FEmitter& Em)
 {
 	FValueRef PositionValue = EmitWorldPositionOrDefault(Em, Em.TryInput(&WorldPosition), WorldPositionOriginType);
-	Em.Output(0, MaterialExpressionExternalCodeBase_EmitInlineHLSL(Em, *this, 0, { PositionValue, Em.ConstantInt(LightIndex) }));
+	Em.Output(0, MaterialExpressionExternalCodeBase_EmitExtern(Em, *this, 0, PositionValue, Em.ConstantInt(LightIndex)));
 }
 
 void UMaterialExpressionSkyAtmosphereLightIlluminanceOnGround::Build(MIR::FEmitter& Em)
 {
-	Em.Output(0, MaterialExpressionExternalCodeBase_EmitInlineHLSL(Em, *this, 0, { Em.ConstantInt(LightIndex) }));
+	Em.Output(0, MaterialExpressionExternalCodeBase_EmitExtern(Em, *this, 0, Em.ConstantInt(LightIndex)));
 }
 
 void UMaterialExpressionSkyAtmosphereViewLuminance::Build(MIR::FEmitter& Em)
@@ -3727,9 +4334,9 @@ void UMaterialExpressionSkyAtmosphereViewLuminance::Build(MIR::FEmitter& Em)
 	FValueRef WorldDirectionValue = Em.TryInput(&WorldDirection);
 	if (!WorldDirectionValue)
 	{
-		WorldDirectionValue = Em.Multiply(Em.ConstantFloat(-1.0f), Em.InlineHLSL(NAME_CameraVector));
+		WorldDirectionValue = Em.Multiply(Em.ConstantFloat(-1.0f), Em.Extern<MIR::FExternFromMaterialDecl>({ NAME_CameraVector }));
 	}
-	Em.Output(0, MaterialExpressionExternalCodeBase_EmitInlineHLSL(Em, *this, 0, { WorldDirectionValue }));
+	Em.Output(0, MaterialExpressionExternalCodeBase_EmitExtern(Em, *this, 0, WorldDirectionValue));
 }
 
 void UMaterialExpressionSkyLightEnvMapSample::Build(MIR::FEmitter& Em)
@@ -3742,7 +4349,7 @@ void UMaterialExpressionSkyLightEnvMapSample::Build(MIR::FEmitter& Em)
 	FValueRef DirectionValue = Em.InputDefaultFloat3(&Direction, FVector3f(0.0f, 0.0f, 1.0f));
 	FValueRef RoughnessValue = Em.InputDefaultFloat(&Roughness, 0.0f);
 
-	Em.Output(0, MaterialExpressionExternalCodeBase_EmitInlineHLSL(Em, *this, 0, { DirectionValue, RoughnessValue }));
+	Em.Output(0, MaterialExpressionExternalCodeBase_EmitExtern(Em, *this, 0, DirectionValue, RoughnessValue));
 }
 
 void UMaterialExpressionSphereMask::Build(MIR::FEmitter& Em)
@@ -3924,10 +4531,10 @@ void UMaterialExpressionBounds::Build(MIR::FEmitter& Em)
 	static_assert(MEILB_PreSkinnedLocal == 2);
 
 	int32 OutputOffset = Type * 4;
-	Em.Output(0, MaterialExpressionExternalCodeBase_EmitInlineHLSL(Em, *this, OutputOffset + 0));
-	Em.Output(1, MaterialExpressionExternalCodeBase_EmitInlineHLSL(Em, *this, OutputOffset + 1));
-	Em.Output(2, MaterialExpressionExternalCodeBase_EmitInlineHLSL(Em, *this, OutputOffset + 2));
-	Em.Output(3, MaterialExpressionExternalCodeBase_EmitInlineHLSL(Em, *this, OutputOffset + 3));
+	Em.Output(0, MaterialExpressionExternalCodeBase_EmitExtern(Em, *this, OutputOffset + 0));
+	Em.Output(1, MaterialExpressionExternalCodeBase_EmitExtern(Em, *this, OutputOffset + 1));
+	Em.Output(2, MaterialExpressionExternalCodeBase_EmitExtern(Em, *this, OutputOffset + 2));
+	Em.Output(3, MaterialExpressionExternalCodeBase_EmitExtern(Em, *this, OutputOffset + 3));
 }
 
 void UMaterialExpressionBumpOffset::Build(MIR::FEmitter& Em)
@@ -3937,13 +4544,13 @@ void UMaterialExpressionBumpOffset::Build(MIR::FEmitter& Em)
 	FValueRef TexCoordArg = Em.TryInput(&Coordinate);
 	if (!TexCoordArg)
 	{
-		TexCoordArg = Em.ExternalInput(MIR::TexCoordIndexToExternalInput(ConstCoordinate));
+		TexCoordArg = EmitTexCoord(Em, ConstCoordinate);
 	}
-
+	
 	Em.Output(0,
 		Em.Add(
 			Em.Multiply(
-				Em.Swizzle(EmitTransformVector(Em, Em.InlineHLSL(NAME_CameraVector), MCB_World, MCB_Tangent, false, nullptr, nullptr), MIR::FSwizzleMask(MIR::EVectorComponent::X, MIR::EVectorComponent::Y)),
+				Em.Swizzle(EmitTransformVector(Em, Em.Extern<MIR::FExternFromMaterialDecl>({ NAME_CameraVector }), MCB_World, MCB_Tangent, false, nullptr, nullptr), MIR::FSwizzleMask(MIR::EVectorComponent::X, MIR::EVectorComponent::Y)),
 				Em.Add(
 					Em.Multiply(
 						HeightRatioArg,
@@ -3959,9 +4566,47 @@ void UMaterialExpressionBumpOffset::Build(MIR::FEmitter& Em)
 
 void UMaterialExpressionDynamicParameter::Build(MIR::FEmitter& Em)
 {
-	FValueRef Default = Em.ConstantFloat4(FVector4f(DefaultValue.R, DefaultValue.G, DefaultValue.B, DefaultValue.A));
-	FValueRef DynamicParameterIndex = Em.ExternalInput(MIR::EExternalInput::DynamicParticleParameterIndex, ParameterIndex);
-	FValueRef Result = Em.InlineHLSL(MIR::FType::MakeFloatVector(4), TEXTVIEW("GetDynamicParameter(Parameters.Particle, $0, $1)"), { Default, DynamicParameterIndex });
+	struct FDynamicParameter
+	{
+		FVector4f Default;
+		uint32 ParameterIndex;
+			
+		MIR::FExternInfo GetInfo() const
+		{
+			return {
+				.Name = TEXTVIEW("DynamicParameter"),
+				.Type = MIR::FType::MakeFloatVector(4),
+				.Flags = MIR::EExternFlags::Inline,
+			};
+		}
+
+		void Analyze(MIR::FExternAnalysisContext& Context)
+		{
+			Context.Module->GetStatistics().DynamicParticleParameterMask |= 1 << ParameterIndex;
+		}
+
+		void ToHLSL(MIR::FExternPrinterHLSL& Printer) const
+		{
+			if (Printer.Differential != MIR::EExternDifferential::None)
+			{
+				Printer << TEXT("0");
+				return;
+			}
+			Printer << TEXTVIEW("GetDynamicParameter(Parameters.Particle, MaterialFloat4(")
+				    << Default.X << TEXTVIEW(", ") << Default.Y << TEXTVIEW(", ") << Default.Z << TEXTVIEW(", ") << Default.W << TEXTVIEW("), ")
+				    << ParameterIndex << TEXTVIEW(")");
+		}
+
+		void EmitDebugInfo(FString& Out) const
+		{
+			Out.Appendf(TEXT("Default=(%f, %f, %f, %f) ParameterIndex=%d"), Default.X, Default.Y, Default.Z, Default.W, ParameterIndex);
+		}
+	};
+
+	FValueRef Result = Em.Extern<FDynamicParameter>({
+		FVector4f(DefaultValue.R, DefaultValue.G, DefaultValue.B, DefaultValue.A),
+		ParameterIndex
+	});
 
 	Em.Output(0, Em.Subscript(Result, 0));
 	Em.Output(1, Em.Subscript(Result, 1));
@@ -3993,21 +4638,21 @@ void UMaterialExpressionNoise::Build(MIR::FEmitter& Em)
 	FValueRef TilingValue = Em.ConstantBool(bTiling);
 	FValueRef RepeatSizeValue = Em.ConstantFloat(RepeatSize);
 
-	Em.Output(0, Em.InlineHLSL(MIR::FType::MakeFloatScalar(),
-		TEXTVIEW("MaterialExpressionNoise($0,$1,$2,$3,$4,$5,$6,$7,$8,$9,$10,$11)"), {
-			PositionInput,
-			ScaleValue,
-			QualityValue,
-			NoiseFunctionValue,
-			TurbulenceValue,
-			LevelsValue,
-			OutputMinValue,
-			OutputMaxValue,
-			LevelScaleValue,
-			FilterWidthInput,
-			TilingValue,
-			RepeatSizeValue
-		}));
+	Em.Output(0, Em.Extern<MIR::FExternSimpleHLSL>(
+		{ MIR::EExternSimpleType::Float1, TEXT("MaterialExpressionNoise($0,$1,$2,$3,$4,$5,$6,$7,$8,$9,$10,$11)") },
+		PositionInput,
+		ScaleValue,
+		QualityValue,
+		NoiseFunctionValue,
+		TurbulenceValue,
+		LevelsValue,
+		OutputMinValue,
+		OutputMaxValue,
+		LevelScaleValue,
+		FilterWidthInput,
+		TilingValue,
+		RepeatSizeValue
+	));
 }
 
 void UMaterialExpressionVectorNoise::Build(MIR::FEmitter& Em)
@@ -4028,14 +4673,14 @@ void UMaterialExpressionVectorNoise::Build(MIR::FEmitter& Em)
 	FValueRef TilingValue = Em.ConstantBool(bTiling);
 	FValueRef TileSizeValue = Em.ConstantFloat(TileSize);
 
-	FValueRef NoiseResult = Em.InlineHLSL(MIR::FType::MakeFloatVector(4),
-		TEXTVIEW("MaterialExpressionVectorNoise($0,$1,$2,$3,$4)"), {
-			PositionInput,
-			QualityValue,
-			NoiseFunctionValue,
-			TilingValue,
-			TileSizeValue
-		});
+	FValueRef NoiseResult = Em.Extern<MIR::FExternSimpleHLSL>(
+		{ MIR::EExternSimpleType::Float4, TEXT("MaterialExpressionVectorNoise($0,$1,$2,$3,$4)") },
+		PositionInput,
+		QualityValue,
+		NoiseFunctionValue,
+		TilingValue,
+		TileSizeValue
+	);
 
 	// Function returns float4, but only certain noise functions fill in all four elements, so downcast to float3 if not those cases.
 	if (NoiseFunction != VNF_GradientALU && NoiseFunction != VNF_VoronoiALU)
@@ -4055,7 +4700,7 @@ void UMaterialExpressionPanner::Build(MIR::FEmitter& Em)
 	}
 	else
 	{
-		TimeArg = Em.InlineHLSL(MIR::FType::MakeFloatScalar(), TEXTVIEW("View.<PREVFRAME>GameTime"), {}, MIR::EValueFlags::SubstituteTags);
+		TimeArg = Em.Extern<MIR::FExternSimpleHLSL>({ MIR::EExternSimpleType::Float1, TEXT("View.<PREVFRAME>GameTime") });
 	}
 	
 	FValueRef SpeedVectorArg = Em.InputDefaultFloat2(&Speed, FVector2f(SpeedX, SpeedY));
@@ -4074,7 +4719,7 @@ void UMaterialExpressionPanner::Build(MIR::FEmitter& Em)
 	FValueRef TexCoordArg = Em.TryInput(&Coordinate);
 	if (!TexCoordArg)
 	{
-		TexCoordArg = Em.ExternalInput(MIR::TexCoordIndexToExternalInput(0));
+		TexCoordArg = EmitTexCoord(Em, 0);
 	}
 
 	Em.Output(0, Em.Add(SpeedVectorArg, TexCoordArg));
@@ -4085,7 +4730,7 @@ void UMaterialExpressionParticlePositionWS::Build(MIR::FEmitter& Em)
 	int32 ExternalCodeIndex = OriginType == EPositionOrigin::Absolute ? 0 : 1;
 	MIR::FType ResultType = OriginType == EPositionOrigin::Absolute ? MIR::FType::MakeDoubleVector(3) : MIR::FType::MakeFloatVector(3);
 
-	Em.Output(0, MaterialExpressionExternalCodeBase_EmitInlineHLSL(Em, *this, ExternalCodeIndex, {}, MIR::EValueFlags::SubstituteTags));
+	Em.Output(0, MaterialExpressionExternalCodeBase_EmitExtern(Em, *this, ExternalCodeIndex));
 }
 
 void UMaterialExpressionPerInstanceCustomData::Build(MIR::FEmitter& Em)
@@ -4093,7 +4738,7 @@ void UMaterialExpressionPerInstanceCustomData::Build(MIR::FEmitter& Em)
 	FValueRef DataIndexArgument = Em.ConstantInt(DataIndex);
 	FValueRef DefaultArgument = Em.InputDefaultFloat(&DefaultValue, ConstDefaultValue);
 
-	Em.Output(0, MaterialExpressionExternalCodeBase_EmitInlineHLSL(Em, *this, 0, { DataIndexArgument, DefaultArgument }));
+	Em.Output(0, MaterialExpressionExternalCodeBase_EmitExtern(Em, *this, 0, DataIndexArgument, DefaultArgument));
 }
 
 void UMaterialExpressionPerInstanceCustomData3Vector::Build(MIR::FEmitter& Em)
@@ -4101,7 +4746,7 @@ void UMaterialExpressionPerInstanceCustomData3Vector::Build(MIR::FEmitter& Em)
 	FValueRef DataIndexArgument = Em.ConstantInt(DataIndex);
 	FValueRef DefaultArgument = Em.InputDefaultFloat3(&DefaultValue, FVector3f(ConstDefaultValue.R, ConstDefaultValue.G, ConstDefaultValue.B));
 
-	Em.Output(0, MaterialExpressionExternalCodeBase_EmitInlineHLSL(Em, *this, 0, { DataIndexArgument, DefaultArgument }));
+	Em.Output(0, MaterialExpressionExternalCodeBase_EmitExtern(Em, *this, 0, DataIndexArgument, DefaultArgument));
 }
 
 void UMaterialExpressionPreviousFrameSwitch::Build(MIR::FEmitter& Em)
@@ -4110,31 +4755,52 @@ void UMaterialExpressionPreviousFrameSwitch::Build(MIR::FEmitter& Em)
 	FValueRef PreviousFrameValue = Em.Input(&PreviousFrame);
 	UE_MIR_CHECKPOINT(Em);
 
-	Em.Output(0, Em.Branch(Em.ExternalInput(MIR::EExternalInput::CompilingPreviousFrame), PreviousFrameValue, CurrentFrameValue));
+	struct FPreviousFrame
+	{
+		MIR::FExternInfo GetInfo() const
+		{
+			return {
+				.Name  = TEXTVIEW("PreviousFrameExtern"),
+				.Type  = MIR::FType::MakeBoolScalar(),
+				.Flags = MIR::EExternFlags::Inline,
+			};
+		}
+	
+		void ToHLSL(MIR::FExternPrinterHLSL& Printer) const
+		{
+			Printer << (Printer.bIsPreviousFrame ? TEXTVIEW("true") : TEXTVIEW("false"));
+		}
+	};
+
+	Em.Output(0, Em.Branch(Em.Extern<FPreviousFrame>({}), PreviousFrameValue, CurrentFrameValue));
 }
 
 void UMaterialExpressionHairAttributes::Build(MIR::FEmitter& Em)
 {
-	FValueRef HairUV = Em.InlineHLSL(MIR::FType::MakeFloatVector(2), TEXTVIEW("MaterialExpressionGetHairUV(Parameters)"));
-	Em.Output(0, Em.Subscript(HairUV, 0));
-	Em.Output(1, Em.Subscript(HairUV, 1));
+	{
+		FValueRef HairUV = Em.Extern<MIR::FExternSimpleHLSL>({ MIR::EExternSimpleType::Float2, TEXT("MaterialExpressionGetHairUV(Parameters)") });
+		Em.Output(0, Em.Subscript(HairUV, 0));
+		Em.Output(1, Em.Subscript(HairUV, 1));
+	}
 
-	FValueRef HairDimensions = Em.InlineHLSL(MIR::FType::MakeFloatVector(2), TEXTVIEW("MaterialExpressionGetHairDimensions(Parameters)"));
-	Em.Output(2, Em.Subscript(HairDimensions, 0));		// Length
-	Em.Output(3, Em.Subscript(HairDimensions, 1));		// Radius
+	{
+		FValueRef HairDimensions = Em.Extern<MIR::FExternSimpleHLSL>({ MIR::EExternSimpleType::Float2, TEXT("MaterialExpressionGetHairDimensions(Parameters)") });
+		Em.Output(2, Em.Subscript(HairDimensions, 0));		// Length
+		Em.Output(3, Em.Subscript(HairDimensions, 1));		// Radius
+	}
 
-	Em.Output(4, Em.InlineHLSL(MIR::FType::MakeFloatScalar(), TEXTVIEW("MaterialExpressionGetHairSeed(Parameters)")));
-	Em.Output(5, Em.InlineHLSL(MIR::FType::MakeFloatVector(3), TEXTVIEW("MaterialExpressionGetHairTangent(Parameters, $0)"), { Em.ConstantBool(bUseTangentSpace) }));
-	Em.Output(6, Em.InlineHLSL(MIR::FType::MakeFloatVector(2), TEXTVIEW("MaterialExpressionGetHairRootUV(Parameters)")));
-	Em.Output(7, Em.InlineHLSL(MIR::FType::MakeFloatVector(3), TEXTVIEW("MaterialExpressionGetHairBaseColor(Parameters)")));
-	Em.Output(8, Em.InlineHLSL(MIR::FType::MakeFloatScalar(), TEXTVIEW("MaterialExpressionGetHairRoughness(Parameters)")));
-	Em.Output(9, Em.InlineHLSL(MIR::FType::MakeFloatScalar(), TEXTVIEW("MaterialExpressionGetHairDepth(Parameters)")));
-	Em.Output(10, Em.InlineHLSL(MIR::FType::MakeFloatScalar(), TEXTVIEW("MaterialExpressionGetHairCoverage(Parameters)")));
-	Em.Output(11, Em.InlineHLSL(MIR::FType::MakeFloatVector(4), TEXTVIEW("MaterialExpressionGetHairAuxilaryData(Parameters)")));
-	Em.Output(12, Em.InlineHLSL(MIR::FType::MakeFloatVector(2), TEXTVIEW("MaterialExpressionGetAtlasUVs(Parameters)")));
-	Em.Output(13, Em.InlineHLSL(MIR::FType::MakeFloatScalar(), TEXTVIEW("MaterialExpressionGetHairGroupIndex(Parameters)")));
-	Em.Output(14, Em.InlineHLSL(MIR::FType::MakeFloatScalar(), TEXTVIEW("MaterialExpressionGetHairAO(Parameters)")));
-	Em.Output(15, Em.Subscript(Em.InlineHLSL(MIR::FType::MakeFloatVector(3), TEXTVIEW("MaterialExpressionGetHairClumpID(Parameters)")), 0));
+	Em.Output(4, Em.Extern<MIR::FExternSimpleHLSL>({ MIR::EExternSimpleType::Float1, TEXT("MaterialExpressionGetHairSeed(Parameters)") }));
+	Em.Output(5, Em.Extern<MIR::FExternSimpleHLSL>({ MIR::EExternSimpleType::Float3, TEXT("MaterialExpressionGetHairTangent(Parameters, $0)") }, Em.ConstantBool(bUseTangentSpace)));
+	Em.Output(6, Em.Extern<MIR::FExternSimpleHLSL>({ MIR::EExternSimpleType::Float2, TEXT("MaterialExpressionGetHairRootUV(Parameters)") }));
+	Em.Output(7, Em.Extern<MIR::FExternSimpleHLSL>({ MIR::EExternSimpleType::Float3, TEXT("MaterialExpressionGetHairBaseColor(Parameters)") }));
+	Em.Output(8, Em.Extern<MIR::FExternSimpleHLSL>({ MIR::EExternSimpleType::Float1, TEXT("MaterialExpressionGetHairRoughness(Parameters)") }));
+	Em.Output(9, Em.Extern<MIR::FExternSimpleHLSL>({ MIR::EExternSimpleType::Float1, TEXT("MaterialExpressionGetHairDepth(Parameters)") }));
+	Em.Output(10, Em.Extern<MIR::FExternSimpleHLSL>({ MIR::EExternSimpleType::Float1, TEXT("MaterialExpressionGetHairCoverage(Parameters)") }));
+	Em.Output(11, Em.Extern<MIR::FExternSimpleHLSL>({ MIR::EExternSimpleType::Float4, TEXT("MaterialExpressionGetHairAuxilaryData(Parameters)") }));
+	Em.Output(12, Em.Extern<MIR::FExternSimpleHLSL>({ MIR::EExternSimpleType::Float2, TEXT("MaterialExpressionGetAtlasUVs(Parameters)") }));
+	Em.Output(13, Em.Extern<MIR::FExternSimpleHLSL>({ MIR::EExternSimpleType::Float1, TEXT("MaterialExpressionGetHairGroupIndex(Parameters)") }));
+	Em.Output(14, Em.Extern<MIR::FExternSimpleHLSL>({ MIR::EExternSimpleType::Float1, TEXT("MaterialExpressionGetHairAO(Parameters)") }));
+	Em.Output(15, Em.Subscript(Em.Extern<MIR::FExternSimpleHLSL>({ MIR::EExternSimpleType::Float3, TEXT("MaterialExpressionGetHairClumpID(Parameters)") }), 0));
 }
 
 void UMaterialExpressionHairColor::Build(MIR::FEmitter& Em)
@@ -4143,7 +4809,7 @@ void UMaterialExpressionHairColor::Build(MIR::FEmitter& Em)
 	FValueRef RednessInput = Em.InputDefaultFloat(&Redness, 0.0f);
 	FValueRef DyeColorInput = Em.InputDefaultFloat3(&DyeColor, FVector3f(1.f, 1.f, 1.f));
 
-	Em.Output(0, Em.InlineHLSL(MIR::FType::MakeFloatVector(3), TEXTVIEW("MaterialExpressionGetHairColorFromMelanin($0, $1, $2)"), { MelaninInput, RednessInput, DyeColorInput }));
+	Em.Output(0, Em.Extern<MIR::FExternSimpleHLSL>({ MIR::EExternSimpleType::Float3, TEXT("MaterialExpressionGetHairColorFromMelanin($0, $1, $2)") }, MelaninInput, RednessInput, DyeColorInput));
 }
 
 void UMaterialExpressionMapARPassthroughCameraUV::Build(MIR::FEmitter& Em)
@@ -4151,8 +4817,8 @@ void UMaterialExpressionMapARPassthroughCameraUV::Build(MIR::FEmitter& Em)
 	FValueRef UV = Em.Input(&Coordinates);
 	UE_MIR_CHECKPOINT(Em);
 
-	FValueRef UVPair0 = Em.InlineHLSL(MIR::FType::MakeFloatVector(4), TEXTVIEW("ResolvedView.XRPassthroughCameraUVs[0]"));
-	FValueRef UVPair1 = Em.InlineHLSL(MIR::FType::MakeFloatVector(4), TEXTVIEW("ResolvedView.XRPassthroughCameraUVs[1]"));
+	FValueRef UVPair0 = Em.Extern<MIR::FExternSimpleHLSL>({ MIR::EExternSimpleType::Float4, TEXT("ResolvedView.XRPassthroughCameraUVs[0]") });
+	FValueRef UVPair1 = Em.Extern<MIR::FExternSimpleHLSL>({ MIR::EExternSimpleType::Float4, TEXT("ResolvedView.XRPassthroughCameraUVs[1]") });
 
 	FValueRef ULerp = Em.Lerp(UVPair0, UVPair1, Em.Subscript(UV, 0));
 	Em.Output(0, Em.Lerp(Em.Swizzle(ULerp, MIR::FSwizzleMask(MIR::EVectorComponent::X, MIR::EVectorComponent::Y)), Em.Swizzle(ULerp, MIR::FSwizzleMask(MIR::EVectorComponent::Z, MIR::EVectorComponent::W)), Em.Subscript(UV, 1)));
@@ -4313,7 +4979,7 @@ void UMaterialExpressionThinTranslucentMaterialOutput::Build(MIR::FEmitter& Em)
 
 void UMaterialExpressionVolumetricAdvancedMaterialInput::Build(MIR::FEmitter& Em)
 {
-	FValueRef VolumeSampleConservativeDensity = Em.InlineHLSL(TEXT("GetVolumeSampleConservativeDensity"));
+	FValueRef VolumeSampleConservativeDensity = Em.Extern<MIR::FExternFromMaterialDecl>({ TEXT("GetVolumeSampleConservativeDensity") });
 	Em.Output(0, Em.CastToFloat(VolumeSampleConservativeDensity, 3));
 	Em.Output(1, VolumeSampleConservativeDensity);
 }
@@ -4347,7 +5013,7 @@ void UMaterialExpressionDecalMipmapLevel::Build(MIR::FEmitter& Em)
 	Em.CheckMaterialDomainIsAnyOf({ MD_DeferredDecal });
 
 	FValueRef TextureSizeValue = Em.CastToFloat(Em.InputDefaultFloat2(&TextureSize, { ConstWidth, ConstHeight }), 2);
-	BuildInlineHLSLOutput(Em, *this, { TextureSizeValue });
+	BuildMaterialExpressionExternalCodeBase(Em, *this, TextureSizeValue);
 }
 
 void UMaterialExpressionDepthOfFieldFunction::Build(MIR::FEmitter& Em)
@@ -4355,16 +5021,38 @@ void UMaterialExpressionDepthOfFieldFunction::Build(MIR::FEmitter& Em)
 	FValueRef DepthValue = Em.CastToFloat(Em.TryInput(&Depth), 1);
 	if (!DepthValue)
 	{
-		DepthValue = Em.InlineHLSL(FName("PixelDepth"));
+		DepthValue = Em.Extern<MIR::FExternFromMaterialDecl>({ TEXT("PixelDepth") });
 	}
 
 	UE_MIR_CHECKPOINT(Em);
-	BuildInlineHLSLOutput(Em, *this, { DepthValue, Em.ConstantInt((MIR::TInteger)FunctionValue) });
+	BuildMaterialExpressionExternalCodeBase(Em, *this, DepthValue, Em.ConstantInt((MIR::TInteger)FunctionValue));
 }
 
 void UMaterialExpressionDistanceCullFade::Build(MIR::FEmitter& Em)
 {
-	Em.Output(0, Em.ExternalInput(MIR::EExternalInput::DistanceCullFade));
+	struct FDistanceCullFade
+	{
+		MIR::FExternInfo GetInfo() const
+		{
+			return {
+				.Name = TEXTVIEW("DistanceCullFade"),
+				.Type = MIR::FType::MakeFloatScalar(),
+				.Flags = MIR::EExternFlags::Inline,
+			};
+		}
+
+		void Analyze(MIR::FExternAnalysisContext& Context)
+		{
+			Context.Module->GetCompilationOutput().bUsesDistanceCullFade = true;
+		}
+
+		void ToHLSL(MIR::FExternPrinterHLSL& Printer) const
+		{
+			Printer << TEXTVIEW("GetDistanceCullFade()");
+		}
+	};
+
+	Em.Output(0, Em.Extern<FDistanceCullFade>({}));
 }
 
 void UMaterialExpressionGIReplace::Build(MIR::FEmitter& Em)
@@ -4374,7 +5062,7 @@ void UMaterialExpressionGIReplace::Build(MIR::FEmitter& Em)
 
 	// Note: StaticIndirect is unused.
 	FValueRef LocalDynamicIndirectValue = Em.InputDefault(&DynamicIndirect, DefaultValue);
-	FValueRef GIReplaceStateValue = Em.InlineHLSL(MIR::FType::MakeBoolScalar(), TEXT("GetGIReplaceState()"));
+	FValueRef GIReplaceStateValue = Em.Extern<MIR::FExternSimpleHLSL>({ MIR::EExternSimpleType::Bool1, TEXT("GetGIReplaceState()") });
 
 	Em.Output(0, Em.Branch(GIReplaceStateValue, LocalDynamicIndirectValue, DefaultValue));
 }
@@ -4448,24 +5136,38 @@ void UMaterialExpressionPathTracingBufferTexture::Build(MIR::FEmitter& Em)
 	FValueRef CoordinatesValue = Em.TryInput(&Coordinates);
 	if (!CoordinatesValue)
 	{
-		CoordinatesValue = Em.InlineHLSL(MIR::FType::MakeFloatVector(2), TEXT("GetDefaultPathTracingBufferTextureUV(Parameters, 0)"));
+		CoordinatesValue = Em.Extern<MIR::FExternSimpleHLSL>({ MIR::EExternSimpleType::Float2, TEXT("GetDefaultPathTracingBufferTextureUV(Parameters, 0)") });
 	}
 	CoordinatesValue = Em.CastToFloat(CoordinatesValue, 2);
 
 	FValueRef PathTracingBufferTextureIndexValue = Em.ConstantInt(PathTracingBufferTextureId);
 
-	FValueRef OutputValue = Em.InlineHLSL(
-		MIR::FType::MakeFloatVector(4),
-		TEXTVIEW("MaterialExpressionPathTracingBufferTextureLookup(Parameters, $0, $1)"),
-		{ CoordinatesValue, PathTracingBufferTextureIndexValue },
-		MIR::EValueFlags::None, MIR::EGraphProperties::None,
-		[](FMaterialIRModule* Module, uintptr_t PathTracingBufferTextureId)
+	struct FPathTracingBufferTextureExtern
+	{
+		EPathTracingBufferTextureId PathTracingBufferTextureId;
+
+		MIR::FExternInfo GetInfo() const
 		{
-			Module->GetCompilationOutput().SetIsPathTracingBufferTextureUsed((int32)PathTracingBufferTextureId);
-		},
-		PathTracingBufferTextureId);
+			return
+				{
+					.Name = TEXT("PathTracingBufferTexture"),
+					.Type  = MIR::FType::MakeFloatVector(4),
+					.Flags = MIR::EExternFlags::Inline | MIR::EExternFlags::ZeroDifferentials, 
+				};
+		}
 
-	Em.Output(0, OutputValue);
+		void Analyze(MIR::FExternAnalysisContext& Context)
+		{
+			Context.Module->GetCompilationOutput().SetIsPathTracingBufferTextureUsed((int32)PathTracingBufferTextureId);
+		}
+
+		void ToHLSL(MIR::FExternPrinterHLSL& Printer) const
+		{
+			Printer << TEXT("MaterialExpressionPathTracingBufferTextureLookup(Parameters, $0, $1)");
+		}
+	};
+
+	Em.Output(0, Em.Extern<FPathTracingBufferTextureExtern>({ PathTracingBufferTextureId }, CoordinatesValue, PathTracingBufferTextureIndexValue));
 }
 
 void UMaterialExpressionPostVolumeUserFlagTest::Build(MIR::FEmitter& Em)
@@ -4481,8 +5183,7 @@ void UMaterialExpressionPostVolumeUserFlagTest::Build(MIR::FEmitter& Em)
 		BitIndexValue = Em.ConstantInt(ConstBitIndex);
 	}
 
-	FValueRef OutputValue = Em.InlineHLSL(MIR::FType::MakeFloatScalar(), TEXTVIEW("PostVolumeUserFlagTest($0)"), { BitIndexValue });
-	Em.Output(0, OutputValue);
+	Em.Output(0, Em.Extern<MIR::FExternSimpleHLSL>({ MIR::EExternSimpleType::Float1, TEXT("PostVolumeUserFlagTest($0)") }, BitIndexValue));
 }
 
 void UMaterialExpressionPreSkinnedLocalBounds::Build(MIR::FEmitter& Em)
@@ -4490,13 +5191,13 @@ void UMaterialExpressionPreSkinnedLocalBounds::Build(MIR::FEmitter& Em)
 	Em.CheckMaterialDomainIsAnyOf( { MD_Surface, MD_Volume });
 	UE_MIR_CHECKPOINT(Em);
 
-	Em.Output(0, Em.InlineHLSL(TEXT("PreSkinnedLocalBoundsHalf"))); // Half extents
-	Em.Output(1, Em.InlineHLSL(TEXT("PreSkinnedLocalBoundsFull"))); // Full extents
-	Em.Output(2, Em.InlineHLSL(TEXT("PreSkinnedLocalBoundsMin")));  // Min point
-	Em.Output(3, Em.InlineHLSL(TEXT("PreSkinnedLocalBoundsMax")));  // Max point
+	Em.Output(0, Em.Extern<MIR::FExternFromMaterialDecl>({ TEXT("PreSkinnedLocalBoundsHalf") })); // Half extents
+	Em.Output(1, Em.Extern<MIR::FExternFromMaterialDecl>({ TEXT("PreSkinnedLocalBoundsFull") })); // Full extents
+	Em.Output(2, Em.Extern<MIR::FExternFromMaterialDecl>({ TEXT("PreSkinnedLocalBoundsMin") }));  // Min point
+	Em.Output(3, Em.Extern<MIR::FExternFromMaterialDecl>({ TEXT("PreSkinnedLocalBoundsMax") }));  // Max point
 }
 
-static void EmitRayOrPathTracingQualitySwitch(MIR::FEmitter& Em, FStringView ConditionHLSL, FValueRef NormalValue, FValueRef TracedValue)
+static void EmitRayOrPathTracingQualitySwitch(MIR::FEmitter& Em, const TCHAR* ConditionHLSL, FValueRef NormalValue, FValueRef TracedValue)
 {
 	if (!NormalValue && !TracedValue)
 	{
@@ -4515,7 +5216,7 @@ static void EmitRayOrPathTracingQualitySwitch(MIR::FEmitter& Em, FStringView Con
 	}
 	else
 	{
-		FValueRef ConditionValue = Em.InlineHLSL(MIR::FType::MakeBoolScalar(), ConditionHLSL);
+		FValueRef ConditionValue = Em.Extern<MIR::FExternSimpleHLSL>({ MIR::EExternSimpleType::Bool1, ConditionHLSL });
 		OutputValue = Em.Branch(ConditionValue, TracedValue, NormalValue);
 	}
 
@@ -4526,14 +5227,14 @@ void UMaterialExpressionRayTracingQualitySwitch::Build(MIR::FEmitter& Em)
 {
 	FValueRef NormalValue = Em.TryInput(&Normal);
 	FValueRef RayTracedValue = FDataDrivenShaderPlatformInfo::GetSupportsRayTracing(Em.GetShaderPlatform()) ? Em.TryInput(&RayTraced) : FValueRef{};
-	EmitRayOrPathTracingQualitySwitch(Em, TEXTVIEW("GetRayTracingQualitySwitch()"), NormalValue, RayTracedValue);
+	EmitRayOrPathTracingQualitySwitch(Em, TEXT("GetRayTracingQualitySwitch()"), NormalValue, RayTracedValue);
 }
 
 void UMaterialExpressionPathTracingQualitySwitch::Build(MIR::FEmitter& Em)
 {
 	FValueRef NormalValue = Em.TryInput(&Normal);
 	FValueRef PathTracedValue = Em.TryInput(&PathTraced);
-	EmitRayOrPathTracingQualitySwitch(Em, TEXTVIEW("GetPathTracingQualitySwitch()"), NormalValue, PathTracedValue);
+	EmitRayOrPathTracingQualitySwitch(Em, TEXT("GetPathTracingQualitySwitch()"), NormalValue, PathTracedValue);
 }
 
 void UMaterialExpressionPathTracingRayTypeSwitch::Build(MIR::FEmitter& Em)
@@ -4550,10 +5251,10 @@ void UMaterialExpressionPathTracingRayTypeSwitch::Build(MIR::FEmitter& Em)
 	// return Compiler->PathTracingRayTypeSwitch(ArgMain, ArgShadow, ArgDiffuse, ArgSpecular, ArgVolume);
 		
 	FValueRef OutputValue;
-	OutputValue = Em.Branch(Em.InlineHLSL(MIR::FType::MakeBoolScalar(), TEXTVIEW("GetPathTracingIsShadow()")), ShadowValue ? ShadowValue : MainValue, MainValue);
-	OutputValue = Em.Branch(Em.InlineHLSL(MIR::FType::MakeBoolScalar(), TEXTVIEW("GetPathTracingIsIndirectDiffuse()")), DiffuseValue ? DiffuseValue : MainValue, OutputValue);
-	OutputValue = Em.Branch(Em.InlineHLSL(MIR::FType::MakeBoolScalar(), TEXTVIEW("GetPathTracingIsIndirectSpecular()")), SpecularValue ? SpecularValue : MainValue, OutputValue);
-	OutputValue = Em.Branch(Em.InlineHLSL(MIR::FType::MakeBoolScalar(), TEXTVIEW("GetPathTracingIsIndirectVolume()")), VolumeValue ? VolumeValue : MainValue, OutputValue);
+	OutputValue = Em.Branch(Em.Extern<MIR::FExternSimpleHLSL>({ MIR::EExternSimpleType::Bool1, TEXT("GetPathTracingIsShadow()") }), ShadowValue ? ShadowValue : MainValue, MainValue);
+	OutputValue = Em.Branch(Em.Extern<MIR::FExternSimpleHLSL>({ MIR::EExternSimpleType::Bool1, TEXT("GetPathTracingIsIndirectDiffuse()") }), DiffuseValue ? DiffuseValue : MainValue, OutputValue);
+	OutputValue = Em.Branch(Em.Extern<MIR::FExternSimpleHLSL>({ MIR::EExternSimpleType::Bool1, TEXT("GetPathTracingIsIndirectSpecular()") }), SpecularValue ? SpecularValue : MainValue, OutputValue);
+	OutputValue = Em.Branch(Em.Extern<MIR::FExternSimpleHLSL>({ MIR::EExternSimpleType::Bool1, TEXT("GetPathTracingIsIndirectVolume()") }), VolumeValue ? VolumeValue : MainValue, OutputValue);
 	
 	Em.Output(0, OutputValue);
 }
@@ -4616,27 +5317,30 @@ static FValueRef EmitProcessPhysicsField(MIR::FEmitter& Em, FExpressionInput* Po
 	// Make a constant integer value from the target index.
 	FValueRef TargetIndexValue = Em.ConstantInt(TargetIndex);
 
+	MIR::EExternSimpleType ResultType = MIR::EExternSimpleType::Float1;
+	const TCHAR* ResultCode;
+
 	// LWC_TODO: LWC aware physics field
 	if (OutputType == Field_Output_Vector)
 	{
-		return bSampleField ? Em.InlineHLSL(MIR::FType::MakeFloatVector(3), TEXTVIEW("MatPhysicsField_SamplePhysicsVectorField($0, $1)"), { PositionValue, TargetIndexValue }) :
-								Em.InlineHLSL(MIR::FType::MakeFloatVector(3), TEXTVIEW("MatPhysicsField_EvalPhysicsVectorField($0, $1)"), { PositionValue, TargetIndexValue });
+		ResultType = MIR::EExternSimpleType::Float3;
+		ResultCode = bSampleField ? TEXT("MatPhysicsField_SamplePhysicsVectorField($0, $1)") : TEXT("MatPhysicsField_EvalPhysicsVectorField($0, $1)");
 	}
 	else if (OutputType == Field_Output_Scalar)
 	{
-		return bSampleField ? Em.InlineHLSL(MIR::FType::MakeFloatScalar(), TEXTVIEW("MatPhysicsField_SamplePhysicsScalarField($0, $1)"), { PositionValue, TargetIndexValue }) :
-								Em.InlineHLSL(MIR::FType::MakeFloatScalar(), TEXTVIEW("MatPhysicsField_EvalPhysicsScalarField($0, $1)"), { PositionValue, TargetIndexValue });
+		ResultCode = bSampleField ? TEXT("MatPhysicsField_SamplePhysicsScalarField($0, $1)") : TEXT("MatPhysicsField_EvalPhysicsScalarField($0, $1)");
 	}
 	else if (OutputType == Field_Output_Integer)
 	{
-		return bSampleField ? Em.InlineHLSL(MIR::FType::MakeFloatScalar(), TEXTVIEW("MatPhysicsField_SamplePhysicsIntegerField($0, $1)"), { PositionValue, TargetIndexValue }) :
-								Em.InlineHLSL(MIR::FType::MakeFloatScalar(), TEXTVIEW("MatPhysicsField_EvalPhysicsIntegerField($0, $1)"), { PositionValue, TargetIndexValue }) ;
+		ResultCode = bSampleField ? TEXT("MatPhysicsField_SamplePhysicsIntegerField($0, $1)") : TEXT("MatPhysicsField_EvalPhysicsIntegerField($0, $1)");
 	}
 	else
 	{
 		checkf(false, TEXT("Unsupported output type."));
 		return Em.Poison();
 	}
+
+	return Em.Extern<MIR::FExternSimpleHLSL>({ ResultType, ResultCode }, PositionValue, TargetIndexValue);
 }
 
 void UMaterialExpressionSamplePhysicsIntegerField::Build(MIR::FEmitter& Em)
@@ -4676,8 +5380,8 @@ void UMaterialExpressionSobol::Build(MIR::FEmitter& Em)
 
 	FValueRef CellValue   = Em.CastToFloat(Em.InputDefaultFloat2(&Cell, { 0.f, 0.f }), 2);   // float2
 	FValueRef IndexValue  = Em.InputDefaultInt(&Index, ConstIndex);                          // int
-	FValueRef SeedValue   = Em.CastToFloat(Em.InputDefaultFloat2(&Seed, {(float)ConstSeed.X, (float)ConstSeed.Y }), 2); // float2
-	FValueRef SobolBits   = Em.InlineHLSL(MIR::FType::MakeIntVector(2), TEXT("SobolIndex(SobolPixel(uint2($0)), uint($1))"), { CellValue, IndexValue });
+	FValueRef SeedValue   = Em.CastToFloat(Em.InputDefaultFloat2(&Seed, { (float)ConstSeed.X, (float)ConstSeed.Y }), 2); // float2
+	FValueRef SobolBits   = Em.Extern<MIR::FExternSimpleHLSL>({ MIR::EExternSimpleType::Int2, TEXT("SobolIndex(SobolPixel(uint2($0)), uint($1))") }, CellValue, IndexValue);
 	FValueRef SeedUint    = Em.CastToInt(Em.Multiply(SeedValue,Em.ConstantFloat(65536.0f)), 2);
 	FValueRef MaskedBits  = Em.BitwiseAnd(Em.BitwiseXor(SobolBits, SeedUint), Em.ConstantInt(0xffff));
 	FValueRef Normalized  = Em.Divide(Em.CastToFloat(MaskedBits, 2), Em.ConstantFloat(65536.0f));
@@ -4706,28 +5410,41 @@ void UMaterialExpressionSpeedTree::Build(MIR::FEmitter& Em)
 		ExtraBendValue = Em.ConstantFloat3({ 0.f, 0.f, 0.f });
 	}
  
-	// Only generate previous frame's computations if required and opted-in
-	FStringView InlineCode = TEXTVIEW("Get<PREV>SpeedTreeVertexOffset(Parameters, $0, $1, $2, $3, $4, $5)");
-	if (bAccurateWindVelocities)
+	struct FSpeedTreeExtern
 	{
-		InlineCode = TEXTVIEW("GetSpeedTreeVertexOffset(Parameters, $0, $1, $2, $3, $4, $5)");
-	}
+		bool bAccurateWindVelocities;
+
+		MIR::FExternInfo GetInfo() const
+		{
+			return
+			{
+				.Name = TEXT("SpeedTree"),
+				.Type = MIR::FType::MakeFloatVector(3),
+				.Flags = MIR::EExternFlags::Inline | MIR::EExternFlags::ZeroDifferentials,
+			};
+		}
+
+		void Analyze(MIR::FExternAnalysisContext& Context)
+		{
+			Context.Module->AddEnvironmentDefine(TEXT("USES_SPEEDTREE"));
+			Context.Module->GetStatistics().NumInterpolatedTexCoords = FMath::Max(Context.Module->GetStatistics().NumInterpolatedTexCoords, 8);
+		}
 
-	FValueRef ResultValue = Em.InlineHLSL(
-		MIR::FType::MakeFloatVector(3),
-		InlineCode,
-		{ GeometryValue, WindValue, LODValue, BillboardThresholdValue, UseExtraBendValue, ExtraBendValue },
-		MIR::EValueFlags::None, MIR::EGraphProperties::None,
-		[](FMaterialIRModule* Module, uintptr_t)
+		void AnalyzeInStage(MIR::FExternAnalysisContext& Context, MIR::EStage Stage)
 		{
-			Module->AddEnvironmentDefine(TEXT("USES_SPEEDTREE"));
-			Module->GetStatistics().NumInterpolatedTexCoords = FMath::Max(Module->GetStatistics().NumInterpolatedTexCoords, 8);
-		});
+			Context.CheckStage(Stage, MIR::EStageMask::Vertex, TEXT("MaterialExpressionSpeedTree"));
+		}
 
-	TStaticArray<FValueRef, MIR::EStage::NumStages> ValuePerStage = {};
-	ValuePerStage[MIR::EStage::Stage_Vertex] = ResultValue;
+		void ToHLSL(MIR::FExternPrinterHLSL& Printer) const
+		{
+			// Only generate previous frame's computations if required and opted-in
+			Printer << (bAccurateWindVelocities
+				? TEXTVIEW("GetSpeedTreeVertexOffset(Parameters, $0, $1, $2, $3, $4, $5)")
+				: TEXTVIEW("Get<PREV>SpeedTreeVertexOffset(Parameters, $0, $1, $2, $3, $4, $5)"));
+		}
+	};
 
-	Em.Output(0, Em.StageSwitch(MIR::FType::MakeFloatVector(3), ValuePerStage));
+	Em.Output(0, Em.Extern<FSpeedTreeExtern>({ bAccurateWindVelocities }, GeometryValue, WindValue, LODValue, BillboardThresholdValue, UseExtraBendValue, ExtraBendValue));
 }
 
 void UMaterialExpressionSRGBColorToWorkingColorSpace::Build(MIR::FEmitter& Em)
@@ -4760,13 +5477,13 @@ void UMaterialExpressionTemporalSobol::Build(MIR::FEmitter& Em)
 
     FValueRef IndexValue = Em.CastToInt(Em.InputDefaultInt(&Index, ConstIndex), 1);
     FValueRef SeedValue  = Em.CastToFloat(Em.InputDefaultFloat2(&Seed, { (float)ConstSeed.X, (float)ConstSeed.Y }), 2);
-    FValueRef FrameMod8  = Em.InlineHLSL(MIR::FType::MakeIntScalar(), TEXT("uint(View.StateFrameIndexMod8)"));
+	FValueRef FrameMod8  = Em.Extern<MIR::FExternSimpleHLSL>({ MIR::EExternSimpleType::Int1, TEXT("int(View.StateFrameIndexMod8)") });
 
     // View.StateFrameIndexMod8 + 8 * Index
     FValueRef FrameIndex  = Em.Add(FrameMod8, Em.Multiply(IndexValue, Em.ConstantInt(8)));
 
     // SobolIndex(SobolPixel(uint2(Parameters.SvPosition.xy)), FrameIndexU)
-    FValueRef Sobol = Em.InlineHLSL(MIR::FType::MakeIntVector(2), TEXT("SobolIndex(SobolPixel(uint2(Parameters.SvPosition.xy)), uint($0))"), { FrameIndex });
+	FValueRef Sobol = Em.Extern<MIR::FExternSimpleHLSL>({ MIR::EExternSimpleType::Int2, TEXT("SobolIndex(SobolPixel(uint2(Parameters.SvPosition.xy)), uint($0))") }, FrameIndex);
 
 	// uint2(SeedValue * 0x10000) & 0xffff)
 	FValueRef Scale      = Em.ConstantFloat(65536.0f); // 0x1000
@@ -4795,7 +5512,7 @@ void UMaterialExpressionNeuralNetworkInput::Build(MIR::FEmitter& Em)
 	}
 	else
 	{
-		FValueRef ViewportUV = Em.InlineHLSL(TEXT("ViewportUV"));
+		FValueRef ViewportUV = Em.Extern<MIR::FExternFromMaterialDecl>({ TEXT("ViewportUV") });
 		float     BatchIndex = bUseTextureAsInput ? -1.0f : 0.0f;
 		CoordinatesValue     = Em.Vector4(Em.ConstantFloat2({ BatchIndex, 0.0f }), ViewportUV);
 	}
@@ -4822,7 +5539,7 @@ void UMaterialExpressionNeuralNetworkOutput::Build(MIR::FEmitter& Em)
 
 	if (!ViewportUV)
 	{
-		ViewportUV = Em.InlineHLSL(TEXT("ViewportUV"));
+		ViewportUV = Em.Extern<MIR::FExternFromMaterialDecl>({ TEXT("ViewportUV") });
 		BufferIndex = Em.Vector(Em.ConstantFloat2({ 0.0f, 0.0f }), ViewportUV);
 	}
 
@@ -4831,13 +5548,32 @@ void UMaterialExpressionNeuralNetworkOutput::Build(MIR::FEmitter& Em)
 	// @massimo.tristano another note on this. do we need this in the new translator?
 	// AddEstimatedTextureSample();
 
-	static const auto EnableWithNeuralNetworks = [] (FMaterialIRModule* Module, uintptr_t)
+	struct FNeuralNetworkOutputExtern
 	{
-		Module->GetCompilationOutput().bUsedWithNeuralNetworks = true;
-	};
+		const TCHAR* Code;
+
+		MIR::FExternInfo GetInfo() const
+		{
+			return {
+				.Name = TEXT("NeuralNetworkOutputExtern"),
+				.Type = MIR::FType::MakeFloatVector(4),
+				.Flags = MIR::EExternFlags::Inline | MIR::EExternFlags::ZeroDifferentials,
+			};
+		}
+
+		void Analyze(MIR::FExternAnalysisContext& Context)
+		{
+			Context.Module->GetCompilationOutput().bUsedWithNeuralNetworks = true;
+		}
 
-	Em.Output(0, Em.InlineHLSL(MIR::FType::MakeFloatVector(4), TEXTVIEW("NeuralTextureOutput(Parameters,$0)"), { ViewportUV }, {}, {}, EnableWithNeuralNetworks));
-	Em.Output(1, Em.InlineHLSL(MIR::FType::MakeFloatVector(4), TEXTVIEW("NeuralBufferOutput(Parameters,$0)"), { BufferIndex }, {}, {}, EnableWithNeuralNetworks));
+		void ToHLSL(MIR::FExternPrinterHLSL& Printer) const
+		{
+			Printer << Code;
+		}
+	};
+	
+	Em.Output(0, Em.Extern<FNeuralNetworkOutputExtern>({ TEXT("NeuralTextureOutput(Parameters, $0)") }, ViewportUV));
+	Em.Output(1, Em.Extern<FNeuralNetworkOutputExtern>({ TEXT("NeuralBufferOutput(Parameters, $0)") }, BufferIndex));
 }
 
 void UMaterialExpressionVertexInterpolator::Build(MIR::FEmitter& Em)
@@ -5024,6 +5760,7 @@ enum EDefaultDirection
 	EDefaultDirectionNormal,
 	EDefaultDirectionTangent
 };
+
 static FValueRef CompileWithDefaultDirectionWS(MIR::FEmitter& Em, UMaterial* Material, FExpressionInput& DirectionInput, bool bConvertToRequestedSpace, EDefaultDirection DefaultDirection)
 {
 	FSubstrateTranslatorData* SubstrateTranslatorData = Em.GetSubstrateTranslatorData();
@@ -5081,6 +5818,31 @@ static FValueRef CompileWithDefaultTangentWS(MIR::FEmitter& Em, UMaterial* Mater
 	return CompileWithDefaultDirectionWS(Em, Material, TangentInput, bConvertToRequestedSpace, EDefaultDirection::EDefaultDirectionTangent);
 }
 
+FValueRef EmitSubstrateSharedLocalBasis(MIR::FEmitter& Em, const FSubstrateRegisteredSharedLocalBasis& SharedLocalBasis, ESubstrateCompilationContext Context)
+{
+	struct FSubstrateSharedLocalBasisExtern
+	{
+		uint8 GraphSharedLocalBasisIndex;
+		ESubstrateCompilationContext Context;
+
+		MIR::FExternInfo GetInfo() const
+		{
+			return
+				{
+					.Name = TEXT("SubstrateSharedLocalBasis"),
+					.Type = MIR::FType::MakeIntScalar(),
+					.Flags = MIR::EExternFlags::Inline | MIR::EExternFlags::NoDifferentials,
+				};
+		}
+
+		void ToHLSL(MIR::FExternPrinterHLSL& Printer) const
+		{
+			Printer << TEXT("SHAREDLOCALBASIS_INDEX_") << (uint32)GraphSharedLocalBasisIndex << TEXT('_') << (uint32)Context;
+		}
+	};
+	return Em.Extern<FSubstrateSharedLocalBasisExtern>({ SharedLocalBasis.GraphSharedLocalBasisIndex, Context });
+}
+
 static FValueRef FinalizeBSDFAndPromoteToOperatorIfNeeded(MIR::FEmitter& Em, FSubstrateOperator& SubstrateOperator, FValueRef BSDF)
 {
 	FSubstrateOperator* PromoteToOperator = !SubstrateOperator.bUseParameterBlending || (SubstrateOperator.bUseParameterBlending && SubstrateOperator.bRootOfParameterBlendingSubTree) ? &SubstrateOperator : nullptr;
@@ -5272,7 +6034,7 @@ void UMaterialExpressionSubstrateSlabBSDF::Build(MIR::FEmitter& Em)
 	const bool bIsThinSurface = Material ? Material->IsThinSurface() : false;
 	float ClearCoatUseSecondNormal = 0.0f; // Constant(ClearCoatBottomNormal != Normal ? 1.0f : 0.0f);SUBSTRATE_TODO
 
-	FValueRef SharedLocalBasisIndexMacroArg = Em.InlineHLSL(MIR::FType::MakeIntScalar(), *GetSubstrateSharedLocalBasisIndexMacroInner(NewRegisteredSharedLocalBasis, SubstrateTranslatorData->CurrentSubstrateCompilationContext));
+	FValueRef SharedLocalBasisIndexMacroArg = EmitSubstrateSharedLocalBasis(Em, NewRegisteredSharedLocalBasis, SubstrateTranslatorData->CurrentSubstrateCompilationContext);
 
 	MIR::FSubstrateSlabDesc SubstrateSlabDesc = {
 			NormalArg,
@@ -5346,8 +6108,7 @@ MIR::FValueRef UMaterialExpressionSubstrateShadingModels::BuildCommon(MIR::FEmit
 
 	const FSubstrateRegisteredSharedLocalBasis NewRegisteredSharedLocalBasis = MIRSubstrateCompilationInfoCreateSharedLocalBasis2(SubstrateTranslatorData, &NormalArg, bHasValidTangent ? &TangentArg : nullptr);
 	SubstrateOperator.BSDFRegisteredSharedLocalBasis = NewRegisteredSharedLocalBasis;
-	const FString BasisIndexMacro = GetSubstrateSharedLocalBasisIndexMacroInner(NewRegisteredSharedLocalBasis, SubstrateTranslatorData->CurrentSubstrateCompilationContext);
-
+	FValueRef BasisIndexMacroArg = EmitSubstrateSharedLocalBasis(Em, NewRegisteredSharedLocalBasis, SubstrateTranslatorData->CurrentSubstrateCompilationContext);
 
 	FValueRef EmissiveColorArg = CompileWithDefaultFloat3(Em, EmissiveColor, 0.0f, 0.0f, 0.0f, EditorOnlyData ? &EditorOnlyData->EmissiveColor : nullptr);
 	FValueRef BaseColorArg = CompileWithDefaultFloat3(Em, BaseColor, 0.18f, 0.18f, 0.18f, EditorOnlyData ? &EditorOnlyData->BaseColor : nullptr);
@@ -5371,14 +6132,14 @@ MIR::FValueRef UMaterialExpressionSubstrateShadingModels::BuildCommon(MIR::FEmit
 	// Clear coat normal basis
 	FValueRef ClearCoatNormalArg = NormalArg;
 	FValueRef ClearCoatTangentArg = TangentArg;
-	FString ClearCoat_BasisIndexMacro = BasisIndexMacro;
+	FValueRef ClearCoat_BasisIndexMacroArg = BasisIndexMacroArg;
 	FSubstrateRegisteredSharedLocalBasis ClearCoat_NewRegisteredSharedLocalBasis = NewRegisteredSharedLocalBasis;
 	if (bHasCoatNormal)
 	{
 		ClearCoatNormalArg = CompileWithDefaultNormalWS(Em, Material, ClearCoatNormal);
 		ClearCoat_NewRegisteredSharedLocalBasis = MIRSubstrateCompilationInfoCreateSharedLocalBasis2(SubstrateTranslatorData, &ClearCoatNormalArg, bHasValidTangent ? &ClearCoatTangentArg : nullptr);
 		//SubstrateCompilationInfoCreateSharedLocalBasis(Compiler, ClearCoat_NormalCodeChunk, ClearCoat_TangentCodeChunk);
-		ClearCoat_BasisIndexMacro = GetSubstrateSharedLocalBasisIndexMacroInner(ClearCoat_NewRegisteredSharedLocalBasis, SubstrateTranslatorData->CurrentSubstrateCompilationContext);
+		ClearCoat_BasisIndexMacroArg = EmitSubstrateSharedLocalBasis(Em, ClearCoat_NewRegisteredSharedLocalBasis, SubstrateTranslatorData->CurrentSubstrateCompilationContext);
 	}
 
 	FValueRef OpacityArg = FValueRef();
@@ -5426,9 +6187,6 @@ MIR::FValueRef UMaterialExpressionSubstrateShadingModels::BuildCommon(MIR::FEmit
 		}
 	}
 
-	FValueRef BasisIndexMacroArg = Em.InlineHLSL(MIR::FType::MakeIntScalar(), *BasisIndexMacro);
-	FValueRef ClearCoat_BasisIndexMacroArg = Em.InlineHLSL(MIR::FType::MakeIntScalar(), *ClearCoat_BasisIndexMacro);
-
 	MIR::FSubstrateShadingModelsDesc SubstrateShadingModelsDesc = {
 		BaseColorArg,
 		SpecularArg,
@@ -5521,7 +6279,7 @@ void UMaterialExpressionSubstrateSimpleClearCoatBSDF::Build(MIR::FEmitter& Em)
 	FValueRef NormalArg = CompileWithDefaultNormalWS(Em, Material, Normal);
 	const FSubstrateRegisteredSharedLocalBasis NewRegisteredSharedLocalBasis = MIRSubstrateCompilationInfoCreateSharedLocalBasis2(SubstrateTranslatorData, &NormalArg, nullptr);
 	SubstrateOperator.BSDFRegisteredSharedLocalBasis = NewRegisteredSharedLocalBasis;
-	FValueRef SharedLocalBasisIndexMacroArg = Em.InlineHLSL(MIR::FType::MakeIntScalar(), *GetSubstrateSharedLocalBasisIndexMacroInner(NewRegisteredSharedLocalBasis, SubstrateTranslatorData->CurrentSubstrateCompilationContext));
+	FValueRef SharedLocalBasisIndexMacroArg = EmitSubstrateSharedLocalBasis(Em, NewRegisteredSharedLocalBasis, SubstrateTranslatorData->CurrentSubstrateCompilationContext);
 
 	MIR::FSubstrateSlabDesc SubstrateSlabDesc = {
 			NormalArg,
@@ -5577,22 +6335,22 @@ void UMaterialExpressionSubstrateEyeBSDF::Build(MIR::FEmitter& Em)
 	FValueRef NormalArg = CompileWithDefaultNormalWS(Em, Material, CorneaNormal);
 	const FSubstrateRegisteredSharedLocalBasis NewRegisteredSharedLocalBasis = MIRSubstrateCompilationInfoCreateSharedLocalBasis2(SubstrateTranslatorData, &NormalArg, nullptr);
 	SubstrateOperator.BSDFRegisteredSharedLocalBasis = NewRegisteredSharedLocalBasis;
-	FValueRef SharedLocalBasisIndexMacroArg = Em.InlineHLSL(MIR::FType::MakeIntScalar(), *GetSubstrateSharedLocalBasisIndexMacroInner(NewRegisteredSharedLocalBasis, SubstrateTranslatorData->CurrentSubstrateCompilationContext));
+	FValueRef SharedLocalBasisIndexMacroArg = EmitSubstrateSharedLocalBasis(Em, NewRegisteredSharedLocalBasis, SubstrateTranslatorData->CurrentSubstrateCompilationContext);
 
-	FValueRef EyeBSDF = Em.InlineHLSL(
-		MIR::FType::MakeSubstrateData(),
-		TEXTVIEW("GetSubstrateEyeBSDF($0, $1, $2, $3, $4, $5, ExtractSubsurfaceProfileInt($6), $7, $8)"),
+	FValueRef EyeBSDF = Em.Extern<MIR::FExternSimpleHLSL>(
 		{
-			DiffuseColorArg,
-			RoughnessArg,
-			IrisMaskArg,
-			IrisDistanceArg,
-			IrisNormalArg,
-			IrisPlaneNormalArg,
-			SubsurfaceProfileArg,
-			EmissiveColorArg,
-			SharedLocalBasisIndexMacroArg
-		});
+			MIR::EExternSimpleType::SubstrateData,
+			TEXT("GetSubstrateEyeBSDF($0, $1, $2, $3, $4, $5, ExtractSubsurfaceProfileInt($6), $7, $8)")
+		},
+		DiffuseColorArg,
+		RoughnessArg,
+		IrisMaskArg,
+		IrisDistanceArg,
+		IrisNormalArg,
+		IrisPlaneNormalArg,
+		SubsurfaceProfileArg,
+		EmissiveColorArg,
+		SharedLocalBasisIndexMacroArg);
 
 	FinalizeBSDFAndPromoteToOperatorIfNeeded(Em, SubstrateOperator, EyeBSDF);
 }
@@ -5614,20 +6372,20 @@ void UMaterialExpressionSubstrateHairBSDF::Build(MIR::FEmitter& Em)
 	FValueRef TangentArg		= CompileWithDefaultTangentWS(Em, Material, Tangent);
 	const FSubstrateRegisteredSharedLocalBasis NewRegisteredSharedLocalBasis = MIRSubstrateCompilationInfoCreateSharedLocalBasis2(SubstrateTranslatorData, &TangentArg, nullptr);
 	SubstrateOperator.BSDFRegisteredSharedLocalBasis = NewRegisteredSharedLocalBasis;
-	FValueRef SharedLocalBasisIndexMacroArg = Em.InlineHLSL(MIR::FType::MakeIntScalar(), *GetSubstrateSharedLocalBasisIndexMacroInner(NewRegisteredSharedLocalBasis, SubstrateTranslatorData->CurrentSubstrateCompilationContext));
+	FValueRef SharedLocalBasisIndexMacroArg = EmitSubstrateSharedLocalBasis(Em, NewRegisteredSharedLocalBasis, SubstrateTranslatorData->CurrentSubstrateCompilationContext);
 
-	FValueRef HairBSDF = Em.InlineHLSL(
-		MIR::FType::MakeSubstrateData(),
-		TEXTVIEW("GetSubstrateHairBSDF($0, $1, $2, $3, $4, $5, false/*bComplexTransmittance*/, $6)"),
+	FValueRef HairBSDF = Em.Extern<MIR::FExternSimpleHLSL>(
 		{
-			BaseColorArg,
-			ScatterArg,
-			SpecularArg,
-			RoughnessArg,
-			BacklitArg,
-			EmissiveColorArg,
-			SharedLocalBasisIndexMacroArg
-		});
+			MIR::EExternSimpleType::SubstrateData,
+			TEXT("GetSubstrateHairBSDF($0, $1, $2, $3, $4, $5, false/*bComplexTransmittance*/, $6)")
+		},
+		BaseColorArg,
+		ScatterArg,
+		SpecularArg,
+		RoughnessArg,
+		BacklitArg,
+		EmissiveColorArg,
+		SharedLocalBasisIndexMacroArg);
 
 	FinalizeBSDFAndPromoteToOperatorIfNeeded(Em, SubstrateOperator, HairBSDF);
 }
@@ -5670,15 +6428,15 @@ void UMaterialExpressionSubstrateVolumetricFogCloudBSDF::Build(MIR::FEmitter& Em
 		AmbientOcclusionArg = Em.ConstantFloat(1.0f);
 	}
 
-	FValueRef VolumetricFogCloudBSDF = Em.InlineHLSL(
-		MIR::FType::MakeSubstrateData(),
-		TEXTVIEW("GetSubstrateVolumeFogCloudBSDF($0, $1, $2, $3)"),
+	FValueRef VolumetricFogCloudBSDF = Em.Extern<MIR::FExternSimpleHLSL>(
 		{
-			AlbedoArg,
-			ExtinctionArg,
-			EmissiveColorArg,
-			AmbientOcclusionArg
-		});
+			MIR::EExternSimpleType::SubstrateData,
+			TEXT("GetSubstrateVolumeFogCloudBSDF($0, $1, $2, $3)"),
+		},
+		AlbedoArg,
+		ExtinctionArg,
+		EmissiveColorArg,
+		AmbientOcclusionArg);
 
 	FinalizeBSDFAndPromoteToOperatorIfNeeded(Em, SubstrateOperator, VolumetricFogCloudBSDF);
 }
@@ -5702,7 +6460,7 @@ void UMaterialExpressionSubstrateSingleLayerWaterBSDF::Build(MIR::FEmitter& Em)
 	FValueRef NormalArg = CompileWithDefaultNormalWS(Em, Material, Normal);
 	const FSubstrateRegisteredSharedLocalBasis NewRegisteredSharedLocalBasis = MIRSubstrateCompilationInfoCreateSharedLocalBasis2(SubstrateTranslatorData, &NormalArg, nullptr);
 	SubstrateOperator.BSDFRegisteredSharedLocalBasis = NewRegisteredSharedLocalBasis;
-	FValueRef SharedLocalBasisIndexMacroArg = Em.InlineHLSL(MIR::FType::MakeIntScalar(), *GetSubstrateSharedLocalBasisIndexMacroInner(NewRegisteredSharedLocalBasis, SubstrateTranslatorData->CurrentSubstrateCompilationContext));
+	FValueRef SharedLocalBasisIndexMacroArg = EmitSubstrateSharedLocalBasis(Em, NewRegisteredSharedLocalBasis, SubstrateTranslatorData->CurrentSubstrateCompilationContext);
 
 	FValueRef BaseColorArg				= CompileWithDefaultFloat3(Em, BaseColor,				0.0f, 0.0f, 0.0f);
 	FValueRef MetallicArg				= CompileWithDefaultFloat1(Em, Metallic,				0.0f);
@@ -5715,22 +6473,22 @@ void UMaterialExpressionSubstrateSingleLayerWaterBSDF::Build(MIR::FEmitter& Em)
 	FValueRef WaterPhaseGArg			= CompileWithDefaultFloat1(Em, WaterPhaseG,				0.0f);
 	FValueRef ColorScaleBehindWaterArg	= CompileWithDefaultFloat3(Em, ColorScaleBehindWater,	1.0f, 1.0f, 1.0f);
 
-	FValueRef SingleLayerWaterBSDF = Em.InlineHLSL(
-		MIR::FType::MakeSubstrateData(),
-		TEXTVIEW("GetSubstrateSingleLayerWaterBSDF($0, $1, $2, $3, $4, $5, $6, $7, $8, $9, $10)"),
+	FValueRef SingleLayerWaterBSDF = Em.Extern<MIR::FExternSimpleHLSL>(
 		{
-			BaseColorArg,
-			MetallicArg,
-			SpecularArg,
-			RoughnessArg,
-			EmissiveColorArg,
-			TopMaterialOpacityArg,
-			WaterAlbedoArg,
-			WaterExtinctionArg,
-			WaterPhaseGArg,
-			ColorScaleBehindWaterArg,
-			SharedLocalBasisIndexMacroArg
-		});
+			MIR::EExternSimpleType::SubstrateData,
+			TEXT("GetSubstrateSingleLayerWaterBSDF($0, $1, $2, $3, $4, $5, $6, $7, $8, $9, $10)"),
+		},
+		BaseColorArg,
+		MetallicArg,
+		SpecularArg,
+		RoughnessArg,
+		EmissiveColorArg,
+		TopMaterialOpacityArg,
+		WaterAlbedoArg,
+		WaterExtinctionArg,
+		WaterPhaseGArg,
+		ColorScaleBehindWaterArg,
+		SharedLocalBasisIndexMacroArg);
 
 	FinalizeBSDFAndPromoteToOperatorIfNeeded(Em, SubstrateOperator, SingleLayerWaterBSDF);
 }
@@ -5774,14 +6532,14 @@ void UMaterialExpressionSubstrateHorizontalMixing::Build(MIR::FEmitter& Em)
 			return Em.Errorf(TEXT("Background input graphs could not be evaluated for parameter blending."));
 		}
 
-		FValueRef NormalMixArg = Em.InlineHLSL(MIR::FType::MakeFloatScalar(), TEXTVIEW("HorizontalMixingParameterBlendingBSDFCoverageToNormalMix($0, $1, $2)"),
+		FValueRef NormalMixArg = Em.Extern<MIR::FExternSimpleHLSL>(
 			{
-				BackgroundArg,
-				ForegroundArg,
-				HorizontalMixArg
-			}
-		);
-
+				MIR::EExternSimpleType::Float1,
+				TEXT("HorizontalMixingParameterBlendingBSDFCoverageToNormalMix($0, $1, $2)")
+			},
+			BackgroundArg,
+			ForegroundArg,
+			HorizontalMixArg);
 
 		FSubstrateOperator* BackgroundBSDFOperator = SubstrateTranslatorData->SubstrateCompilationGetOperatorFromIndex(SubstrateOperator.LeftIndex);
 		FSubstrateOperator* ForegroundBSDFOperator = SubstrateTranslatorData->SubstrateCompilationGetOperatorFromIndex(SubstrateOperator.RightIndex);
@@ -5818,16 +6576,14 @@ void UMaterialExpressionSubstrateHorizontalMixing::Build(MIR::FEmitter& Em)
 		}
 		const FSubstrateRegisteredSharedLocalBasis NewRegisteredSharedLocalBasis = MIRSubstrateCompilationInfoCreateSharedLocalBasis2(SubstrateTranslatorData, &NewNormalArg, NewTangentArg.IsValid() ? &NewTangentArg : nullptr);
 
-
-		FString SharedLocalBasisName = GetSubstrateSharedLocalBasisIndexMacroInner(NewRegisteredSharedLocalBasis, SubstrateTranslatorData->CurrentSubstrateCompilationContext);
-		FValueRef SharedLocalBasisIndexMacroArg = Em.InlineHLSL(MIR::FType::MakeIntScalar(), *SharedLocalBasisName);
+		FValueRef SharedLocalBasisIndexMacroArg = EmitSubstrateSharedLocalBasis(Em, NewRegisteredSharedLocalBasis, SubstrateTranslatorData->CurrentSubstrateCompilationContext);
 
 		if (!ForegroundArg.IsValid() || !BackgroundArg.IsValid() || !HorizontalMixArg.IsValid() || !NormalMixArg.IsValid())
 		{
 			Em.Errorf(TEXT("Invalid Horizontal operator Foreground, Background, Mix or NormalMix input."));
 		}
 
-		FValueRef CameraVectorArg = Em.InlineHLSL(NAME_CameraVector);
+		FValueRef CameraVectorArg = Em.Extern<MIR::FExternFromMaterialDecl>(NAME_CameraVector);
 		FValueRef BackgroundNoV = Em.Saturate(Em.Dot(BackgroundBasis.BasisInfo.NormalArg, CameraVectorArg));
 		FValueRef ForegroundNoV = Em.Saturate(Em.Dot(ForegroundBasis.BasisInfo.NormalArg, CameraVectorArg));
 
@@ -5923,8 +6679,7 @@ void UMaterialExpressionSubstrateVerticalLayering::Build(MIR::FEmitter& Em)
 			return Em.Errorf(TEXT("Background input graphs could not be evaluated for parameter blending."));
 		}
 
-		FValueRef NormalMixArg = Em.InlineHLSL(MIR::FType::MakeFloatScalar(), TEXTVIEW("VerticalLayeringParameterBlendingBSDFCoverageToNormalMix($0)"), { TopArg } );
-
+		FValueRef NormalMixArg = Em.Extern<MIR::FExternSimpleHLSL>({ MIR::EExternSimpleType::Float1, TEXT("VerticalLayeringParameterBlendingBSDFCoverageToNormalMix($0)") }, TopArg);
 
 		FSubstrateOperator* TopBSDFOperator = SubstrateTranslatorData->SubstrateCompilationGetOperatorFromIndex(SubstrateOperator.LeftIndex);
 		FSubstrateOperator* BaseBSDFOperator = SubstrateTranslatorData->SubstrateCompilationGetOperatorFromIndex(SubstrateOperator.RightIndex);
@@ -5961,15 +6716,14 @@ void UMaterialExpressionSubstrateVerticalLayering::Build(MIR::FEmitter& Em)
 		}
 		const FSubstrateRegisteredSharedLocalBasis NewRegisteredSharedLocalBasis = MIRSubstrateCompilationInfoCreateSharedLocalBasis2(SubstrateTranslatorData, &NewNormalArg, NewTangentArg.IsValid() ? &NewTangentArg : nullptr);
 
-		FString SharedLocalBasisName = GetSubstrateSharedLocalBasisIndexMacroInner(NewRegisteredSharedLocalBasis, SubstrateTranslatorData->CurrentSubstrateCompilationContext);
-		FValueRef SharedLocalBasisIndexMacroArg = Em.InlineHLSL(MIR::FType::MakeIntScalar(), *SharedLocalBasisName);
+		FValueRef SharedLocalBasisIndexMacroArg = EmitSubstrateSharedLocalBasis(Em, NewRegisteredSharedLocalBasis, SubstrateTranslatorData->CurrentSubstrateCompilationContext);
 
 		if (!TopArg.IsValid() || !BaseArg.IsValid() || !ThicknessArg.IsValid() || !NormalMixArg.IsValid())
 		{
 			Em.Errorf(TEXT("Invalid Vertical operator Base, Top, Thickness or NormalMix input."));
 		}
 
-		FValueRef CameraVectorArg = Em.InlineHLSL(NAME_CameraVector);
+		FValueRef CameraVectorArg = Em.Extern<MIR::FExternFromMaterialDecl>(NAME_CameraVector);
 		FValueRef TopNoV = Em.Saturate(Em.Dot(TopBasis.BasisInfo.NormalArg, CameraVectorArg));
 		FValueRef BaseNoV = Em.Saturate(Em.Dot(BaseBasis.BasisInfo.NormalArg, CameraVectorArg));
 
@@ -6127,12 +6881,9 @@ void UMaterialExpressionSubstrateAdd::Build(MIR::FEmitter& Em)
 			return Em.Errorf(TEXT("B input graphs could not be evaluated for parameter blending."));
 		}
 
-		FValueRef NormalMixArg = Em.InlineHLSL(MIR::FType::MakeFloatScalar(), TEXTVIEW("AddParameterBlendingBSDFCoverageToNormalMix($0, $1)"),
-			{
-				BArg,
-				AArg
-			}
-		);
+		FValueRef NormalMixArg = Em.Extern<MIR::FExternSimpleHLSL>({ MIR::EExternSimpleType::Float1, TEXT("AddParameterBlendingBSDFCoverageToNormalMix($0, $1)") },
+			BArg,
+			AArg);
 
 		FSubstrateOperator* BBSDFOperator = SubstrateTranslatorData->SubstrateCompilationGetOperatorFromIndex(SubstrateOperator.LeftIndex);
 		FSubstrateOperator* ABSDFOperator = SubstrateTranslatorData->SubstrateCompilationGetOperatorFromIndex(SubstrateOperator.RightIndex);
@@ -6169,15 +6920,14 @@ void UMaterialExpressionSubstrateAdd::Build(MIR::FEmitter& Em)
 		}
 		const FSubstrateRegisteredSharedLocalBasis NewRegisteredSharedLocalBasis = MIRSubstrateCompilationInfoCreateSharedLocalBasis2(SubstrateTranslatorData, &NewNormalArg, NewTangentArg.IsValid() ? &NewTangentArg : nullptr);
 
-		FString SharedLocalBasisName = GetSubstrateSharedLocalBasisIndexMacroInner(NewRegisteredSharedLocalBasis, SubstrateTranslatorData->CurrentSubstrateCompilationContext);
-		FValueRef SharedLocalBasisIndexMacroArg = Em.InlineHLSL(MIR::FType::MakeIntScalar(), *SharedLocalBasisName);
+		FValueRef SharedLocalBasisIndexMacroArg = EmitSubstrateSharedLocalBasis(Em, NewRegisteredSharedLocalBasis, SubstrateTranslatorData->CurrentSubstrateCompilationContext);
 
 		if (!AArg.IsValid() || !BArg.IsValid() || !NormalMixArg.IsValid())
 		{
 			Em.Errorf(TEXT("Invalid Add operator A, B, Mix or NormalMix input."));
 		}
 
-		FValueRef CameraVectorArg = Em.InlineHLSL(NAME_CameraVector);
+		FValueRef CameraVectorArg = Em.Extern<MIR::FExternFromMaterialDecl>(NAME_CameraVector);
 		FValueRef BNoV = Em.Saturate(Em.Dot(BBasis.BasisInfo.NormalArg, CameraVectorArg));
 		FValueRef ANoV = Em.Saturate(Em.Dot(ABasis.BasisInfo.NormalArg, CameraVectorArg));
 
@@ -6269,10 +7019,7 @@ void UMaterialExpressionSubstrateSelect::Build(MIR::FEmitter& Em)
 		FValueRef SelectValueArg = CompileWithDefaultFloat1(Em, SelectValue, 0.0f);
 
 		FValueRef ThresholdArg = Em.ConstantFloat(Threshold);
-		FValueRef FinalSelectValueArg = Em.InlineHLSL(
-			MIR::FType::MakeScalar(MIR::EScalarKind::Float),
-			TEXTVIEW("($0 <= $1 ? 0.0f : 1.0f)"),
-			{ SelectValueArg, ThresholdArg });
+		FValueRef FinalSelectValueArg = Em.Select(Em.LessThanOrEquals(SelectValueArg, ThresholdArg), Em.ConstantZero(MIR::EScalarKind::Float), Em.ConstantOne(MIR::EScalarKind::Float));
 
 		if (!AArg.IsValid() || !BArg.IsValid() || !ThresholdArg.IsValid())
 		{
@@ -6281,10 +7028,7 @@ void UMaterialExpressionSubstrateSelect::Build(MIR::FEmitter& Em)
 
 		auto SubstrateSelectNormal = [&](FValueRef NormalA, FValueRef NormalB)
 			{
-				return Em.InlineHLSL(
-					MIR::FType::MakeFloatVector(3),
-					TEXTVIEW("($0 <= $1 ? $2 : $3)"),
-					{ SelectValueArg, ThresholdArg, NormalA, NormalB });
+				return Em.Select(Em.LessThanOrEquals(SelectValueArg, ThresholdArg), NormalA, NormalB);
 			};
 
 		FSubstrateOperator* BBSDFOperator = SubstrateTranslatorData->SubstrateCompilationGetOperatorFromIndex(SubstrateOperator.LeftIndex);
@@ -6322,8 +7066,7 @@ void UMaterialExpressionSubstrateSelect::Build(MIR::FEmitter& Em)
 		}
 		const FSubstrateRegisteredSharedLocalBasis NewRegisteredSharedLocalBasis = MIRSubstrateCompilationInfoCreateSharedLocalBasis2(SubstrateTranslatorData, &NewNormalArg, NewTangentArg.IsValid() ? &NewTangentArg : nullptr);
 
-		FString SharedLocalBasisName = GetSubstrateSharedLocalBasisIndexMacroInner(NewRegisteredSharedLocalBasis, SubstrateTranslatorData->CurrentSubstrateCompilationContext);
-		FValueRef SharedLocalBasisIndexMacroArg = Em.InlineHLSL(MIR::FType::MakeIntScalar(), *SharedLocalBasisName);
+		FValueRef SharedLocalBasisIndexMacroArg = EmitSubstrateSharedLocalBasis(Em, NewRegisteredSharedLocalBasis, SubstrateTranslatorData->CurrentSubstrateCompilationContext);
 
 		const FSubstrateOperator* PromoteToOperator = SubstrateOperator.bRootOfParameterBlendingSubTree ? &SubstrateOperator : nullptr;
 		if (PromoteToOperator)
@@ -6384,10 +7127,11 @@ void UMaterialExpressionSubstrateTransmittanceToMFP::Build(MIR::FEmitter& Em)
 		return;
 	}
 
-	MIR::FValueRef MFPArg = Em.InlineHLSL(
-		MIR::FType::MakeFloatVector(3), 
-		TEXTVIEW("(TransmittanceToMeanFreePath($0, $1 * CENTIMETER_TO_METER) * METER_TO_CENTIMETER)"),
-		{ TransmittanceColorArg, ThicknessArg} );
+	MIR::FValueRef MFPArg = Em.Extern<MIR::FExternSimpleHLSL>(
+		{
+			MIR::EExternSimpleType::Float3,
+			TEXT("(TransmittanceToMeanFreePath($0, $1 * CENTIMETER_TO_METER) * METER_TO_CENTIMETER)")
+		}, TransmittanceColorArg, ThicknessArg);
 
 	Em.Output(0, MFPArg);
 	Em.Output(1, ThicknessArg);
@@ -6399,15 +7143,8 @@ void UMaterialExpressionSubstrateMetalnessToDiffuseAlbedoF0::Build(MIR::FEmitter
 	FValueRef SpecularArg	= Specular.GetTracedInput().Expression	? Em.TryInput(&Specular)	: Em.ConstantFloat(0.5f);
 	FValueRef MetallicArg	= Metallic.GetTracedInput().Expression	? Em.TryInput(&Metallic)	: Em.ConstantFloat(0.0f);
 
-	MIR::FValueRef DiffuseAlbedoArg = Em.InlineHLSL(
-		MIR::FType::MakeFloatVector(3),
-		TEXTVIEW("ComputeDiffuseAlbedo($0, saturate($1))"),
-		{ BaseColorArg, MetallicArg });
-
-	MIR::FValueRef F0Arg = Em.InlineHLSL(
-		MIR::FType::MakeFloatVector(3),
-		TEXTVIEW("ComputeF0($0, $1, saturate($2))"),
-		{ SpecularArg, BaseColorArg, MetallicArg });
+	MIR::FValueRef DiffuseAlbedoArg = Em.Extern<MIR::FExternSimpleHLSL>({ MIR::EExternSimpleType::Float3, TEXT("ComputeDiffuseAlbedo($0, saturate($1))") }, BaseColorArg, MetallicArg);
+	MIR::FValueRef F0Arg = Em.Extern<MIR::FExternSimpleHLSL>({ MIR::EExternSimpleType::Float3, TEXT("ComputeF0($0, $1, saturate($2))") }, SpecularArg, BaseColorArg, MetallicArg);
 
 	Em.Output(0, DiffuseAlbedoArg);
 	Em.Output(1, F0Arg);
@@ -6418,15 +7155,8 @@ void UMaterialExpressionSubstrateHazinessToSecondaryRoughness::Build(MIR::FEmitt
 	FValueRef BaseRoughnessArg	= BaseRoughness.GetTracedInput().Expression ?	Em.TryInput(&BaseRoughness)	: Em.ConstantFloat(0.1f);
 	FValueRef HazinessArg		= Haziness.GetTracedInput().Expression ?		Em.TryInput(&Haziness)		: Em.ConstantFloat(0.5f);
 
-	MIR::FValueRef HazeRoughnessArg = Em.InlineHLSL(
-		MIR::FType::MakeScalar(MIR::EScalarKind::Float),
-		TEXTVIEW("SubstrateComputeHazeRoughness(saturate($0))"),
-		{ BaseRoughnessArg });
-
-	MIR::FValueRef HazeWeightArg = Em.InlineHLSL(
-		MIR::FType::MakeScalar(MIR::EScalarKind::Float),
-		TEXTVIEW("SubstrateComputeHazeWeight(saturate($0), saturate($1))"),
-		{ BaseRoughnessArg, HazinessArg });
+	MIR::FValueRef HazeRoughnessArg = Em.Extern<MIR::FExternSimpleHLSL>({ MIR::EExternSimpleType::Float1, TEXT("SubstrateComputeHazeRoughness(saturate($0))") }, BaseRoughnessArg);
+	MIR::FValueRef HazeWeightArg = Em.Extern<MIR::FExternSimpleHLSL>({ MIR::EExternSimpleType::Float1, TEXT("SubstrateComputeHazeWeight(saturate($0), saturate($1))") }, BaseRoughnessArg, HazinessArg);
 
 	Em.Output(0, HazeRoughnessArg);
 	Em.Output(1, HazeWeightArg);
@@ -6443,30 +7173,22 @@ void UMaterialExpressionSubstrateThinFilm::Build(MIR::FEmitter& Em)
 	FValueRef ThicknessArg	= Thickness.GetTracedInput().Expression ?	Em.TryInput(&Thickness)	: Em.ConstantFloat(1.0f);
 	FValueRef IORArg		= IOR.GetTracedInput().Expression ?			Em.TryInput(&IOR)		: Em.ConstantFloat(1.44f);
 
-	FValueRef CameraVector = Em.InlineHLSL(NAME_CameraVector);
+	FValueRef CameraVector = Em.Extern<MIR::FExternFromMaterialDecl>(NAME_CameraVector);
 
-	MIR::FValueRef OutF0Arg = Em.InlineHLSL(
-		MIR::FType::MakeFloatVector(3),
-		TEXTVIEW("SubstrateGetThinFilmF0F90($0, $1, $2, $3, $4).F0"),
-		{ 
-			Em.Dot(NormalArg, CameraVector),
-			F0Arg,
-			F90Arg,
-			ThicknessArg,
-			IORArg
-		});
+	MIR::FValueRef OutF0Arg = Em.Extern<MIR::FExternSimpleHLSL>({ MIR::EExternSimpleType::Float3, TEXT("SubstrateGetThinFilmF0F90($0, $1, $2, $3, $4).F0") },
+		Em.Dot(NormalArg, CameraVector),
+		F0Arg,
+		F90Arg,
+		ThicknessArg,
+		IORArg);
 
 	// Calling SubstrateGetThinFilmF0F90 a second time. SUBSTRATE_TODO: create a structure to pull the data from the first call instead
-	MIR::FValueRef OutF90Arg = Em.InlineHLSL(
-		MIR::FType::MakeFloatVector(3),
-		TEXTVIEW("SubstrateGetThinFilmF0F90($0, $1, $2, $3, $4).F90"),
-		{
-			Em.Dot(NormalArg, CameraVector),
-			F0Arg,
-			F90Arg,
-			ThicknessArg,
-			IORArg
-		});
+	MIR::FValueRef OutF90Arg = Em.Extern<MIR::FExternSimpleHLSL>({ MIR::EExternSimpleType::Float3, TEXT("SubstrateGetThinFilmF0F90($0, $1, $2, $3, $4).F90") },
+		Em.Dot(NormalArg, CameraVector),
+		F0Arg,
+		F90Arg,
+		ThicknessArg,
+		IORArg);
 
 	Em.Output(0, OutF0Arg);
 	Em.Output(1, OutF90Arg);
