New material translator:

diff --git a/Engine/Source/Runtime/Engine/Private/Materials/MaterialExternalCodeRegistry.h b/Engine/Source/Runtime/Engine/Private/Materials/MaterialExternalCodeRegistry.h
--- a/Engine/Source/Runtime/Engine/Private/Materials/MaterialExternalCodeRegistry.h
+++ b/Engine/Source/Runtime/Engine/Private/Materials/MaterialExternalCodeRegistry.h
@@ -17,6 +17,7 @@
 UENUM(Flags)
 enum class EMaterialShaderFrequency : uint8
 {
+	None    = 0,
 	Vertex  = 1 << SF_Vertex,
 	Pixel   = 1 << SF_Pixel,
 	Compute = 1 << SF_Compute,
diff --git a/Engine/Source/Runtime/Engine/Private/Materials/MaterialIR.cpp b/Engine/Source/Runtime/Engine/Private/Materials/MaterialIR.cpp
--- a/Engine/Source/Runtime/Engine/Private/Materials/MaterialIR.cpp
+++ b/Engine/Source/Runtime/Engine/Private/Materials/MaterialIR.cpp
@@ -13,28 +13,16 @@
 namespace MIR
 {
 
-const TCHAR* LexToString(EStage Stage)
-{
-	switch (Stage)
-	{
-		case Stage_Vertex:	return TEXT("Vertex");
-		case Stage_Pixel: 	return TEXT("Pixel");
-		case Stage_Compute: return TEXT("Compute");
-		case NumStages: UE_MIR_UNREACHABLE();
-	}
-	return nullptr;
-}
-
 const TCHAR* LexToString(EValueKind Kind)
 {
+	#define CASE(Name) case VK_##Name: return TEXT(#Name);
 	switch (Kind)
 	{
 		case VK_Poison: return TEXT("Poison");
 		case VK_AnalysisError: return TEXT("AnalysisError");
 		case VK_Constant: return TEXT("Constant");
-		case VK_ExternalInput: return TEXT("ExternalInput");
+		case VK_Builtin: return TEXT("Builtin");
 		case VK_MaterialParameterCollection: return TEXT("MaterialParameterCollection");
-		case VK_ScreenTexture: return TEXT("ScreenTexture");
 		case VK_ShadingModel: return TEXT("ShadingModel");
 		case VK_TextureObject: return TEXT("TextureObject");
 		case VK_RuntimeVirtualTextureObject: return TEXT("RuntimeVirtualTextureObject");
@@ -49,7 +37,8 @@ const TCHAR* LexToString(EValueKind Kind)
 		case VK_Scalar: return TEXT("Scalar");
 		case VK_TextureRead: return TEXT("TextureRead");
 		case VK_VTPageTableRead: return TEXT("VTPageTableRead");
-		case VK_InlineHLSL: return TEXT("InlineHLSL");
+		case VK_Extern: return TEXT("Extern");
+		case VK_ExternDifferential: return TEXT("ExternDifferential");
 		case VK_SubstrateDefaultSlab: return TEXT("SubstrateDefaultSlab");
 		case VK_SubstrateSlab: return TEXT("SubstrateSlab");
 		case VK_SubstrateShadingModels: return TEXT("SubstrateShadingModel");
@@ -71,6 +60,7 @@ const TCHAR* LexToString(EValueKind Kind)
 			UE_MIR_UNREACHABLE();
 	}
 	return nullptr;
+	#undef CASE
 }
 
 bool FValue::IsAnalyzed(EStage Stage) const
@@ -110,9 +100,8 @@ uint32 FValue::GetSizeInBytes() const
 		case VK_Poison: return sizeof(FPoison);
 		case VK_AnalysisError: return sizeof(FAnalysisError);
 		case VK_Constant: return sizeof(FConstant);
-		case VK_ExternalInput: return sizeof(FExternalInput);
+		case VK_Builtin: return sizeof(FBuiltin);
 		case VK_MaterialParameterCollection: return sizeof(FMaterialParameterCollection);
-		case VK_ScreenTexture: return sizeof(FScreenTexture);
 		case VK_ShadingModel: return sizeof(FShadingModel);
 		case VK_TextureObject: return sizeof(FTextureObject);
 		case VK_RuntimeVirtualTextureObject: return sizeof(FRuntimeVirtualTextureObject);
@@ -127,7 +116,6 @@ uint32 FValue::GetSizeInBytes() const
 		case VK_Scalar: return sizeof(FScalar);
 		case VK_TextureRead: return sizeof(FTextureRead);
 		case VK_VTPageTableRead: return sizeof(FVTPageTableRead);
-		case VK_InlineHLSL: return sizeof(FInlineHLSL);
 		case VK_SubstrateDefaultSlab: return sizeof(FSubstrateDefaultSlab);
 		case VK_SubstrateSlab: return sizeof(FSubstrateSlab);
 		case VK_SubstrateShadingModels: return sizeof(FSubstrateShadingModels);
@@ -143,6 +131,9 @@ uint32 FValue::GetSizeInBytes() const
 		case VK_Call: return sizeof(FCall);
 		case VK_CallParameterOutput: return sizeof(FCallParameterOutput);
 		case VK_PreshaderParameter: return sizeof(FPreshaderParameter);
+		case VK_Extern: return static_cast<const FExtern*>(this)->Impl->ByteSize;
+		case VK_ExternDifferential: return sizeof(FExternDifferential);
+
 			/* invalid entries */
 		case VK_InstructionBegin:
 		case VK_InstructionEnd:
@@ -215,10 +206,16 @@ TConstArrayView<FValue*> FValue::GetUses() const
 			return { &This->TextureObject, FVTPageTableRead::NumStaticUses };
 		}
 
-		case VK_InlineHLSL:
+		case VK_Extern:
 		{
-			auto This = static_cast<const FInlineHLSL*>(this);
-			return { This->Arguments, This->NumArguments };
+			auto This = static_cast<const FExtern*>(this);
+			return This->GetArguments();
+		}
+
+		case VK_ExternDifferential:
+		{
+			auto This = static_cast<const FExternDifferential*>(this);
+			return { &This->ExternArg, FExternDifferential::NumStaticUses };
 		}
 
 		case VK_SubstrateDefaultSlab:
@@ -571,194 +568,16 @@ FPoison* FPoison::Get()
 	return &Poison;
 }
 
-const TCHAR* LexToString(EExternalInput Input)
-{
-	switch (Input)
-	{
-		case EExternalInput::TexCoord0:     return TEXT("TexCoord0");
-		case EExternalInput::TexCoord1:     return TEXT("TexCoord1");
-		case EExternalInput::TexCoord2:     return TEXT("TexCoord2");
-		case EExternalInput::TexCoord3:     return TEXT("TexCoord3");
-		case EExternalInput::TexCoord4:     return TEXT("TexCoord4");
-		case EExternalInput::TexCoord5:     return TEXT("TexCoord5");
-		case EExternalInput::TexCoord6:     return TEXT("TexCoord6");
-		case EExternalInput::TexCoord7:     return TEXT("TexCoord7");
-		case EExternalInput::WorldPosition_Absolute: return TEXT("WorldPosition_Absolute");
-		case EExternalInput::WorldPosition_AbsoluteNoOffsets: return TEXT("WorldPosition_AbsoluteNoOffsets");
-		case EExternalInput::WorldPosition_CameraRelative: return TEXT("WorldPosition_CameraRelative");
-		case EExternalInput::WorldPosition_CameraRelativeNoOffsets: return TEXT("WorldPosition_CameraRelativeNoOffsets");
-		case EExternalInput::LocalPosition_Instance:			return TEXT("LocalPosition_Instance");
-		case EExternalInput::LocalPosition_InstanceNoOffsets:	return TEXT("LocalPosition_InstanceNoOffsets");
-		case EExternalInput::LocalPosition_Primitive:			return TEXT("LocalPosition_Primitive");
-		case EExternalInput::LocalPosition_PrimitiveNoOffsets:	return TEXT("LocalPosition_PrimitiveNoOffsets");
-
-		case EExternalInput::TexCoord0_Ddx: return TEXT("TexCoord0_Ddx");
-		case EExternalInput::TexCoord1_Ddx: return TEXT("TexCoord1_Ddx");
-		case EExternalInput::TexCoord2_Ddx: return TEXT("TexCoord2_Ddx");
-		case EExternalInput::TexCoord3_Ddx: return TEXT("TexCoord3_Ddx");
-		case EExternalInput::TexCoord4_Ddx: return TEXT("TexCoord4_Ddx");
-		case EExternalInput::TexCoord5_Ddx: return TEXT("TexCoord5_Ddx");
-		case EExternalInput::TexCoord6_Ddx: return TEXT("TexCoord6_Ddx");
-		case EExternalInput::TexCoord7_Ddx: return TEXT("TexCoord7_Ddx");
-		case EExternalInput::WorldPosition_Absolute_Ddx: return TEXT("WorldPosition_Absolute_Ddx");
-		case EExternalInput::WorldPosition_AbsoluteNoOffsets_Ddx: return TEXT("WorldPosition_AbsoluteNoOffsets_Ddx");
-		case EExternalInput::WorldPosition_CameraRelative_Ddx: return TEXT("WorldPosition_CameraRelative_Ddx");
-		case EExternalInput::WorldPosition_CameraRelativeNoOffsets_Ddx: return TEXT("WorldPosition_CameraRelativeNoOffsets_Ddx");
-		case EExternalInput::LocalPosition_Instance_Ddx:			return TEXT("LocalPosition_Instance_Ddx");
-		case EExternalInput::LocalPosition_InstanceNoOffsets_Ddx:	return TEXT("LocalPosition_InstanceNoOffsets_Ddx");
-		case EExternalInput::LocalPosition_Primitive_Ddx:			return TEXT("LocalPosition_Primitive_Ddx");
-		case EExternalInput::LocalPosition_PrimitiveNoOffsets_Ddx:	return TEXT("LocalPosition_PrimitiveNoOffsets_Ddx");
-
-		case EExternalInput::TexCoord0_Ddy: return TEXT("TexCoord0_Ddy");
-		case EExternalInput::TexCoord1_Ddy: return TEXT("TexCoord1_Ddy");
-		case EExternalInput::TexCoord2_Ddy: return TEXT("TexCoord2_Ddy");
-		case EExternalInput::TexCoord3_Ddy: return TEXT("TexCoord3_Ddy");
-		case EExternalInput::TexCoord4_Ddy: return TEXT("TexCoord4_Ddy");
-		case EExternalInput::TexCoord5_Ddy: return TEXT("TexCoord5_Ddy");
-		case EExternalInput::TexCoord6_Ddy: return TEXT("TexCoord6_Ddy");
-		case EExternalInput::TexCoord7_Ddy: return TEXT("TexCoord7_Ddy");
-		case EExternalInput::WorldPosition_Absolute_Ddy: return TEXT("WorldPosition_Absolute_Ddy");
-		case EExternalInput::WorldPosition_AbsoluteNoOffsets_Ddy: return TEXT("WorldPosition_AbsoluteNoOffsets_Ddy");
-		case EExternalInput::WorldPosition_CameraRelative_Ddy: return TEXT("WorldPosition_CameraRelative_Ddy");
-		case EExternalInput::WorldPosition_CameraRelativeNoOffsets_Ddy: return TEXT("WorldPosition_CameraRelativeNoOffsets_Ddy");
-		case EExternalInput::LocalPosition_Instance_Ddy:			return TEXT("LocalPosition_Instance_Ddy");
-		case EExternalInput::LocalPosition_InstanceNoOffsets_Ddy:	return TEXT("LocalPosition_InstanceNoOffsets_Ddy");
-		case EExternalInput::LocalPosition_Primitive_Ddy:			return TEXT("LocalPosition_Primitive_Ddy");
-		case EExternalInput::LocalPosition_PrimitiveNoOffsets_Ddy:	return TEXT("LocalPosition_PrimitiveNoOffsets_Ddy");
-
-		case EExternalInput::ActorPosition_Absolute: return TEXT("ActorPosition_Absolute");
-		case EExternalInput::ActorPosition_CameraRelative: return TEXT("ActorPosition_CameraRelative");
-		case EExternalInput::ObjectPosition_Absolute: return TEXT("ObjectPosition_Absolute");
-		case EExternalInput::ObjectPosition_CameraRelative: return TEXT("ObjectPosition_CameraRelative");
-		case EExternalInput::ViewMaterialTextureMipBias: return TEXT("ViewMaterialTextureMipBias");
-		case EExternalInput::ViewMaterialTextureDerivativeMultiply: return TEXT("ViewMaterialTextureDerivativeMultiply");
-		case EExternalInput::GlobalDistanceField: return TEXT("GlobalDistanceField");
-		case EExternalInput::DynamicParticleParameterIndex: return TEXT("DynamicParticleParameterIndex");
-		case EExternalInput::CompilingPreviousFrame: return TEXT("CompilingPreviousFrame");
-		case EExternalInput::DistanceCullFade: return TEXT("GetDistanceCullFade()");
-
-		default: UE_MIR_UNREACHABLE();
-	}
-}
-
-FType GetExternalInputType(EExternalInput Id)
+const TCHAR* LexToString(EBuiltin ID)
 {
-	if (IsExternalInputTexCoordOrPartialDerivative(Id))
-	{
-		return FType::MakeFloatVector(2);
-	}
-	
-	switch (Id)
+	switch (ID)
 	{
-		case EExternalInput::ViewMaterialTextureMipBias:
-		case EExternalInput::ViewMaterialTextureDerivativeMultiply:
-		case EExternalInput::GlobalDistanceField:						// Type is arbitrary, value not actually used, only as a flag for value analyzer
-		case EExternalInput::DistanceCullFade:
-			return FType::MakeFloatScalar();
-
-		case EExternalInput::ActorPosition_CameraRelative:
-		case EExternalInput::ObjectPosition_CameraRelative:
-		case EExternalInput::WorldPosition_CameraRelative:
-		case EExternalInput::WorldPosition_CameraRelativeNoOffsets:
-		case EExternalInput::WorldPosition_CameraRelative_Ddx:
-		case EExternalInput::WorldPosition_CameraRelativeNoOffsets_Ddx:
-		case EExternalInput::WorldPosition_CameraRelative_Ddy:
-		case EExternalInput::WorldPosition_CameraRelativeNoOffsets_Ddy:
-		case EExternalInput::LocalPosition_Instance:
-		case EExternalInput::LocalPosition_InstanceNoOffsets:
-		case EExternalInput::LocalPosition_Primitive:
-		case EExternalInput::LocalPosition_PrimitiveNoOffsets:
-		case EExternalInput::LocalPosition_Instance_Ddx:
-		case EExternalInput::LocalPosition_InstanceNoOffsets_Ddx:
-		case EExternalInput::LocalPosition_Primitive_Ddx:
-		case EExternalInput::LocalPosition_PrimitiveNoOffsets_Ddx:
-		case EExternalInput::LocalPosition_Instance_Ddy:
-		case EExternalInput::LocalPosition_InstanceNoOffsets_Ddy:
-		case EExternalInput::LocalPosition_Primitive_Ddy:
-		case EExternalInput::LocalPosition_PrimitiveNoOffsets_Ddy:
-			return FType::MakeFloatVector(3);
-
-		case EExternalInput::ActorPosition_Absolute:
-		case EExternalInput::ObjectPosition_Absolute:
-		case EExternalInput::WorldPosition_Absolute:
-		case EExternalInput::WorldPosition_AbsoluteNoOffsets:
-		case EExternalInput::WorldPosition_Absolute_Ddx:				// Technically these derivatives are not LWC, but the emitter needs them to be for evaluation of
-		case EExternalInput::WorldPosition_AbsoluteNoOffsets_Ddx:		// analytic derivative expressions.  So we cast them to LWC on load.  The shader compiler should
-		case EExternalInput::WorldPosition_Absolute_Ddy:				// be able to optimize the math (the tile offsets will all be zero, cancelling things out).
-		case EExternalInput::WorldPosition_AbsoluteNoOffsets_Ddy:
-			return FType::MakeDoubleVector(3);
-
-		case EExternalInput::DynamicParticleParameterIndex:
-			return FType::MakeIntScalar();
-
-		case EExternalInput::CompilingPreviousFrame:
-			return FType::MakeBoolScalar();
-
+		case EBuiltin::TextureMipBias: return TEXT("TextureMipBias");
+		case EBuiltin::TextureDerivativeMultiply: return TEXT("TextureDerivativeMultiply");
 		default: UE_MIR_UNREACHABLE();
 	}
 }
 
-// Validate derivative related assumptions in EExternalInput enum
-static_assert(EExternalInput::WithDerivatives_First == EExternalInput::TexCoord0);
-static_assert(EExternalInput::WithDerivatives_Last == EExternalInput::WithDerivatives_LastDdy);
-static_assert((int32)EExternalInput::WithDerivatives_LastDdx - (int32)EExternalInput::WithDerivatives_LastVal == ExternalInputWithDerivativesNum);
-static_assert((int32)EExternalInput::WithDerivatives_LastDdy - (int32)EExternalInput::WithDerivatives_LastDdx == ExternalInputWithDerivativesNum);
-
-bool IsExternalInputWithDerivatives(EExternalInput Id)
-{
-	return Id >= EExternalInput::WithDerivatives_First && Id <= EExternalInput::WithDerivatives_Last;
-}
-
-EExternalInput GetExternalInputDerivative(EExternalInput Id, EDerivativeAxis Axis)
-{
-	check(IsExternalInputWithDerivatives(Id));
-
-	int32 IntId = (int32(Id) - (int32)MIR::EExternalInput::WithDerivatives_First) % ExternalInputWithDerivativesNum + (int32)MIR::EExternalInput::WithDerivatives_First;
-
-	// Add one to axis enum (X==0, Y==1) to produce a value with range 1 to 2 for computing offset to the derivative variation of a value
-	return (EExternalInput)(IntId + ((int)Axis + 1) * ExternalInputWithDerivativesNum);
-}
-
-MIR::EExternalInput TexCoordIndexToExternalInput(int32 TexCoordIndex)
-{
-	check(TexCoordIndex < TexCoordMaxNum);
-	return (MIR::EExternalInput)((int)MIR::EExternalInput::TexCoord0 + TexCoordIndex);
-}
-
-int32 ExternalInputToTexCoordIndex(EExternalInput Id)
-{
-	check(IsExternalInputTexCoordOrPartialDerivative(Id));
-	return (int32(Id) - (int32)MIR::EExternalInput::WithDerivatives_First) % ExternalInputWithDerivativesNum;
-}
-
-bool IsExternalInputTexCoord(EExternalInput Id)
-{
-	return Id >= EExternalInput::TexCoord0 && Id <= EExternalInput::TexCoord7;
-}
-
-bool IsExternalInputTexCoordDdx(EExternalInput Id)
-{
-	return Id >= EExternalInput::TexCoord0_Ddx && Id <= EExternalInput::TexCoord7_Ddx;
-}
-
-bool IsExternalInputTexCoordDdy(EExternalInput Id)
-{
-	return Id >= EExternalInput::TexCoord0_Ddy && Id <= EExternalInput::TexCoord7_Ddy;
-}
-
-bool IsExternalInputTexCoordOrPartialDerivative(EExternalInput Id)
-{
-	return IsExternalInputTexCoord(Id) || IsExternalInputTexCoordDdx(Id) || IsExternalInputTexCoordDdy(Id);
-}
-
-bool IsExternalInputWorldPosition(EExternalInput Id)
-{
-	return
-		(Id >= EExternalInput::WorldPosition_Absolute && Id <= EExternalInput::WorldPosition_CameraRelativeNoOffsets) ||
-		(Id >= EExternalInput::WorldPosition_Absolute_Ddx && Id <= EExternalInput::WorldPosition_CameraRelativeNoOffsets_Ddx) ||
-		(Id >= EExternalInput::WorldPosition_Absolute_Ddy && Id <= EExternalInput::WorldPosition_CameraRelativeNoOffsets_Ddy);
-}
-
 FBlock* FBlock::FindCommonParentWith(MIR::FBlock* Other)
 {
 	FBlock* A = this;
diff --git a/Engine/Source/Runtime/Engine/Private/Materials/MaterialIRDebug.cpp b/Engine/Source/Runtime/Engine/Private/Materials/MaterialIRDebug.cpp
--- a/Engine/Source/Runtime/Engine/Private/Materials/MaterialIRDebug.cpp
+++ b/Engine/Source/Runtime/Engine/Private/Materials/MaterialIRDebug.cpp
@@ -124,13 +124,9 @@ struct FDebugDumpIRUseGraphState
 				default: UE_MIR_UNREACHABLE();
 			}
 		}
-		else if (const FExternalInput* ExternalInput = Value->As<FExternalInput>())
+		else if (const FBuiltin* ExternalInput = Value->As<FBuiltin>())
 		{
 			Out.Append(LexToString(ExternalInput->Id));
-			if (ExternalInput->Id == EExternalInput::DynamicParticleParameterIndex)
-			{
-				Out.Appendf(TEXT("[%u]"), ExternalInput->UserData);
-			}
 		}
 		else if (const FSetMaterialOutput* SetMaterailOutput = Value->As<FSetMaterialOutput>())
 		{
@@ -157,6 +153,10 @@ struct FDebugDumpIRUseGraphState
 		{
 			Out.Append(LexToString(Operator->Op));
 		}
+		else if (const FExtern* Extern = Value->As<FExtern>())
+		{
+			Out.Append(Extern->GetInfo().Name);
+		}
 	}
 
 	void DumpUseInfo(const FValue* Value, const FValue* Use, int32 UseIndex)
@@ -302,6 +302,7 @@ struct FDebugDumpIRState
 
 			// Dump the instruction properties.
 			AppendInstructionProperties(Instr);
+			Out.Appendf(TEXT(" <Users=%d>"), Instr->Linkage[CurrentEntryPointIndex].NumUsers);
 
 			Out.Append(TEXT("\n"));
 		}
@@ -341,17 +342,6 @@ struct FDebugDumpIRState
 			Out.Appendf(TEXT(" TextureIndex=%d"), PreshaderParameter->TextureIndex);
 			Out.Appendf(TEXT(" PreshaderOffset=%u"), PreshaderParameter->Analysis_PreshaderOffset);
 		}
-		else if (const FInlineHLSL* InlineHLSL = Instr->As<FInlineHLSL>())
-		{
-			if (InlineHLSL->HasFlags(EValueFlags::HasDynamicHLSLCode))
-			{
-				Out.Appendf(TEXT(" \"%.*s\""), InlineHLSL->Code.Len, InlineHLSL->Code.Ptr);
-			}
-			else
-			{
-				Out.Appendf(TEXT(" \"%s\""), *InlineHLSL->ExternalCodeDeclaration->Definition);
-			}
-		}
 		else if (const FHardwarePartialDerivative* Derivative = Instr->As<FHardwarePartialDerivative>())
 		{
 			Out.Append(Derivative->Axis == MIR::EDerivativeAxis::X ? TEXT(" \"ddx\"") : TEXT(" \"ddy\""));
@@ -368,6 +358,15 @@ struct FDebugDumpIRState
 		{
 			Out.Appendf(TEXT(" Output=\"%s\""), *CallOutput->Call->As<FCall>()->Function->GetOutputParameter(CallOutput->Index).Name.ToString());
 		}
+		else if (const FExtern* Extern = Instr->As<FExtern>())
+		{
+			Out.Appendf(TEXT(" \"%s\" "), Extern->GetInfo().Name.GetData());
+			Extern->EmitDebugInfo(Out);
+		}
+		else if (const FExternDifferential* ExternDifferential = Instr->As<FExternDifferential>())
+		{
+			Out.Append(ExternDifferential->Axis == EDerivativeAxis::X ? TEXT("ddx") : TEXT("ddy"));
+		}
 	}
 
 	// Appends a reference to the specified value. This will look like "%x" if Value is an
@@ -407,9 +406,9 @@ struct FDebugDumpIRState
 				UE_MIR_UNREACHABLE();
 			}
 		}
-		else if (const FExternalInput* ExternalInput = Value->As<FExternalInput>())
+		else if (const FBuiltin* ExternalInput = Value->As<FBuiltin>())
 		{
-			Out.Appendf(TEXT("[ExternalInput \"%s\"]"), LexToString(Value->As<FExternalInput>()->Id));
+			Out.Appendf(TEXT("[ExternalInput \"%s\"]"), LexToString(Value->As<FBuiltin>()->Id));
 		}
 		else if (const FTextureObject* TextureObject = Value->As<FTextureObject>())
 		{
diff --git a/Engine/Source/Runtime/Engine/Private/Materials/MaterialIREmitter.cpp b/Engine/Source/Runtime/Engine/Private/Materials/MaterialIREmitter.cpp
--- a/Engine/Source/Runtime/Engine/Private/Materials/MaterialIREmitter.cpp
+++ b/Engine/Source/Runtime/Engine/Private/Materials/MaterialIREmitter.cpp
@@ -9,6 +9,7 @@
 #include "Materials/MaterialExternalCodeRegistry.h"
 #include "Materials/MaterialAggregate.h"
 #include "Materials/MaterialInterface.h"
+#include "Materials/MaterialIRExtern.h"
 #include "Shader/Preshader2.h"
 #include "Shader/ShaderTypes.h"
 #include "MaterialShared.h"
@@ -89,11 +90,6 @@ FSwizzleMask::FSwizzleMask(bool bMaskX, bool bMaskY, bool bMaskZ, bool bMaskW)
 	}
 }
 
-FSwizzleMask FSwizzleMask::XYZ()
-{
-	return { EVectorComponent::X, EVectorComponent::Y, EVectorComponent::Z };
-}
-
 void FSwizzleMask::Append(EVectorComponent Component)
 {
 	check(NumComponents < 4);
@@ -1077,11 +1073,11 @@ FValueRef FEmitter::AnalysisError(FType Type, FStringView Message)
 	return Error;
 }
 
-FValueRef FEmitter::ExternalInput(EExternalInput Id, uint32 UserData)
+FValueRef FEmitter::Builtin(EBuiltin Id)
 {
-	FExternalInput Prototype = MakePrototype<FExternalInput>(GetExternalInputType(Id));
+	FType Type = FType::MakeFloatScalar();
+	FBuiltin Prototype = MakePrototype<FBuiltin>(Type);
 	Prototype.Id = Id;
-	Prototype.UserData = UserData;
 	return EmitPrototype(*this, Prototype);
 }
 
@@ -1187,7 +1183,7 @@ static EMaterialSamplerType MapSamplerTypeForTexture(FValueRef InTexture, EMater
 }
 
 static FValueRef VTPageTableLoadFromSamplerSource(
-	FEmitter& Em, FValueRef Texture, const FTextureSampleBaseAttributes& BaseAttributes, FValueRef TexCoord,
+	FEmitter& Em, FValueRef Texture, const FTextureSampleAttributes& BaseAttributes, FValueRef TexCoord,
 	FValueRef TexCoordDdx = {}, FValueRef TexCoordDdy = {}, ETextureMipValueMode MipValueMode = TMVM_None, FValueRef MipValue = {})
 {
 	// Cast input texture to UTexture. If it's a URuntimeVirtualTexture, we accept the cast to be null when passed to GetTextureAddressForSamplerSource().
@@ -1241,6 +1237,7 @@ FValueRef FEmitter::Parameter(FName Name, FMaterialParameterMetadata& Metadata,
 		Prototype.VTLayerIndex = VTLayerIndex;
 		Prototype.VTPageTableIndex = VTPageTableIndex;
 		Prototype.Analysis_UniformParameterIndex = INDEX_NONE;
+		
 		return EmitPrototype(*this, Prototype);
 	};
 
@@ -1367,40 +1364,6 @@ FValueRef FEmitter::CustomPrimitiveData(uint32 PrimitiveDataIndex)
 	return {};
 }
 
-FValueRef FEmitter::SceneTexture(ESceneTextureId SceneTextureId)
-{
-	FScreenTexture Proto = MakePrototype<FScreenTexture>(MIR::FType::MakeIntScalar());
-	Proto.TextureKind = EScreenTexture::SceneTexture;
-	Proto.Id = SceneTextureId;
-	return EmitPrototype(*this, Proto);
-}
-
-FValueRef FEmitter::UserSceneTexture(FName UserSceneTexture)
-{
-	FScreenTexture Proto = MakePrototype<FScreenTexture>(MIR::FType::MakeIntScalar());
-	Proto.TextureKind = EScreenTexture::UserSceneTexture;
-	Proto.Id = PPI_UserSceneTexture0;				// Actual ID determined later during MIR lowering
-	Proto.UserSceneTexture = UserSceneTexture;
-	return EmitPrototype(*this, Proto);
-}
-
-FValueRef FEmitter::DBufferTexture(EDBufferTextureId DBufferId)
-{
-	FScreenTexture Proto = MakePrototype<FScreenTexture>(MIR::FType::MakeIntScalar());
-	Proto.TextureKind = EScreenTexture::DBufferTexture;
-	Proto.DBufferId = DBufferId;
-	return EmitPrototype(*this, Proto);
-}
-
-FValueRef FEmitter::ScreenTexture(EScreenTexture TextureKind)
-{
-	check(TextureKind != EScreenTexture::SceneTexture && TextureKind != EScreenTexture::UserSceneTexture && TextureKind != EScreenTexture::DBufferTexture);
-	FScreenTexture Proto = MakePrototype<FScreenTexture>(MIR::FType::MakeIntScalar());
-	Proto.TextureKind = TextureKind;
-	Proto.Id = PPI_SceneColor;					// Ignored
-	return EmitPrototype(*this, Proto);
-}
-
 /*------------------------------ Instruction emission ------------------------------*/
 
 // Helper function to create a single SetMaterialOutput instruction. It is used to create
@@ -3544,7 +3507,7 @@ FValueRef FEmitter::Nop(FValueRef Arg)
 	return EmitPrototype(*this, Prototype);
 }
 
-FValueRef FEmitter::TextureGather(FValueRef Texture, FValueRef TexCoord, ETextureReadMode GatherMode, const FTextureSampleBaseAttributes& BaseAttributes)
+FValueRef FEmitter::TextureGather(FValueRef Texture, FValueRef TexCoord, ETextureReadMode GatherMode, const FTextureSampleAttributes& BaseAttributes)
 {
 	check(GatherMode >= ETextureReadMode::GatherRed && GatherMode <= ETextureReadMode::GatherAlpha);
 
@@ -3570,7 +3533,7 @@ FValueRef FEmitter::TextureGather(FValueRef Texture, FValueRef TexCoord, ETextur
 	return EmitPrototype(*this, Prototype);
 }
 
-FValueRef FEmitter::TextureSample(FValueRef Texture, FValueRef TexCoord, bool bAutomaticViewMipBias, const FTextureSampleBaseAttributes& BaseAttributes)
+FValueRef FEmitter::TextureSample(FValueRef Texture, FValueRef TexCoord, bool bAutomaticViewMipBias, const FTextureSampleAttributes& BaseAttributes)
 {
 	if (IsAnyNotValid(Texture, TexCoord))
 	{
@@ -3597,7 +3560,7 @@ FValueRef FEmitter::TextureSample(FValueRef Texture, FValueRef TexCoord, bool bA
 
 	if (bAutomaticViewMipBias)
 	{
-		FValueRef ViewMaterialTextureMipBias = ExternalInput(EExternalInput::ViewMaterialTextureMipBias);
+		FValueRef ViewMaterialTextureMipBias = Builtin(EBuiltin::TextureMipBias);
 		PrototypePixel.Mode = ETextureReadMode::MipBias;
 		PrototypePixel.MipValue = ViewMaterialTextureMipBias;
 
@@ -3621,7 +3584,7 @@ FValueRef FEmitter::TextureSample(FValueRef Texture, FValueRef TexCoord, bool bA
 	return EmitPrototype(*this, StageSwitch);
 }
 
-FValueRef FEmitter::TextureSampleLevel(FValueRef Texture, FValueRef TexCoord, FValueRef MipLevel, bool bAutomaticViewMipBias, const FTextureSampleBaseAttributes& BaseAttributes)
+FValueRef FEmitter::TextureSampleLevel(FValueRef Texture, FValueRef TexCoord, FValueRef MipLevel, bool bAutomaticViewMipBias, const FTextureSampleAttributes& BaseAttributes)
 {
 	if (IsAnyNotValid(Texture, TexCoord, MipLevel))
 	{
@@ -3640,7 +3603,7 @@ FValueRef FEmitter::TextureSampleLevel(FValueRef Texture, FValueRef TexCoord, FV
 
 	if (bAutomaticViewMipBias)
 	{
-		Prototype.MipValue = Operator(BO_Add, MipLevel, ExternalInput(EExternalInput::ViewMaterialTextureMipBias));
+		Prototype.MipValue = Operator(BO_Add, MipLevel, Builtin(EBuiltin::TextureMipBias));
 	}
 
 	if (IsVirtualSamplerType(SamplerType))
@@ -3651,7 +3614,7 @@ FValueRef FEmitter::TextureSampleLevel(FValueRef Texture, FValueRef TexCoord, FV
 	return EmitPrototype(*this, Prototype);
 }
 
-FValueRef FEmitter::TextureSampleBias(FValueRef Texture, FValueRef TexCoord, FValueRef MipBias, bool bAutomaticViewMipBias, const FTextureSampleBaseAttributes& BaseAttributes)
+FValueRef FEmitter::TextureSampleBias(FValueRef Texture, FValueRef TexCoord, FValueRef MipBias, bool bAutomaticViewMipBias, const FTextureSampleAttributes& BaseAttributes)
 {
 	if (IsAnyNotValid(Texture, TexCoord, MipBias))
 	{
@@ -3660,7 +3623,7 @@ FValueRef FEmitter::TextureSampleBias(FValueRef Texture, FValueRef TexCoord, FVa
 
 	if (bAutomaticViewMipBias)
 	{
-		MipBias = Operator(BO_Add, MipBias, ExternalInput(EExternalInput::ViewMaterialTextureMipBias));
+		MipBias = Operator(BO_Add, MipBias, Builtin(EBuiltin::TextureMipBias));
 	}
 
 	const EMaterialSamplerType SamplerType = MapSamplerTypeForTexture(Texture, BaseAttributes.SamplerType);
@@ -3703,7 +3666,7 @@ FValueRef FEmitter::TextureSampleBias(FValueRef Texture, FValueRef TexCoord, FVa
 	return EmitPrototype(*this, StageSwitch);
 }
 
-FValueRef FEmitter::TextureSampleGrad(FValueRef Texture, FValueRef TexCoord, FValueRef TexCoordDdx, FValueRef TexCoordDdy, bool bAutomaticViewMipBias, const FTextureSampleBaseAttributes& BaseAttributes)
+FValueRef FEmitter::TextureSampleGrad(FValueRef Texture, FValueRef TexCoord, FValueRef TexCoordDdx, FValueRef TexCoordDdy, bool bAutomaticViewMipBias, const FTextureSampleAttributes& BaseAttributes)
 {
 	if (IsAnyNotValid(Texture, TexCoord, TexCoordDdx, TexCoordDdy))
 	{
@@ -3723,7 +3686,7 @@ FValueRef FEmitter::TextureSampleGrad(FValueRef Texture, FValueRef TexCoord, FVa
 
 	if (bAutomaticViewMipBias)
 	{
-		FValueRef ViewMaterialTextureDerivativeMultiply = ExternalInput(EExternalInput::ViewMaterialTextureDerivativeMultiply);
+		FValueRef ViewMaterialTextureDerivativeMultiply = Builtin(EBuiltin::TextureDerivativeMultiply);
 		Prototype.TexCoordDdx = Operator(BO_Multiply, Prototype.TexCoordDdx, ViewMaterialTextureDerivativeMultiply);
 		Prototype.TexCoordDdy = Operator(BO_Multiply, Prototype.TexCoordDdy, ViewMaterialTextureDerivativeMultiply);
 	}
@@ -3791,18 +3754,6 @@ FValueRef FEmitter::PartialDerivative(FValueRef Value, EDerivativeAxis Axis)
 	return StageSwitch(Value->Type, StageValues);
 }
 
-static FValue* DifferentiateExternalInput(FEmitter& Emitter, FExternalInput* ExternalInput, EDerivativeAxis Axis)
-{
-	// External inputs with derivatives have their own matching DDX/DDY inputs.
-	if (IsExternalInputWithDerivatives(ExternalInput->Id))
-	{
-		return Emitter.ExternalInput(GetExternalInputDerivative(ExternalInput->Id, Axis));
-	}
-
-	// All other inputs are assumed constant.
-	return Emitter.Cast(Emitter.ConstantZero(ExternalInput->Type.GetPrimitive().ScalarKind), ExternalInput->Type);
-}
-
 static FValue* DifferentiateOperator(FEmitter& E, FOperator* Op, EDerivativeAxis Axis)
 {
 	EScalarKind OpScalarKind = Op->Type.GetPrimitive().ScalarKind;
@@ -4077,8 +4028,10 @@ FValueRef FEmitter::AnalyticalPartialDerivative(FValueRef Value, EDerivativeAxis
 
 	switch (Value->Kind)
 	{
-		case VK_ExternalInput:
-			return DifferentiateExternalInput(*this, Value->As<FExternalInput>(), Axis);
+		case VK_Builtin:
+		{
+			return Cast(ConstantZero(Value->Type.GetPrimitive().ScalarKind), Value->Type);
+		}
 
 		case VK_Composite:
 		{
@@ -4139,15 +4092,20 @@ FValueRef FEmitter::AnalyticalPartialDerivative(FValueRef Value, EDerivativeAxis
 			return Cast(AnalyticalPartialDerivative(AsScalar->Arg, Axis), AsScalar->Type);
 		}
 
-		case VK_InlineHLSL:
+		case VK_Extern:
 		{
-			FInlineHLSL* AsInlineHLSL = Value->As<FInlineHLSL>();
-			if (!EnumHasAnyFlags(AsInlineHLSL->Flags, EValueFlags::HasDynamicHLSLCode) && AsInlineHLSL->ExternalCodeDeclaration->Derivative == EDerivativeStatus::Valid)
+			FExtern* AsExtern = Value->As<FExtern>();
+			const FExternInfo& Info = AsExtern->GetInfo();
+			if ((Info.Flags & MIR::EExternFlags::NoDifferentials) != MIR::EExternFlags::None)
 			{
-				TArray<FValueRef> Arguments(TArrayView<FValue*>(AsInlineHLSL->Arguments, AsInlineHLSL->NumArguments));
-				return InlineHLSL(AsInlineHLSL->ExternalCodeDeclaration, Arguments, Axis == EDerivativeAxis::X ? EValueFlags::DerivativeDDX : EValueFlags::DerivativeDDY);
+				Errorf(TEXT("Extern '%s' does not support analytical partial derivatives."), Info.Name.GetData());
+				return Value.ToPoison();
 			}
-			return Cast(ConstantZero(Value->Type.GetPrimitive().ScalarKind), Value->Type);
+			FType DerivativeType = !Info.DifferentialType.IsPoison() ? Info.DifferentialType : Info.Type;
+			FExternDifferential DerivativeProto = MakePrototype<FExternDifferential>(DerivativeType);
+			DerivativeProto.ExternArg = AsExtern;
+			DerivativeProto.Axis = Axis;
+			return EmitPrototype(*this, DerivativeProto);
 		}
 
 		case VK_GetVertexInterpolator:
@@ -4156,12 +4114,6 @@ FValueRef FEmitter::AnalyticalPartialDerivative(FValueRef Value, EDerivativeAxis
 			return VertexInterpolator(AnalyticalPartialDerivative(AsVertexInterpolator->SetInstr->Arg, Axis));
 		}
 
-		// These values are uniform (constant), thus their value is always zero.
-		case VK_Constant:
-		case VK_UniformParameter:
-		case VK_PreshaderParameter:
-			return Cast(ConstantZero(Value->Type.GetPrimitive().ScalarKind), Value->Type);
-
 		// These values don't work with analytic derivatives, and force hardware derivatives (or zero if the shader model doesn't support compute shader derivatives).
 		case VK_TextureRead:
 		case VK_Call:
@@ -4181,6 +4133,13 @@ FValueRef FEmitter::AnalyticalPartialDerivative(FValueRef Value, EDerivativeAxis
 			}
 		}
 
+		// Finally, these values are uniform (constant) or their differential is not provided, thus their value is always zero.
+		case VK_Constant:
+		case VK_UniformParameter:
+		case VK_PreshaderParameter:
+		case VK_ExternDifferential:
+			return Cast(ConstantZero(Value->Type.GetPrimitive().ScalarKind), Value->Type);
+
 		default:
 			UE_MIR_UNREACHABLE();
 	}
@@ -4206,75 +4165,6 @@ FValueRef FEmitter::VertexInterpolator(FValueRef VertexValue)
 	return GetInstr;
 }
 
-static FValueRef EmitInlineHLSL(FEmitter& Emitter, FType Type, const FMaterialExternalCodeDeclaration* InExternalCodeDeclaration,
-								FStringView Code, TConstArrayView<FValueRef> InputValues, EValueFlags ValueFlags, EGraphProperties UsedGraphProperties,
-								FInlineHLSL::TOnAnalyzeCallback OnAnalyze = {}, uintptr_t OnAnalyzeUserData = {})
-{
-	FInlineHLSL Prototype = MakePrototype<FInlineHLSL>(Type);
-	Prototype.Type = Type;
-	Prototype.Flags = ValueFlags;
-	Prototype.GraphProperties = UsedGraphProperties;
-	Prototype.OnAnalyze = OnAnalyze;
-	Prototype.OnAnalyzeUserData = OnAnalyzeUserData;
-
-	if (InExternalCodeDeclaration)
-	{
-		check(Code.IsEmpty());
-		Prototype.ExternalCodeDeclaration = InExternalCodeDeclaration;
-	}
-	else
-	{
-		Prototype.Code = Code;
-	}
-
-	if (!InputValues.IsEmpty())
-	{
-		checkf(InputValues.Num() < FInlineHLSL::MaxNumArguments, TEXT("Number of arguments for inline-HLSL out of bounds: %d was specified, but upper bound is %d"), InputValues.Num(), FInlineHLSL::MaxNumArguments);
-		Prototype.NumArguments = InputValues.Num();
-		for (int32 i = 0; i < InputValues.Num(); ++i)
-		{
-			checkf(InputValues[i], TEXT("InputValues[%d] must not be null when InlineHLSL-instruction is emitted"), i);
-			Prototype.Arguments[i] = InputValues[i];
-		}
-	}
-
-	return EmitPrototype(Emitter, Prototype);
-}
-
-FValueRef FEmitter::InlineHLSL(FType Type, FStringView Code, TConstArrayView<FValueRef> InputValues, EValueFlags ValueFlags, EGraphProperties UsedGraphProperties, FInlineHLSL::TOnAnalyzeCallback OnAnalyze, uintptr_t OnAnalyzeUserData)
-{
-	if (IsAnyNotValid(InputValues))
-	{
-		return Poison();
-	}
-	
-	return EmitInlineHLSL(*this, Type, nullptr, Module->InternString(Code), InputValues, ValueFlags | EValueFlags::HasDynamicHLSLCode, UsedGraphProperties, OnAnalyze, OnAnalyzeUserData);
-}
-
-FValueRef FEmitter::InlineHLSL(FName InExternalCodeDeclarationName, TConstArrayView<FValueRef> InputValues, EValueFlags ValueFlags, EGraphProperties UsedGraphProperties)
-{
-	const FMaterialExternalCodeDeclaration* ExternalCodeDeclaration = MaterialExternalCodeRegistry::Get().FindExternalCode(InExternalCodeDeclarationName);
-	if (!ExternalCodeDeclaration)
-	{
-		Errorf(TEXT("Missing external code declaration for '%s'"), *InExternalCodeDeclarationName.ToString());
-		return Poison();
-	}
-
-	return InlineHLSL(ExternalCodeDeclaration, InputValues, ValueFlags, UsedGraphProperties);
-}
-
-FValueRef FEmitter::InlineHLSL(const FMaterialExternalCodeDeclaration* InExternalCodeDeclaration, TConstArrayView<FValueRef> InputValues, EValueFlags ValueFlags, EGraphProperties UsedGraphProperties)
-{
-	if (IsAnyNotValid(InputValues))
-	{
-		return Poison();
-	}
-
-	check(InExternalCodeDeclaration != nullptr);
-	FType ReturnType = FType::FromMaterialValueType( InExternalCodeDeclaration->GetReturnTypeValue());
-	return EmitInlineHLSL(*this, ReturnType, InExternalCodeDeclaration, {}, InputValues, ValueFlags, UsedGraphProperties);
-}
-
 static void SetFunctionFromDesc(FFunction& Function, FFunctionKind Kind, const FFunctionDesc& Desc)
 {
 	Function.Kind = Kind;
@@ -4283,7 +4173,7 @@ static void SetFunctionFromDesc(FFunction& Function, FFunctionKind Kind, const F
 	Function.NumInputAndOutputParams = Desc.NumInputOnlyParams + Desc.NumInputOutputParams;
 	Function.NumParameters = Function.NumInputAndOutputParams + Desc.NumOutputOnlyParams;
 	Function.ReturnType = Desc.ReturnType;
-
+ 
 	// Copy over the parameter declarations from the description to the function prototype
 	for (uint32 i = 0; i < Function.NumParameters; ++i)
 	{
@@ -4686,6 +4576,11 @@ void FEmitter::Initialize()
 	FalseConstant = EmitPrototype(*this, Temp);
 }
 
+FValueRef FEmitter::EmitPrototype_Internal(const FValue& Value)
+{
+	return EmitPrototype(*this, Value);
+}
+
 bool FEmitter::FValueKeyFuncs::Matches(KeyInitType A, KeyInitType B)
 {
 	return A->Equals(B);
diff --git a/Engine/Source/Runtime/Engine/Private/Materials/MaterialIRModuleBuilder.cpp b/Engine/Source/Runtime/Engine/Private/Materials/MaterialIRModuleBuilder.cpp
--- a/Engine/Source/Runtime/Engine/Private/Materials/MaterialIRModuleBuilder.cpp
+++ b/Engine/Source/Runtime/Engine/Private/Materials/MaterialIRModuleBuilder.cpp
@@ -1242,7 +1242,7 @@ struct FMaterialIRModuleBuilderImpl
 		// Run a final material compilation output validation and report any error reported.
 		TArray<FString> ValidationErrors;
 		UE::MaterialTranslatorUtils::FinalCompileValidation(
-			Builder->MaterialInterface->GetMaterial(),
+			Emitter.BaseMaterial,
 			CompilationOutput,
 			Module->GetCompiledShadingModels(),
 			Module->GetBlendMode(),
diff --git a/Engine/Source/Runtime/Engine/Private/Materials/MaterialIRToHLSLTranslator.cpp b/Engine/Source/Runtime/Engine/Private/Materials/MaterialIRToHLSLTranslator.cpp
--- a/Engine/Source/Runtime/Engine/Private/Materials/MaterialIRToHLSLTranslator.cpp
+++ b/Engine/Source/Runtime/Engine/Private/Materials/MaterialIRToHLSLTranslator.cpp
@@ -40,39 +40,39 @@ enum EBeginArgs { BeginArgs };
 enum EEndArgs { EndArgs };
 enum EListSeparator { ListSeparator };
 
-struct FHLSLPrinter
+struct FPrinter
 {
 	FString Buffer;
 	bool bFirstListItem = false;
 	int32 Tabs = 0;
 	TArray<const UMaterialAggregate*>& MaterialAggregatesRequiringDeclaration;
 	
-	FHLSLPrinter(TArray<const UMaterialAggregate*>& MaterialAggregatesRequiringDeclaration, int32 Tabs = 0)
+	FPrinter(TArray<const UMaterialAggregate*>& MaterialAggregatesRequiringDeclaration, int32 Tabs = 0)
 		: Tabs{ Tabs }
 		, MaterialAggregatesRequiringDeclaration { MaterialAggregatesRequiringDeclaration }
 	{
 		*this << Indentation;
 	}
 
-	FHLSLPrinter& operator<<(FStringView Text)
+	FPrinter& operator<<(FStringView Text)
 	{
 		Buffer.Append(Text);
 		return *this;
 	}
 
-	FHLSLPrinter& operator<<(int32 Value)
+	FPrinter& operator<<(int32 Value)
 	{
 		Buffer.Appendf(TEXT("%d"), Value);
 		return *this;
 	}
 
-	FHLSLPrinter& operator<<(uint32 Value)
+	FPrinter& operator<<(uint32 Value)
 	{
 		Buffer.Appendf(TEXT("%u"), Value);
 		return *this;
 	}
 
-	FHLSLPrinter& operator<<(float Value)
+	FPrinter& operator<<(float Value)
 	{
 		if (FGenericPlatformMath::IsNaN(Value))
 		{
@@ -89,19 +89,19 @@ struct FHLSLPrinter
 		return *this;
 	}
 
-    FHLSLPrinter& operator<<(ENoOp)
+    FPrinter& operator<<(ENoOp)
 	{
 		return *this;
 	}
 
-    FHLSLPrinter& operator<<(ENewLine)
+    FPrinter& operator<<(ENewLine)
     {
 		Buffer.AppendChar(TEXT('\n'));
 		operator<<(Indentation);
         return *this;
     }
 
-	FHLSLPrinter& operator<<(EIndentation)
+	FPrinter& operator<<(EIndentation)
 	{
 		for (int32 i = 0; i < Tabs; ++i)
 		{
@@ -110,7 +110,7 @@ struct FHLSLPrinter
 		return *this;
 	}
 
-	FHLSLPrinter& operator<<(EEndOfStatement)
+	FPrinter& operator<<(EEndOfStatement)
 	{
 		if (Buffer.IsEmpty() || Buffer[Buffer.Len() - 1] != TEXT('}'))
 		{
@@ -120,7 +120,7 @@ struct FHLSLPrinter
         return *this;
 	}
 
-    FHLSLPrinter& operator<<(EOpenBlock)
+    FPrinter& operator<<(EOpenBlock)
     {
 		*this << NewLine;
 		Buffer.AppendChar(TEXT('{'));
@@ -129,7 +129,7 @@ struct FHLSLPrinter
         return *this;
     }
 
-    FHLSLPrinter& operator<<(ECloseBlock)
+    FPrinter& operator<<(ECloseBlock)
     {
         --Tabs;
         Buffer.LeftChopInline(1); // undo tab
@@ -137,26 +137,26 @@ struct FHLSLPrinter
         return *this;
     }
 	
-    FHLSLPrinter& operator<<(EBeginArgs)
+    FPrinter& operator<<(EBeginArgs)
     {
         Buffer.AppendChar(TEXT('('));
 		BeginList();
         return *this;
     }
 
-    FHLSLPrinter& operator<<(EEndArgs)
+    FPrinter& operator<<(EEndArgs)
     {
         Buffer.AppendChar(TEXT(')'));
         return *this;
     }
 
-	FHLSLPrinter& operator<<(EListSeparator)
+	FPrinter& operator<<(EListSeparator)
     {
 		PrintListSeparator();
         return *this;
     }
 
-	FHLSLPrinter& operator<<(TCHAR Ch)
+	FPrinter& operator<<(TCHAR Ch)
     {
         Buffer.AppendChar(Ch);
         return *this;
@@ -177,7 +177,7 @@ struct FHLSLPrinter
 	}
 };
 
-static void PrintUserIdentifier(FHLSLPrinter& Printer, FStringView Identifier)
+static void PrintUserIdentifier(FPrinter& Printer, FStringView Identifier)
 {
 	if (!Identifier.IsEmpty() && FChar::IsDigit(Identifier[0]))
 	{
@@ -190,7 +190,7 @@ static void PrintUserIdentifier(FHLSLPrinter& Printer, FStringView Identifier)
 }
 
 // Prints the MIR type HLSL declaration to the specified printer.
-static FHLSLPrinter& operator<<(FHLSLPrinter& Printer, const MIR::FType& Type)
+static FPrinter& operator<<(FPrinter& Printer, const MIR::FType& Type)
 {
 	if (TOptional<MIR::FPrimitive> PrimitiveType = Type.AsPrimitive())
 	{
@@ -271,7 +271,7 @@ struct FEntryPointCodeGen
 	const FMaterialIRModule* Module{};
 
 	// Utility for emitting HLSL code
-	FHLSLPrinter Printer;
+	FPrinter Printer;
 
 	// Used internally. See FResult::OutputNormalAttributeCode
 	FString OutputNormalAttributeCode;
@@ -400,10 +400,8 @@ struct FEntryPointCodeGen
 		switch (Value->Kind)
 		{
 			case MIR::VK_Constant: LowerConstant(static_cast<const MIR::FConstant*>(Value)); break;
-			case MIR::VK_ExternalInput: LowerExternalInput(static_cast<const MIR::FExternalInput*>(Value)); break;
+			case MIR::VK_Builtin: LowerBuiltin(static_cast<const MIR::FBuiltin*>(Value)); break;
 			case MIR::VK_MaterialParameterCollection: LowerMaterialParameterCollection(static_cast<const MIR::FMaterialParameterCollection*>(Value)); break;
-			case MIR::VK_ScreenTexture: LowerScreenTexture(static_cast<const MIR::FScreenTexture*>(Value)); break;
-			case MIR::VK_ShadingModel: LowerShadingModel(static_cast<const MIR::FShadingModel*>(Value)); break;
 			case MIR::VK_TextureObject: LowerTextureObject(static_cast<const MIR::FTextureObject*>(Value)); break;
 			case MIR::VK_RuntimeVirtualTextureObject: LowerRuntimeVirtualTextureObject(static_cast<const MIR::FRuntimeVirtualTextureObject*>(Value)); break;
 			case MIR::VK_UniformParameter: LowerUniformParameter(static_cast<const MIR::FUniformParameter*>(Value)); break;
@@ -427,7 +425,8 @@ struct FEntryPointCodeGen
 			case MIR::VK_Scalar: LowerScalar(static_cast<const MIR::FScalar*>(Instr)); break;
 			case MIR::VK_TextureRead: LowerTextureRead(static_cast<const MIR::FTextureRead*>(Instr)); break;
 			case MIR::VK_VTPageTableRead: LowerVTPageTableRead(static_cast<const MIR::FVTPageTableRead*>(Instr)); break;
-			case MIR::VK_InlineHLSL: LowerInlineHLSL(static_cast<const MIR::FInlineHLSL*>(Instr)); break;
+			case MIR::VK_Extern: LowerExtern(static_cast<const MIR::FExtern*>(Instr)); break;
+			case MIR::VK_ExternDifferential: LowerExternDifferential(static_cast<const MIR::FExternDifferential*>(Instr)); break;
 			case MIR::VK_StageSwitch: LowerStageSwitch(static_cast<const MIR::FStageSwitch*>(Instr)); break;
 			case MIR::VK_HardwarePartialDerivative: LowerHardwarePartialDerivative(static_cast<const MIR::FHardwarePartialDerivative*>(Instr)); break;
 			case MIR::VK_SetVertexInterpolator: LowerSetVertexInterpolator(static_cast<const MIR::FSetVertexInterpolator*>(Instr)); break;
@@ -477,149 +476,18 @@ struct FEntryPointCodeGen
 		}
 	}
 
-	void LowerExternalInput(const MIR::FExternalInput* ExternalInput)
+	void LowerBuiltin(const MIR::FBuiltin* ExternalInput)
 	{
 		int32 ExternalInputIndex = (int32)ExternalInput->Id;
-		if (MIR::IsExternalInputTexCoord(ExternalInput->Id))
-		{
-			int32 Index = ExternalInputIndex - (int32)MIR::EExternalInput::TexCoord0;
-			Printer << TEXTVIEW("Parameters.TexCoords[") << Index << TEXTVIEW("]");
-		}
-		else if (MIR::IsExternalInputTexCoordDdx(ExternalInput->Id))
-		{
-			int32 Index = ExternalInputIndex - (int32)MIR::EExternalInput::TexCoord0_Ddx;
-			Printer << TEXTVIEW("Parameters.TexCoords_DDX[") << Index << TEXTVIEW("]");
-		}
-		else if (MIR::IsExternalInputTexCoordDdy(ExternalInput->Id))
-		{
-			int32 Index = ExternalInputIndex - (int32)MIR::EExternalInput::TexCoord0_Ddy;
-			Printer << TEXTVIEW("Parameters.TexCoords_DDY[") << Index << TEXTVIEW("]");
-		}
-		else if (ExternalInput->Id >= MIR::EExternalInput::WorldPosition_Absolute && ExternalInput->Id <= MIR::EExternalInput::WorldPosition_CameraRelativeNoOffsets)
-		{
-			// Various function permutations exist to fetch world position.  "Prev" permutations are only available in the vertex shader, while
-			// "MoMaterialOffsets" permutations are only available in the pixel shader, so we need to cobble together the permutation string
-			// factoring in those limitations.  Logic adapted from FHLSLMaterialTranslator::WorldPosition.
-			// 
-			// Format:  Get[Prev][Translated]WorldPosition[_NoMaterialOffsets](Parameters)
-			//
-			Printer << TEXTVIEW("Get");
-
-			// If compiling for the previous frame
-			if (bCompilingPreviousFrame)
-			{
-				Printer << TEXTVIEW("Prev");
-			}
-
-			if (ExternalInput->Id == MIR::EExternalInput::WorldPosition_CameraRelative || ExternalInput->Id == MIR::EExternalInput::WorldPosition_CameraRelativeNoOffsets)
-			{
-				Printer << TEXTVIEW("Translated");
-			}
-
-			Printer << TEXTVIEW("WorldPosition");
-
-			if (CurrentStage == MIR::EStage::Stage_Pixel && (ExternalInput->Id == MIR::EExternalInput::WorldPosition_AbsoluteNoOffsets || ExternalInput->Id == MIR::EExternalInput::WorldPosition_CameraRelativeNoOffsets))
-			{
-				Printer << TEXTVIEW("_NoMaterialOffsets");
-			}
-
-			Printer << TEXTVIEW("(Parameters)");
-		}
-		else if (ExternalInput->Id >= MIR::EExternalInput::WorldPosition_Absolute_Ddx && ExternalInput->Id <= MIR::EExternalInput::WorldPosition_CameraRelativeNoOffsets_Ddx)
-		{
-			// Expression emitter assumes these to be LWC for analytic derivative evaluation, so we promote them on load
-			Printer << TEXTVIEW("WSPromote(Parameters.WorldPosition_DDX)");
-		}
-		else if (ExternalInput->Id >= MIR::EExternalInput::WorldPosition_Absolute_Ddy && ExternalInput->Id <= MIR::EExternalInput::WorldPosition_CameraRelativeNoOffsets_Ddy)
-		{
-			// Expression emitter assumes these to be LWC for analytic derivative evaluation, so we promote them on load
-			Printer << TEXTVIEW("WSPromote(Parameters.WorldPosition_DDY)");
-		}
-		else if (ExternalInput->Id >= MIR::EExternalInput::LocalPosition_Instance && ExternalInput->Id <= MIR::EExternalInput::LocalPosition_PrimitiveNoOffsets)
-		{
-			// Various function permutations exist to fetch local position.  "Prev" permutations are only available in the vertex shader, while
-			// "MoMaterialOffsets" permutations are only available in the pixel shader, so we need to cobble together the permutation string
-			// factoring in those limitations.  Logic adapted from FHLSLMaterialTranslator::LocalPosition (and similar to WorldPosition above).
-			// 
-			// Format:  Get[Prev]Position[Instance|Primitive]Space[_NoMaterialOffsets](Parameters)
-			//
-			Printer << TEXTVIEW("Get");
-
-			// If compiling for the previous frame
-			if (bCompilingPreviousFrame)
-			{
-				Printer << TEXTVIEW("Prev");
-			}
-
-			Printer << TEXTVIEW("Position");
-			Printer << (ExternalInput->Id <= MIR::EExternalInput::LocalPosition_InstanceNoOffsets ? TEXTVIEW("Instance") : TEXTVIEW("Primitive"));
-			Printer << TEXTVIEW("Space");
-
-			if (CurrentStage == MIR::EStage::Stage_Pixel && (ExternalInput->Id == MIR::EExternalInput::LocalPosition_InstanceNoOffsets || ExternalInput->Id == MIR::EExternalInput::LocalPosition_PrimitiveNoOffsets))
-			{
-				Printer << TEXTVIEW("_NoMaterialOffsets");
-			}
-
-			Printer << TEXTVIEW("(Parameters)");
-		}
-		else if (ExternalInput->Id >= MIR::EExternalInput::LocalPosition_Instance_Ddx && ExternalInput->Id <= MIR::EExternalInput::LocalPosition_PrimitiveNoOffsets_Ddx)
-		{
-			// Local position derivatives reuse the world position derivatives, transforming them to local space
-			if (ExternalInput->Id <= MIR::EExternalInput::LocalPosition_InstanceNoOffsets_Ddx)
-			{
-				Printer << TEXTVIEW("mul(Parameters.WorldPosition_DDX, DFToFloat3x3(GetWorldToInstanceDF(Parameters)))");
-			}
-			else
-			{
-				Printer << (bCompilingPreviousFrame ?
-					TEXTVIEW("mul(Parameters.WorldPosition_DDX, DFToFloat3x3(GetPrevWorldToLocalDF(Parameters)))") :
-					TEXTVIEW("mul(Parameters.WorldPosition_DDX, DFToFloat3x3(GetWorldToLocalDF(Parameters)))"));
-			}
-		}
-		else if (ExternalInput->Id >= MIR::EExternalInput::LocalPosition_Instance_Ddy && ExternalInput->Id <= MIR::EExternalInput::LocalPosition_PrimitiveNoOffsets_Ddy)
-		{
-			// Local position derivatives reuse the world position derivatives, transforming them to local space
-			if (ExternalInput->Id <= MIR::EExternalInput::LocalPosition_InstanceNoOffsets_Ddy)
-			{
-				Printer << TEXTVIEW("mul(Parameters.WorldPosition_DDY, DFToFloat3x3(GetWorldToInstanceDF(Parameters)))");
-			}
-			else
-			{
-				Printer << (bCompilingPreviousFrame ?
-					TEXTVIEW("mul(Parameters.WorldPosition_DDY, DFToFloat3x3(GetPrevWorldToLocalDF(Parameters)))") :
-					TEXTVIEW("mul(Parameters.WorldPosition_DDY, DFToFloat3x3(GetWorldToLocalDF(Parameters)))"));
-			}
-		}
-		else if (ExternalInput->Id == MIR::EExternalInput::ActorPosition_Absolute)
-		{
-			Printer << (bCompilingPreviousFrame ? TEXTVIEW("GetPreviousActorWorldPosition(Parameters)") : TEXTVIEW("GetActorWorldPosition(Parameters)"));
-		}
-		else if (ExternalInput->Id == MIR::EExternalInput::ActorPosition_CameraRelative)
-		{
-			Printer << (bCompilingPreviousFrame ? TEXTVIEW("GetPreviousActorTranslatedWorldPosition(Parameters)") : TEXTVIEW("GetActorTranslatedWorldPosition(Parameters)"));
-		}
-		else if (ExternalInput->Id == MIR::EExternalInput::DynamicParticleParameterIndex)
-		{
-			Printer << ExternalInput->UserData;
-		}
-		else if (ExternalInput->Id == MIR::EExternalInput::CompilingPreviousFrame)
-		{
-			Printer << (bCompilingPreviousFrame ? TEXTVIEW("true") : TEXTVIEW("false"));
-		}
-		else
+	
+		FStringView Code;
+		switch (ExternalInput->Id)
 		{
-			FStringView Code;
-			switch (ExternalInput->Id)
-			{
-				case MIR::EExternalInput::ObjectPosition_Absolute: Code = TEXTVIEW("GetObjectWorldPosition(Parameters)"); break;
-				case MIR::EExternalInput::ObjectPosition_CameraRelative: Code = TEXTVIEW("GetObjectTranslatedWorldPosition(Parameters)"); break;
-				case MIR::EExternalInput::ViewMaterialTextureMipBias: Code = TEXTVIEW("View.MaterialTextureMipBias"); break;
-				case MIR::EExternalInput::ViewMaterialTextureDerivativeMultiply: Code = TEXTVIEW("View.MaterialTextureDerivativeMultiply"); break;
-				case MIR::EExternalInput::GlobalDistanceField: Code = TEXTVIEW(""); break;			// Used as flag for value analyzer, doesn't generate code
-				default: UE_MIR_UNREACHABLE();
-			}
-			Printer << Code;
+			case MIR::EBuiltin::TextureMipBias: Code = TEXTVIEW("View.MaterialTextureMipBias"); break;
+			case MIR::EBuiltin::TextureDerivativeMultiply: Code = TEXTVIEW("View.MaterialTextureDerivativeMultiply"); break;
+			default: UE_MIR_UNREACHABLE();
 		}
+		Printer << Code;
 	}
 
 	void LowerMaterialParameterCollection(const MIR::FMaterialParameterCollection* MaterialParameterCollection)
@@ -627,82 +495,6 @@ struct FEntryPointCodeGen
 		Printer << MaterialParameterCollection->Analysis_CollectionIndex;
 	}
 
-	void LowerInlineHLSLWithArgumentsInternal(FStringView Code, const TConstArrayView<MIR::FValue*>& InArguments)
-	{
-		// Substitute argument tokens with instruction arguments
-		int32 ArgumentTokenStart = 0, ArgumentTokenEnd = 0;
-
-		// Prints the pending substring of the code and moves the range forward.
-		auto FlushCodeSubstring = [this, &ArgumentTokenEnd, &Code](int32 EndIndex) -> void
-		{
-			if (EndIndex > ArgumentTokenEnd)
-			{
-				Printer << Code.Mid(ArgumentTokenEnd, EndIndex - ArgumentTokenEnd);
-				ArgumentTokenEnd = EndIndex;
-			}
-		};
-
-		auto SubstituteNextArgument = [this, &ArgumentTokenStart, &ArgumentTokenEnd, &Code, &InArguments]() -> void
-		{
-			// Scan digits for argument index
-			int32 ArgumentIndexValue = 0;
-			int32 NumDigits = 0;
-
-			while (ArgumentTokenStart + NumDigits < Code.Len())
-			{
-				TCHAR CodeCharacter = Code[ArgumentTokenStart + NumDigits];
-				if (!FChar::IsDigit(CodeCharacter))
-				{
-					break;
-				}
-				ArgumentIndexValue *= 10;
-				ArgumentIndexValue += (CodeCharacter - TEXT('0'));
-				++NumDigits;
-			}
-
-			checkf(NumDigits > 0, TEXT("Failed to scan integer in inline-HLSL after token '$':\n\"%.*s\""), Code.Len(), Code.GetData());
-
-			checkf(
-				ArgumentIndexValue < InArguments.Num(), TEXT("Failed to substitute token $%d in inline-HLSL with given number of arguments (%d):\n\"%.*s\""),
-				ArgumentIndexValue, InArguments.Num(), Code.Len(), Code.GetData()
-			);
-
-			LowerValue(InArguments[ArgumentIndexValue]);
-
-			ArgumentTokenEnd = ArgumentTokenStart + NumDigits;
-		};
-
-		auto MatchCharacter = [&Code](int32& Position, TCHAR InCharacter) -> bool
-		{
-			if (Position < Code.Len() && Code[Position] == InCharacter)
-			{
-				++Position;
-				return true;
-			}
-			return false;
-		};
-
-		// Find all argument token characters '$'. For example "MyFunction($1, $0.xxxw)" can be subsituted with "MyFunction(MySecondArgument, MyFirstArgument.xxxw)"
-		while ((ArgumentTokenStart = Code.Find(TEXT("$"), ArgumentTokenEnd)) != INDEX_NONE)
-		{
-			FlushCodeSubstring(ArgumentTokenStart);
-
-			++ArgumentTokenStart;
-			if (MatchCharacter(ArgumentTokenStart, TEXT('{')))
-			{
-				SubstituteNextArgument();
-				const bool bMatchedArgumentToken = MatchCharacter(ArgumentTokenEnd, TEXT('}'));
-				checkf(bMatchedArgumentToken, TEXT("Failed to match argument token in inline-HLSL with syntax '${N}':\n\"%.*s\""), Code.Len(), Code.GetData());
-			}
-			else
-			{
-				SubstituteNextArgument();
-			}
-		}
-
-		FlushCodeSubstring(Code.Len());
-	}
-
 	void LowerGetVertexInterpolator(const MIR::FGetVertexInterpolator* GetVertexInterpolator)
 	{
 		if (GetVertexInterpolator->Type.GetPrimitive().NumComponents() > 1)
@@ -723,70 +515,45 @@ struct FEntryPointCodeGen
 		}
 	}
 
-	void LowerInlineHLSL(const MIR::FInlineHLSL* ExternalCode)
+	MIR::FExternPrinterHLSL MakeExternPrinter(MIR::EExternDifferential Differential, MIR::FExternArgs Arguments)
 	{
-		FStringView FinalCode;
-		if (ExternalCode->HasFlags(MIR::EValueFlags::HasDynamicHLSLCode))
-		{
-			FinalCode = ExternalCode->Code;
-		}
-		else if (ExternalCode->HasFlags(MIR::EValueFlags::DerivativeDDX))
-		{
-			FinalCode = ExternalCode->ExternalCodeDeclaration->DefinitionDDX;
-		}
-		else if (ExternalCode->HasFlags(MIR::EValueFlags::DerivativeDDY))
-		{
-			FinalCode = ExternalCode->ExternalCodeDeclaration->DefinitionDDY;
-		}
-		else
-		{
-			FinalCode = ExternalCode->ExternalCodeDeclaration->Definition;
-		}
-
-		// Substitute placeholder tokens now unless disabled for custom nodes, if necessary.
-		FString SubstitutedCode;
-		if (ExternalCode->HasFlags(MIR::EValueFlags::SubstituteTags))
-		{
-			SubstitutedCode = FString{ FinalCode };
-			SubstitutedCode.ReplaceInline(TEXT("<PREV>"), bCompilingPreviousFrame ? TEXT("Prev") : TEXT(""));
-			SubstitutedCode.ReplaceInline(TEXT("<PREVIOUS>"), bCompilingPreviousFrame ? TEXT("Previous") : TEXT(""));
-			SubstitutedCode.ReplaceInline(TEXT("<PREVFRAME>"), bCompilingPreviousFrame ? TEXT("PrevFrame") : TEXT(""));
-
-			FinalCode = SubstitutedCode;
-		}
+		return {
+			.Buffer             = Printer.Buffer,
+			.Differential       = Differential,
+			.Stage              = CurrentStage,
+			.bIsPreviousFrame   = bCompilingPreviousFrame,
+			.Arguments          = Arguments,
+			.LowerValueFunction = +[](const MIR::FValue* Value, void* UserData)
+			{
+				static_cast<FEntryPointCodeGen*>(UserData)->LowerValue(Value);
+			},
+			.LowerValueUserData = this,
+		};
+	}
 
-		// Print the final HLSL code.
-		if (ExternalCode->NumArguments > 0)
-		{
-			check(ExternalCode->Arguments);
-			LowerInlineHLSLWithArgumentsInternal(FinalCode, TConstArrayView<MIR::FValue*>(ExternalCode->Arguments, ExternalCode->NumArguments));
-		}
-		else
-		{
-			Printer << FinalCode;
-		}
+	void LowerExtern(const MIR::FExtern* Extern)
+	{
+		MIR::FExternPrinterHLSL ExternPrinter = MakeExternPrinter(MIR::EExternDifferential::None, Extern->GetArguments());
+		Extern->ToHLSL(ExternPrinter);
 	}
 
-	void LowerScreenTexture(const MIR::FScreenTexture* ScreenTexture)
+	void LowerExternDifferential(const MIR::FExternDifferential* ExternDifferential)
 	{
-		// Types besides the following aren't directly printed in HLSL, and don't need to do anything here
-		switch (ScreenTexture->TextureKind)
+		const MIR::FExtern* ExternArg = ExternDifferential->GetExtern();
+		check(ExternArg); // FExternDifferential argument must always be an extern
+
+		// If extern has zero derivatives generate them directly
+		if ((ExternArg->GetInfo().Flags & MIR::EExternFlags::ZeroDifferentials) != MIR::EExternFlags::None)
 		{
-			case MIR::EScreenTexture::SceneTexture:
-				Printer << UE::MaterialTranslatorUtils::SceneTextureIdToHLSLString(ScreenTexture->Id);
-				break;
-			case MIR::EScreenTexture::UserSceneTexture:
-				Printer << UE::MaterialTranslatorUtils::SceneTextureIdToHLSLString((ESceneTextureId)Module->GetCompilationOutput().FindUserSceneTexture(ScreenTexture->UserSceneTexture));
-				break;
-			case MIR::EScreenTexture::DBufferTexture:
-				Printer << (int32)ScreenTexture->DBufferId;
-				break;
+			check(ExternDifferential->Type.IsPrimitive());
+			Printer << ExternDifferential->Type << TEXT("(0)");
+			return;
 		}
-	}
 
-	void LowerShadingModel(const MIR::FShadingModel* ShadingModel)
-	{
-		Printer << (int32)ShadingModel->Id;
+		// Otherwise invoke the extern to generate the differential code by itself
+		const MIR::EExternDifferential Differential = ExternDifferential->Axis == MIR::EDerivativeAxis::X ? MIR::EExternDifferential::ScreenSpaceX : MIR::EExternDifferential::ScreenSpaceY;
+		MIR::FExternPrinterHLSL ExternPrinter = MakeExternPrinter(Differential, ExternArg->GetArguments());
+		ExternArg->ToHLSL(ExternPrinter);
 	}
 
 	void LowerTextureObject(const MIR::FTextureObject* TextureObject)
@@ -2019,7 +1786,7 @@ struct FEntryPointCodeGen
 	
 	void LowerSubstrateDefaultSlab(const MIR::FSubstrateDefaultSlab* SubstrateParameter)
 	{
-		Printer << TEXTVIEW("GetInitialisedSubstrateData();");
+		Printer << TEXTVIEW("GetInitialisedSubstrateData()");
 	}
 
 	static FStringView GetSharedLocalBasesTypes(bool bFullySimplifiedInstruction)
@@ -2319,6 +2086,11 @@ struct FEntryPointCodeGen
 			case MIR::VK_SubstrateSlab:
 			case MIR::VK_SubstrateShadingModels:
 				return false;
+			case MIR::VK_Extern:
+				return (static_cast<const MIR::FExtern*>(Instr)->GetInfo().Flags & MIR::EExternFlags::Inline) != MIR::EExternFlags::None;
+			case MIR::VK_ExternDifferential:
+				return IsFoldable(static_cast<const MIR::FExternDifferential*>(Instr)->GetExtern(), EntryPointIndex);
+
 			default:
 				return true;
 		}
@@ -2358,8 +2130,8 @@ static void GenerateCustomOutputDeclarations(const FMaterial* Material, const FM
 {
 	// Generate the custom output fields of the struct FMaterialPixelParameters. These fields will be computed and 
 	// set by the function that calculates all material outputs (e.g. BaseColor).
-	FHLSLPrinter VertexPrinter{ MaterialAggregatesRequiringDeclaration, 1 };
-	FHLSLPrinter PixelPrinter{ MaterialAggregatesRequiringDeclaration, 1 };
+	FPrinter VertexPrinter{ MaterialAggregatesRequiringDeclaration, 1 };
+	FPrinter PixelPrinter{ MaterialAggregatesRequiringDeclaration, 1 };
 
 	const int NumCustomizedUVs = Material->GetNumCustomizedUVs();
 	if (NumCustomizedUVs > 0)
@@ -2369,7 +2141,7 @@ static void GenerateCustomOutputDeclarations(const FMaterial* Material, const FM
 
 	for (const MIR::FSetMaterialOutput* Output : Module->GetCustomOutputs())
 	{
-		FHLSLPrinter& ThePrinter = Output->Frequency == MIR::EMaterialOutputFrequency::PerVertex ? VertexPrinter : PixelPrinter;
+		FPrinter& ThePrinter = Output->Frequency == MIR::EMaterialOutputFrequency::PerVertex ? VertexPrinter : PixelPrinter;
 		ThePrinter << Output->Arg->Type << TEXT(' ') << *Output->Name << EndOfStatement;
 	}
 		
@@ -2381,7 +2153,7 @@ static void GenerateCustomOutputDeclarations(const FMaterial* Material, const FM
 	// Development Note: this is really just cruft necessary for backwards compatibility with the way the
 	// 	  old translator supported custom outputs and could be removed entirely with a little refactoring
 	//    of MaterialTemplate.ush.
-	FHLSLPrinter Printer = { MaterialAggregatesRequiringDeclaration };
+	FPrinter Printer = { MaterialAggregatesRequiringDeclaration };
 	for (const FMaterialIRModule::FCustomOutputGroup& Group : Module->GetCustomOutputGroups())
 	{
 		Printer << TEXTVIEW("#define NUM_MATERIAL_OUTPUTS_GET") << *FString(Group.Name).ToUpper() << TEXT(' ') << Group.NumOutputs << NewLine;
@@ -2427,7 +2199,7 @@ static void GenerateCustomOutputDeclarations(const FMaterial* Material, const FM
 // into the "custom_functions" source template parameter.
 static void GenerateCustomFunctionsHLSL(const FMaterialIRModule* Module, TMap<FString, FString>& OutParameters, TArray<const UMaterialAggregate*>& MaterialAggregatesRequiringDeclaration)
 {
-	FHLSLPrinter Printer = { MaterialAggregatesRequiringDeclaration };
+	FPrinter Printer = { MaterialAggregatesRequiringDeclaration };
 
 	for (const MIR::FFunctionHLSL* Function : Module->GetFunctionHLSLs())
 	{
@@ -2494,7 +2266,7 @@ static void GenerateOtherTemplateStringParameters(const FMaterial* Material, con
 
 	// pixel_material_inputs (field declarations of the FPixelMaterialInputs struct)
 	{
-		FHLSLPrinter Printer { MaterialAggregatesRequiringDeclaration, 1 };
+		FPrinter Printer { MaterialAggregatesRequiringDeclaration, 1 };
 		for (EMaterialProperty Property : MaterialAttributesAggregate::GetMaterialProperties())
 		{
 			if (Property == MP_WorldPositionOffset || (Property >= MP_CustomizedUVs0 && Property <= MP_CustomizedUVs7)) // todo: fixme
@@ -2518,7 +2290,7 @@ static void GenerateOtherTemplateStringParameters(const FMaterial* Material, con
 
 	// get_material_customized_u_vs
 	{
-		FHLSLPrinter Printer { MaterialAggregatesRequiringDeclaration, 1 };
+		FPrinter Printer { MaterialAggregatesRequiringDeclaration, 1 };
 		for (int32 i = 0; i < ShaderTexCoordsArrayLength; i++)
 		{
 			const bool bIsCustomizedUV = i < MIR::MaxNumCustomizedUVs && Module->IsMaterialPropertyUsed((EMaterialProperty)(MP_CustomizedUVs0 + i));
@@ -2586,7 +2358,7 @@ static const TCHAR* GetShadingModelParameterName(EMaterialShadingModel InModel)
 
 static void GenerateMaterialAggregatesDeclarations(TMap<FString, FString>& OutParameters, TArray<const UMaterialAggregate*>& MaterialAggregatesRequiringDeclaration)
 {
-	FHLSLPrinter Printer { MaterialAggregatesRequiringDeclaration };
+	FPrinter Printer { MaterialAggregatesRequiringDeclaration };
 
 	// Recursively push material aggregates referenced by attributes of existing aggregates
 	for (int32 i = 0; i < MaterialAggregatesRequiringDeclaration.Num(); ++i)
diff --git a/Engine/Source/Runtime/Engine/Private/Materials/MaterialIRValueAnalyzer.cpp b/Engine/Source/Runtime/Engine/Private/Materials/MaterialIRValueAnalyzer.cpp
--- a/Engine/Source/Runtime/Engine/Private/Materials/MaterialIRValueAnalyzer.cpp
+++ b/Engine/Source/Runtime/Engine/Private/Materials/MaterialIRValueAnalyzer.cpp
@@ -23,7 +23,6 @@
 #include "Engine/Texture2D.h"
 #include "VT/RuntimeVirtualTexture.h"
 #include "UObject/Package.h"
-#include "PostProcess/PostProcessMaterialInputs.h"
 #include "Shader/Preshader2.h"
 
 static void AnalyzeAnalysisError(FMaterialIRValueAnalyzer& Analyzer, MIR::FAnalysisError* AnalysisError)
@@ -42,48 +41,6 @@ static void ErrorUnlessFeatureLevelSupported(FMaterialIRModule* Module, ERHIFeat
 	}
 }
 
-static void AnalyzeExternalInput(FMaterialIRValueAnalyzer& Analyzer, MIR::FExternalInput* ExternalInput)
-{
-	FMaterialIRModule::FStatistics& Statistics = Analyzer.Module->GetStatistics();
-
-	if (Analyzer.MaterialInterface->IsPostProcessMaterial() && MIR::IsExternalInputWorldPosition(ExternalInput->Id))
-	{
-		Analyzer.CompilationOutput->SetIsSceneTextureUsed(PPI_SceneDepth);
-	}
-	else if (ExternalInput->Id == MIR::EExternalInput::GlobalDistanceField)
-	{
-		if (!FDataDrivenShaderPlatformInfo::GetSupportsDistanceFields(Analyzer.Module->GetShaderPlatform()))
-		{
-			FString ShaderPlatformName = FDataDrivenShaderPlatformInfo::GetName(Analyzer.Module->GetShaderPlatform()).ToString();
-			Analyzer.Module->AddError(nullptr, FString::Printf(TEXT("Node not supported in shader platform %s. The node requires DistanceField support."), *ShaderPlatformName));
-		}
-
-		Analyzer.CompilationOutput->bUsesGlobalDistanceField = true;
-	}
-	else if (ExternalInput->Id == MIR::EExternalInput::DynamicParticleParameterIndex)
-	{
-		Analyzer.Module->GetStatistics().DynamicParticleParameterMask |= 1 << ExternalInput->UserData;
-	}
-	else if (ExternalInput->Id == MIR::EExternalInput::DistanceCullFade)
-	{
-		Analyzer.CompilationOutput->bUsesDistanceCullFade = true;
-	}
-}
-
-static void AnalyzeInStageExternalInput(FMaterialIRValueAnalyzer& Analyzer, MIR::FExternalInput* ExternalInput, MIR::EStage Stage)
-{
-	FMaterialIRModule::FStatistics& Statistics = Analyzer.Module->GetStatistics();
-
-	if (MIR::IsExternalInputTexCoordOrPartialDerivative(ExternalInput->Id) && Stage != MIR::EStage::Stage_Vertex)
-	{
-		// @massimo.tristano this follows the logic in the old translator. A better way would instead be to **allocate** an index
-		// to a texture coordinate based on usage and map that index here. 
-		int32 TexCoordIndex = MIR::ExternalInputToTexCoordIndex(ExternalInput->Id);
-		Statistics.NumInterpolatedTexCoords = FMath::Max(Statistics.NumInterpolatedTexCoords, TexCoordIndex + 1);
-		Statistics.InterpolatedTexCoordsMask |= 1 << TexCoordIndex;
-	}
-}
-
 static void AnalyzeInStageHardwarePartialDerivative(FMaterialIRValueAnalyzer& Analyzer, MIR::FHardwarePartialDerivative* PartialDerivative, MIR::EStage Stage)
 {
 	if (Stage == MIR::Stage_Vertex)
@@ -110,55 +67,24 @@ static EMaterialShaderFrequency MapToMaterialShaderFrequencyOrAny(MIR::EStage St
 	return EMaterialShaderFrequency::Any;
 }
 
-static void AnalyzeInStageInlineHLSL(FMaterialIRValueAnalyzer& Analyzer, MIR::FInlineHLSL* InlineHLSL, MIR::EStage Stage)
+static void AnalyzeExtern(FMaterialIRValueAnalyzer& Analyzer, MIR::FExtern* Extern)
 {
-	if (!InlineHLSL->HasFlags(MIR::EValueFlags::HasDynamicHLSLCode))
-	{
-		check(InlineHLSL->ExternalCodeDeclaration != nullptr);
-		const EMaterialShaderFrequency StageToMaterialFrequency = MapToMaterialShaderFrequencyOrAny(Stage);
-		for (const FMaterialExternalCodeEnvironmentDefine& EnvironmentDefine : InlineHLSL->ExternalCodeDeclaration->EnvironmentDefines)
-		{
-			if ((int32)(EnvironmentDefine.ShaderFrequency & StageToMaterialFrequency) != 0)
-			{
-				Analyzer.EnvironmentDefines.Add(EnvironmentDefine.Name);
-			}
-		}
-	}
+	MIR::FExternAnalysisContext Context = {
+		.MaterialInterface = Analyzer.MaterialInterface,
+		.Module = Analyzer.Module,
+		.Extern_Internal = Extern,
+	};  
+	Extern->Analyze(Context);
 }
 
-static void AnalyzeInlineHLSL(FMaterialIRValueAnalyzer& Analyzer, MIR::FInlineHLSL* InlineHLSL)
+static void AnalyzeInStageExtern(FMaterialIRValueAnalyzer& Analyzer, MIR::FExtern* Extern, MIR::EStage Stage)
 {
-	if (!InlineHLSL->HasFlags(MIR::EValueFlags::HasDynamicHLSLCode))
-	{
-		check(InlineHLSL->ExternalCodeDeclaration != nullptr);
-
-		// Validate this external code can be used for the current material domain. Empty list implies no restriction on material domains.
-		if (!InlineHLSL->ExternalCodeDeclaration->Domains.IsEmpty() &&
-			InlineHLSL->ExternalCodeDeclaration->Domains.Find(Analyzer.MaterialInterface->GetMaterial()->MaterialDomain) == INDEX_NONE)
-		{
-			FName AssetPathName = Analyzer.MaterialInterface->GetOutermost()->GetFName();
-			Analyzer.Module->AddError(nullptr, MaterialExpressionUtils::FormatUnsupportedMaterialDomainError(*InlineHLSL->ExternalCodeDeclaration, AssetPathName));
-		}
-
-		// Cast from material feature level enum to RHI feature level enum
-		ERHIFeatureLevel::Type MinimumFeatureLevel = (ERHIFeatureLevel::Type)InlineHLSL->ExternalCodeDeclaration->MinimumFeatureLevel;
-
-		if (Analyzer.Module->GetFeatureLevel() < (ERHIFeatureLevel::Type)InlineHLSL->ExternalCodeDeclaration->MinimumFeatureLevel)
-		{
-			FString FeatureLevelName;
-			FString MinimumFeatureLevelName;
-			GetFeatureLevelName(Analyzer.Module->GetFeatureLevel(), FeatureLevelName);
-			GetFeatureLevelName(MinimumFeatureLevel, MinimumFeatureLevelName);
-
-			Analyzer.Module->AddError(nullptr, FString::Printf(TEXT("Node %s requires feature level %s.  Current feature level is %s."),
-				*InlineHLSL->ExternalCodeDeclaration->Name.ToString(), *MinimumFeatureLevelName, *FeatureLevelName));
-		}
-	}
-
-	if (InlineHLSL->OnAnalyze)
-	{
-		InlineHLSL->OnAnalyze(Analyzer.Module, InlineHLSL->OnAnalyzeUserData);
-	}
+	MIR::FExternAnalysisContext Context = {
+		.MaterialInterface = Analyzer.MaterialInterface,
+		.Module = Analyzer.Module,
+		.Extern_Internal = Extern,
+	};  
+	Extern->AnalyzeInStage(Context, Stage);
 }
 
 static int32 AcquireVTStackIndex(FMaterialIRValueAnalyzer& Analyzer, MIR::FVTPageTableRead& VTPageTableRead, bool bGenerateFeedback)
@@ -317,292 +243,6 @@ static void AnalyzeInStageVTPageTableRead(FMaterialIRValueAnalyzer& Analyzer, MI
 	}
 }
 
-static void AnalyzeScreenTexture(FMaterialIRValueAnalyzer& Analyzer, MIR::FScreenTexture* ScreenTexture)
-{
-	const UMaterial* Material = Analyzer.MaterialInterface->GetMaterial();
-	const EMaterialDomain MaterialDomain = Material->MaterialDomain;
-
-	switch (ScreenTexture->TextureKind)
-	{
-		case MIR::EScreenTexture::SceneTexture:
-		case MIR::EScreenTexture::UserSceneTexture:
-			{
-				// TODO:  If referenced from Custom HLSL, this can be false.  Revisit when custom HLSL support is added to new translator.
-				const bool bTextureLookup = true;
-
-				// Add defines and compilation outputs
-				Analyzer.EnvironmentDefines.Add(FName("NEEDS_SCENE_TEXTURES"));
-
-				Analyzer.CompilationOutput->bNeedsSceneTextures = true;
-				Analyzer.CompilationOutput->SetIsSceneTextureUsed(ScreenTexture->Id);
-
-				if (ScreenTexture->TextureKind == MIR::EScreenTexture::UserSceneTexture)
-				{
-					if (ScreenTexture->UserSceneTexture.IsNone())
-					{
-						Analyzer.Module->AddError(nullptr, TEXT("UserSceneTexture missing name -- value must be set to something other than None"));
-					}
-					else
-					{
-						// Allocate value during analyze, and check for failure.  FindUserSceneTexture is later called during HLSL generation to retrieve the allocated Id.
-						int32 SceneTextureId = Analyzer.CompilationOutput->FindOrAddUserSceneTexture(ScreenTexture->UserSceneTexture);
-						if (SceneTextureId == INDEX_NONE)
-						{
-							Analyzer.Module->AddError(nullptr, FString::Printf(TEXT("Too many unique UserSceneTexture inputs in the post process material -- max allowed is %d"), kPostProcessMaterialInputCountMax));
-						}
-					}
-				}
-
-				// Substrate TODO:
-				// Substrate only
-				// When SceneTexture lookup node is used, single/simple paths are disabled to ensure texture decoding is properly handled.
-				// Reading SceneTexture, when Substrate is enabled, implies unpacking material buffer data. The unpacking function exists in different 'flavor' 
-				// for optimization purpose (simple/single/complex). To avoid compiling out single or complex unpacking paths (due to defines set by analyzing 
-				// the current shader, vs. scene texture pixels), we force Simple/Single versions to be disabled
-				// FSubstrateCompilationContext& SubstrateCtx = SubstrateCompilationContext[CurrentSubstrateCompilationContext];
-				// SubstrateCtx.SubstrateMaterialComplexity.bIsSimple = false;
-				// SubstrateCtx.SubstrateMaterialComplexity.bIsSingle = false;
-
-				// Error checking
-				EShaderPlatform Platform = Analyzer.Module->GetShaderPlatform();
-
-				// Error checking from FHLSLMaterialTranslator::SceneTextureLookup
-				// Guard against using unsupported textures with SLW
-				const bool bHasSingleLayerWaterSM = Analyzer.MaterialInterface->GetShadingModels().HasShadingModel(MSM_SingleLayerWater);
-				if (bHasSingleLayerWaterSM && ScreenTexture->Id != PPI_CustomDepth && ScreenTexture->Id != PPI_CustomStencil)
-				{
-					Analyzer.Module->AddError(nullptr, TEXT("Only custom depth and custom stencil can be sampled with SceneTexture when used with the Single Layer Water shading model."));
-				}
-
-				if (ScreenTexture->Id == PPI_DecalMask)
-				{
-					Analyzer.Module->AddError(nullptr, TEXT("Decal Mask bit was moved from GBuffer to the Stencil Buffer for performance optimisation so therefore no longer available."));
-				}
-
-				// Error checking from FHLSLMaterialTranslator::UseSceneTextureId
-				if (MaterialDomain == MD_DeferredDecal)
-				{
-					const bool bSceneTextureSupportsDecal = ScreenTexture->Id == PPI_SceneDepth || ScreenTexture->Id == PPI_WorldNormal || ScreenTexture->Id == PPI_CustomDepth || ScreenTexture->Id == PPI_CustomStencil;
-					if (!bSceneTextureSupportsDecal)
-					{
-						// Note: For DBuffer decals CustomDepth and CustomStencil are not available if r.CustomDepth.Order = 1
-						Analyzer.Module->AddError(nullptr, TEXT("Decals can only access SceneDepth, CustomDepth, CustomStencil, and WorldNormal."));
-					}
-
-					const bool bSceneTextureRequiresSM5 = ScreenTexture->Id == PPI_WorldNormal;
-					if (bSceneTextureRequiresSM5)
-					{
-						ErrorUnlessFeatureLevelSupported(Analyzer.Module, ERHIFeatureLevel::SM5, TEXT("Deferred decals require SM5 for World Normal access."));
-					}
-
-					if (ScreenTexture->Id == PPI_WorldNormal && !IsUsingDBuffers(Platform))
-					{
-						// @massimo.tristano this should not be here, as this is not the correct way to check if an attribute has a user-defined value in the new translator
-						bool bHasNormalConnected =
-							Substrate::IsSubstrateEnabled() ?
-							Material->IsPropertyConnected(MP_Normal) :
-							Material->HasNormalConnected();
-
-						if (bHasNormalConnected)
-						{
-							// GBuffer decals can't bind Normal for read and write.
-							// Note: DBuffer decals can support this but only if the sampled WorldNormal isn't connected to the output normal.
-							Analyzer.Module->AddError(nullptr, TEXT("Decals that read WorldNormal cannot output to normal at the same time. Enable DBuffer to support this."));
-						}
-					}
-				}
-
-				if (ScreenTexture->Id == PPI_SceneColor && MaterialDomain != MD_Surface)
-				{
-					if (MaterialDomain == MD_PostProcess)
-					{
-						Analyzer.Module->AddError(nullptr, TEXT("SceneColor lookups are only available when MaterialDomain = Surface. PostProcessMaterials should use the SceneTexture PostProcessInput0."));
-					}
-					else
-					{
-						Analyzer.Module->AddError(nullptr, TEXT("SceneColor lookups are only available when MaterialDomain = Surface."));
-					}
-				}
-
-				bool bNeedsSceneTexturePostProcessInputs = false;
-				if (bTextureLookup)
-				{
-					bNeedsSceneTexturePostProcessInputs =
-						((ScreenTexture->Id >= PPI_PostProcessInput0 && ScreenTexture->Id <= PPI_PostProcessInput6) ||
-						 (ScreenTexture->Id >= PPI_UserSceneTexture0 && ScreenTexture->Id <= PPI_UserSceneTexture6) ||
-						 ScreenTexture->Id == PPI_Velocity ||
-						 ScreenTexture->Id == PPI_SceneColor);
-				}
-
-				if ((ScreenTexture->TextureKind == MIR::EScreenTexture::SceneTexture) && ((1u << ScreenTexture->Id) & FMaterialCompilationOutput::GetGBufferMask()))
-				{
-					if (IsForwardShadingEnabled(Platform) || (IsMobilePlatform(Platform) && !IsMobileDeferredShadingEnabled(Platform)))
-					{
-						Analyzer.Module->AddError(nullptr, FString::Printf(TEXT("GBuffer scene textures not available with forward shading (platform id %d)."), Platform));
-					}
-
-					// Post-process can't access memoryless GBuffer on mobile
-					if (IsMobilePlatform(Platform))
-					{
-						if (MaterialDomain == MD_PostProcess)
-						{
-							Analyzer.Module->AddError(nullptr, FString::Printf(TEXT("GBuffer scene textures not available in post-processing with mobile shading (platform id %d)."), Platform));
-						}
-
-						if (Material->IsMobileSeparateTranslucencyEnabled())
-						{
-							Analyzer.Module->AddError(nullptr, FString::Printf(TEXT("GBuffer scene textures not available for separate translucency with mobile shading (platform id %d)."), Platform));
-						}
-					}
-				}
-
-				if (ScreenTexture->Id == PPI_Velocity)
-				{
-					if (MaterialDomain != MD_PostProcess)
-					{
-						Analyzer.Module->AddError(nullptr, TEXT("Velocity scene textures are only available in post process materials."));
-					}
-				}
-
-				// Error checking from FHLSLMaterialTranslator::TranslateMaterial
-				if (MaterialDomain != MD_DeferredDecal && MaterialDomain != MD_PostProcess)
-				{
-					if (!Analyzer.MaterialInterface->GetShadingModels().HasShadingModel(MSM_SingleLayerWater) && IsOpaqueOrMaskedBlendMode(Material->BlendMode))
-					{
-						// In opaque pass, none of the textures are available
-						Analyzer.Module->AddError(nullptr, TEXT("SceneTexture expressions cannot be used in opaque materials except if used with the Single Layer Water shading model."));
-					}
-					else if (bNeedsSceneTexturePostProcessInputs)
-					{
-						Analyzer.Module->AddError(nullptr, TEXT("SceneTexture expressions cannot use post process inputs or scene color in non post process domain materials"));
-					}
-				}
-
-				if (ScreenTexture->Id == PPI_SceneDepth && bTextureLookup)
-				{
-					// Don't allow opaque and masked materials to access scene depth as the results are undefined
-					if (MaterialDomain != MD_PostProcess && !IsTranslucentBlendMode(Material->BlendMode))
-					{
-						Analyzer.Module->AddError(nullptr, TEXT("Only transparent or postprocess materials can read from scene depth."));
-					}
-				}
-			}
-			break;
-		case MIR::EScreenTexture::SceneColor:
-			{
-				// Add defines and compilation outputs
-				Analyzer.CompilationOutput->SetIsSceneTextureUsed(PPI_SceneColor);
-
-				// Error checking
-				if (MaterialDomain != MD_Surface)
-				{
-					Analyzer.Module->AddError(nullptr, TEXT("SceneColor lookups are only available when MaterialDomain = Surface."));
-				}
-
-				ErrorUnlessFeatureLevelSupported(Analyzer.Module, ERHIFeatureLevel::SM5, TEXT("Scene Color access require SM5."));
-			}
-			break;
-		case MIR::EScreenTexture::SceneDepth:
-			{
-				// Add defines and compilation outputs
-				Analyzer.CompilationOutput->SetIsSceneTextureUsed(PPI_SceneDepth);
-
-				// Error checking
-				if (Analyzer.MaterialInterface->IsTranslucencyWritingVelocity())
-				{
-					Analyzer.Module->AddError(nullptr, TEXT("Translucenct material with 'Output Velocity' enabled will write to depth buffer, therefore cannot read from depth buffer at the same time."));
-				}
-
-				// Don't allow opaque and masked materials to access scene depth as the results are undefined
-				if (MaterialDomain != MD_PostProcess && !IsTranslucentBlendMode(Material->BlendMode))
-				{
-					Analyzer.Module->AddError(nullptr, TEXT("Only transparent or postprocess materials can read from scene depth."));
-				}
-			}
-			break;
-		case MIR::EScreenTexture::SceneDepthWithoutWater:
-			{
-				// Error checking (no defines or compilation outputs needed for SceneDepthWithoutWater)
-				if (MaterialDomain != MD_PostProcess)
-				{
-					if (!Analyzer.MaterialInterface->GetShadingModels().HasShadingModel(MSM_SingleLayerWater))
-					{
-						Analyzer.Module->AddError(nullptr, TEXT("Can only read scene depth below water when material Shading Model is Single Layer Water or when material Domain is PostProcess."));
-					}
-
-					if (MaterialDomain != MD_Surface)
-					{
-						Analyzer.Module->AddError(nullptr, TEXT("Can only read scene depth below water when material Domain is set to Surface or PostProcess."));
-					}
-
-					if (IsTranslucentBlendMode(Analyzer.Module->GetBlendMode()))
-					{
-						Analyzer.Module->AddError(nullptr, TEXT("Can only read scene depth below water when material Blend Mode isn't translucent."));
-					}
-				}
-			}
-			break;
-		case MIR::EScreenTexture::DBufferTexture:
-			{
-				// Add defines and compilation outputs
-				Analyzer.EnvironmentDefines.Add(FName("MATERIAL_USES_DECAL_LOOKUP"));
-
-				Analyzer.CompilationOutput->SetIsDBufferTextureUsed(ScreenTexture->DBufferId);
-				// set separate flag to indicate that material uses DBuffer lookup specifically
-				// can't rely on UsedDBufferTextures because those bits are also set depending on the default decal response behavior.
-				Analyzer.CompilationOutput->SetIsDBufferTextureLookupUsed(true);
-
-				// Error checking
-				if (MaterialDomain != MD_Surface || IsTranslucentBlendMode(Analyzer.Module->GetBlendMode()))
-				{
-					Analyzer.Module->AddError(nullptr, TEXT("DBuffer scene textures are only available on opaque or masked surfaces."));
-				}
-			}
-			break;
-	}
-}
-
-static void AnalyzeInStageScreenTexture(FMaterialIRValueAnalyzer& Analyzer, MIR::FScreenTexture* ScreenTexture, MIR::EStage Stage)
-{
-	// This function handles stage specific error checking
-	switch (ScreenTexture->TextureKind)
-	{
-		case MIR::EScreenTexture::SceneTexture:
-		case MIR::EScreenTexture::UserSceneTexture:
-			break;
-		case MIR::EScreenTexture::SceneColor:
-			if (Stage == MIR::Stage_Vertex)
-			{
-				Analyzer.Module->AddError(nullptr, TEXT("Scene Color is only supported in pixel shader input!"));
-			}
-			break;
-		case MIR::EScreenTexture::SceneDepth:
-			if (Stage == MIR::Stage_Vertex)
-			{
-				ErrorUnlessFeatureLevelSupported(Analyzer.Module, ERHIFeatureLevel::SM5, TEXT("Reading scene depth from the vertex shader requires SM5."));
-			}
-			break;
-		case MIR::EScreenTexture::SceneDepthWithoutWater:
-			if (Stage == MIR::Stage_Vertex)
-			{
-				// Mobile currently does not support this, we need to read a separate copy of the depth, we must disable framebuffer fetch and force scene texture reads.
-				// (Texture bindings are not setup properly for any platform so we're disallowing usage in vertex shader altogether now)
-				Analyzer.Module->AddError(nullptr, TEXT("Cannot read scene depth without water from the vertex shader."));
-			}
-			break;
-		case MIR::EScreenTexture::DBufferTexture:
-			break;
-	}
-}
-
-static void AnalyzeShadingModel(FMaterialIRValueAnalyzer& Analyzer, MIR::FShadingModel* ShadingModel)
-{
-	if (ShadingModel->Id < MSM_NUM)
-	{
-		Analyzer.Module->AddShadingModel(ShadingModel->Id);
-	}
-}
-
 static void AnalyzeBranch(FMaterialIRValueAnalyzer& Analyzer, MIR::FBranch* Branch)
 {
 	Branch->TrueBlock = Analyzer.Module->AllocateArray<MIR::FBlock>(Analyzer.Module->GetNumEntryPoints());
@@ -883,16 +523,13 @@ void FMaterialIRValueAnalyzer::Analyze(MIR::FValue* Value)
 	switch (Value->Kind)
 	{
 		EXPAND_CASE(AnalysisError);
-		EXPAND_CASE(ExternalInput);
 		EXPAND_CASE(MaterialParameterCollection);
 		EXPAND_CASE(UniformParameter);
 		EXPAND_CASE(Branch);
 		EXPAND_CASE(TextureRead);
 		EXPAND_CASE(PreshaderParameter);
 		EXPAND_CASE(SetMaterialOutput);
-		EXPAND_CASE(InlineHLSL);
-		EXPAND_CASE(ScreenTexture);
-		EXPAND_CASE(ShadingModel);
+		EXPAND_CASE(Extern);
 		EXPAND_CASE(GetVertexInterpolator);
 		EXPAND_CASE(SetVertexInterpolator);
 		default: break;
@@ -913,10 +550,8 @@ void FMaterialIRValueAnalyzer::AnalyzeInStage(MIR::FValue* Value, MIR::EStage St
 	// IMPORTANT: Before adding a case here, read the FMaterialIRValueAnalyzer documentation.
 	switch (Value->Kind)
 	{
-		EXPAND_CASE(ExternalInput);
-		EXPAND_CASE(InlineHLSL);
+		EXPAND_CASE(Extern);
 		EXPAND_CASE(VTPageTableRead);
-		EXPAND_CASE(ScreenTexture);
 		EXPAND_CASE(HardwarePartialDerivative);
 		default: break;
 	}
