New material translator:

diff --git a/Engine/Source/Runtime/Engine/Public/Materials/MaterialIR.h b/Engine/Source/Runtime/Engine/Public/Materials/MaterialIR.h
--- a/Engine/Source/Runtime/Engine/Public/Materials/MaterialIR.h
+++ b/Engine/Source/Runtime/Engine/Public/Materials/MaterialIR.h
@@ -3,9 +3,9 @@
 
 #include "Materials/MaterialIRCommon.h"
 #include "Materials/MaterialIRTypes.h"
+#include "Materials/MaterialIRExtern.h"
 #include "Materials/MaterialParameters.h"
-#include "Materials/MaterialExpressionDBufferTexture.h"
-#include "MaterialSceneTextureId.h"
+#include "Materials/MaterialExpressionDBufferTexture.h" // @massimo.tristano remove this
 #include "Engine/TextureDefines.h"
 #include "Shader/Preshader.h"
 #include "Shader/Preshader2.h"
@@ -19,32 +19,6 @@ class URuntimeVirtualTexture;
 namespace MIR
 {
 
-// A simple string view without constructors, to be used inside FValues.
-struct FSimpleStringView
-{
-	const TCHAR* Ptr; // The string characters
-	int32        Len; // The string length (excluding null byte)
-
-	FStringView ToView() const { return { Ptr, Len }; }
-	FString ToString() const { return FString{ ToView() }; }
-	operator FStringView() const { return ToView(); }
-	FSimpleStringView& operator=(FStringView View) { Ptr = View.GetData(); Len = View.Len(); return *this; }
-	const TCHAR* operator*() const { return Ptr; }
-};
-
-// Identifies the block of execution in which an instruction runs.
-// Note: If you introduce a new stage, make sure you update EValueFlags accordingly.
-enum EStage
-{
-	Stage_Vertex,
-	Stage_Pixel,
-	Stage_Compute,
-	NumStages
-};
-
-// Returns the string representation of given stage.
-const TCHAR* LexToString(EStage Stage);
-
 // A collection of bit flags for a specific Value instance.
 enum class EValueFlags : uint8
 {
@@ -61,43 +35,10 @@ enum class EValueFlags : uint8
 
 	// Mask that includes all the analyzed flag in all stages.
 	AnalyzedInAllStagesMask = AnalyzedForStageVertex | AnalyzedForStagePixel | AnalyzedForStageCompute,
-	
-	// Whether to substitute abstract tokens like "<PREV>" in code expressions of InlineHLSL instructions.
-	SubstituteTags = 1 << 4,
-
-	// Whether 'FInlineHLSL::Code' field is used. Otherwise, 'ExternalCodeDeclaration' is used.
-	HasDynamicHLSLCode = 1 << 5, 
-
-	// External code declaration must emit the secondary definition for the DDX derivative.
-	DerivativeDDX = 1 << 6,
-
-	// External code declaration must emit the secondary definition for the DDY derivative.
-	DerivativeDDY = 1 << 7,
 };
 
 ENUM_CLASS_FLAGS(EValueFlags);
 
-// A collection of graph properties used by a value.
-//
-// If a graph property is true, it means that either the value itself makes has of that
-// property, or one of its dependencies (direct or indirect) has it. This entails these
-// flags are automatically propagated to all dependant values (values that depend on a given one).
-// As an example, if "ReadsPixelNormal" is true on a specific value it means that either
-// that value itself or some other upstream value that value is dependent on reads
-// the pixel normal.
-enum class EGraphProperties : uint8
-{
-	None = 0,
-
-	// This value is only meaningful and available in the pixel stage.
-	PixelStageOnly = 1 << 0,
-
-	// Some value reads the pixel normal.
-	ReadsPixelNormal = 1 << 1,
-};
-
-ENUM_CLASS_FLAGS(EGraphProperties);
-
 // Enumeration of all different structs deriving from FValue. Used for dynamic casting.
 enum EValueKind
 {
@@ -106,9 +47,8 @@ enum EValueKind
 	VK_Poison,
 	VK_AnalysisError,
 	VK_Constant,
-	VK_ExternalInput,
+	VK_Builtin,
 	VK_MaterialParameterCollection,
-	VK_ScreenTexture,
 	VK_ShadingModel,
 	VK_TextureObject,
 	VK_RuntimeVirtualTextureObject,
@@ -128,7 +68,8 @@ enum EValueKind
 	VK_Scalar,
 	VK_TextureRead,
 	VK_VTPageTableRead,
-	VK_InlineHLSL,
+	VK_Extern,
+	VK_ExternDifferential,
 	VK_SubstrateDefaultSlab,
 	VK_SubstrateSlab,
 	VK_SubstrateShadingModels,
@@ -377,17 +318,6 @@ struct FConstant : TValue<VK_Constant>
 	}
 };
 
-// Specifies the axis for computing screen-space derivatives.
-//
-// This enum is used to indicate whether a partial derivative should be taken
-// along the X or Y screen-space direction, corresponding to HLSL's ddx and ddy
-// functions.
-enum class EDerivativeAxis
-{
-	X, // Corresponding to d/dx
-	Y, // Corresponding to d/dy
-};
-
 // Specify the source expression type of a derivative, useful for generating actionable user facing errors if derivatives
 // are used where not allowed (vertex shader).
 enum class EDerivativeSource
@@ -397,131 +327,21 @@ enum class EDerivativeSource
 	AnalyticDerivative,		// Generated during calculation of an analytic derivative
 };
 
-// Enumeration of all supported material external inputs.
-enum class EExternalInput
+// Enumeration of builtin values.
+enum class EBuiltin
 {
-	None,
-
-	TexCoord0,
-	TexCoord1,
-	TexCoord2,
-	TexCoord3,
-	TexCoord4,
-	TexCoord5,
-	TexCoord6,
-	TexCoord7,
-	WorldPosition_Absolute,					// (LWC)	WPT_Default						Absolute World Position, including Material Shader Offsets
-	WorldPosition_AbsoluteNoOffsets,		// (LWC)	WPT_ExcludeAllShaderOffsets		Absolute World Position, excluding Material Shader Offsets
-	WorldPosition_CameraRelative,			// (float)	WPT_CameraRelative				Camera Relative World Position, including Material Shader Offsets
-	WorldPosition_CameraRelativeNoOffsets,	// (float)	WPT_CameraRelativeNoOffsets		Camera Relative World Position, excluding Material Shader Offsets
-	LocalPosition_Instance,
-	LocalPosition_InstanceNoOffsets,
-	LocalPosition_Primitive,
-	LocalPosition_PrimitiveNoOffsets,
-
-	TexCoord0_Ddx,
-	TexCoord1_Ddx,
-	TexCoord2_Ddx,
-	TexCoord3_Ddx,
-	TexCoord4_Ddx,
-	TexCoord5_Ddx,
-	TexCoord6_Ddx,
-	TexCoord7_Ddx,
-	WorldPosition_Absolute_Ddx,				// LWC (cast from float on load)
-	WorldPosition_AbsoluteNoOffsets_Ddx,	// LWC (cast from float on load)
-	WorldPosition_CameraRelative_Ddx,
-	WorldPosition_CameraRelativeNoOffsets_Ddx,
-	LocalPosition_Instance_Ddx,
-	LocalPosition_InstanceNoOffsets_Ddx,
-	LocalPosition_Primitive_Ddx,
-	LocalPosition_PrimitiveNoOffsets_Ddx,
-
-	TexCoord0_Ddy,
-	TexCoord1_Ddy,
-	TexCoord2_Ddy,
-	TexCoord3_Ddy,
-	TexCoord4_Ddy,
-	TexCoord5_Ddy,
-	TexCoord6_Ddy,
-	TexCoord7_Ddy,
-	WorldPosition_Absolute_Ddy,				// LWC (cast from float on load)
-	WorldPosition_AbsoluteNoOffsets_Ddy,	// LWC (cast from float on load)
-	WorldPosition_CameraRelative_Ddy,
-	WorldPosition_CameraRelativeNoOffsets_Ddy,
-	LocalPosition_Instance_Ddy,
-	LocalPosition_InstanceNoOffsets_Ddy,
-	LocalPosition_Primitive_Ddy,
-	LocalPosition_PrimitiveNoOffsets_Ddy,
-
-	// Ranges of External inputs with derivatives.  There must be an equal number of value, ddx, and ddy variations.  TexCoord0 must remain first,
-	// as code that translates texture coordinate indices to external input ID and back assumes this.
-	WithDerivatives_First = TexCoord0,
-	WithDerivatives_LastVal = LocalPosition_PrimitiveNoOffsets,
-	WithDerivatives_LastDdx = LocalPosition_PrimitiveNoOffsets_Ddx,
-	WithDerivatives_LastDdy = LocalPosition_PrimitiveNoOffsets_Ddy,
-	WithDerivatives_Last = WithDerivatives_LastDdy,
-
-	ActorPosition_Absolute,					// LWC
-	ActorPosition_CameraRelative,
-	ObjectPosition_Absolute,				// LWC
-	ObjectPosition_CameraRelative,
-	ViewMaterialTextureMipBias,
-	ViewMaterialTextureDerivativeMultiply,
-	GlobalDistanceField,					// Flag to mark that global distance fields are used, for value analyzer step.  Doesn't generate code.
-	DynamicParticleParameterIndex,			// UserData stores the parameter index.
-	CompilingPreviousFrame,					// Boolean value bCompilingPreviousFrame, set when compiling the vertex shader WorldPositionOffsets for the previous frame.
-	DistanceCullFade,
-	
-	Count,
+	TextureMipBias,				// A globally defined bias to the mip index used in texture sampling with automatic view mip bias
+	TextureDerivativeMultiply,	// Used in texture sampling with automatic bip bias to correctly calculate texcoord screen differentials
 };
 
-// The number of external input derivative groups (value, ddx, ddy).
-static constexpr int32 ExternalInputWithDerivativesGroups = 3;
-
-// Number of external inputs with derivatives
-static constexpr int32 ExternalInputWithDerivativesNum = (int32)EExternalInput::WithDerivatives_LastVal - (int32)EExternalInput::WithDerivatives_First + 1;
-
-// Maximum number of supported texture coordinates.
-static constexpr int32 TexCoordMaxNum = 8;
+// Returns the string representation of given builtin.
+const TCHAR* LexToString(EBuiltin ID);
 
-// Returns the string representation of given external input.
-const TCHAR* LexToString(EExternalInput Input);
-
-// Returns the given external input type.
-FType GetExternalInputType(EExternalInput Id);
-
-// Returns whether the given external input needs to generate non-trivial derivative expressions.
-bool IsExternalInputWithDerivatives(EExternalInput Id);
-
-// Translate an input ID to a derivative variation
-EExternalInput GetExternalInputDerivative(EExternalInput Id, EDerivativeAxis Axis);
-
-// Converts the given texture coordinate index to its corresponding external input mapping.
-MIR::EExternalInput TexCoordIndexToExternalInput(int32 TexCoordIndex);
-
-// Returns the index of the specified texture coordinates external input.
-int32 ExternalInputToTexCoordIndex(EExternalInput Id);
-
-// Returns whether the given external is a texture coordinate.
-bool IsExternalInputTexCoord(EExternalInput Id);
-
-// Returns whether the given external is a texture coordinate ddx.
-bool IsExternalInputTexCoordDdx(EExternalInput Id);
-
-// Returns whether the given external is a texture coordinate ddy.
-bool IsExternalInputTexCoordDdy(EExternalInput Id);
-
-// Returns whether the given external is a texture coordinate, a ddx or ddy.
-bool IsExternalInputTexCoordOrPartialDerivative(EExternalInput Id);
-
-// Returns whether the given external is a world position, including a ddx or ddy.
-bool IsExternalInputWorldPosition(EExternalInput Id);
-
-// It represents an external input (like texture coordinates) value.
-struct FExternalInput : TValue<VK_ExternalInput>
+// It builtin is symbolic value known to the Material IR but provided externally.
+// The actual numerical value is left to the backend, which can resolve it to a constant or any other uniform expression.
+struct FBuiltin : TValue<VK_Builtin>
 {
-	EExternalInput Id;
-	uint32 UserData;				// Optional user data for certain external input types
+	EBuiltin Id;
 };
 
 // It represents a Material Parameter Collection object value.
@@ -588,6 +408,11 @@ struct FUniformParameter: TValue<VK_UniformParameter>
 	int32 Analysis_UniformParameterIndex;
 };
 
+struct FShadingModel : TValue<VK_ShadingModel>
+{
+	EMaterialShadingModel Id;
+};
+
 // Value used to retrieve (in the pixel stage) a value interpolated in the vertex stage.
 struct FGetVertexInterpolator : TValue<VK_GetVertexInterpolator>
 {
@@ -595,35 +420,6 @@ struct FGetVertexInterpolator : TValue<VK_GetVertexInterpolator>
 	struct FSetVertexInterpolator* SetInstr;
 };
 
-enum class EScreenTexture : uint8
-{
-	SceneTexture,
-	UserSceneTexture,
-	SceneColor,
-	SceneDepth,
-	SceneDepthWithoutWater,
-	DBufferTexture,
-};
-
-// Reference to various types of full screen textures accessible by expressions.  These
-// generally involve custom code to set data in the compilation output, extensive validation
-// logic, and HLSL generation deferred to the analyzer for the case of UserSceneTextures.
-struct FScreenTexture : TValue<VK_ScreenTexture>
-{
-	EScreenTexture TextureKind;
-	union
-	{
-		ESceneTextureId Id;
-		EDBufferTextureId DBufferId;
-	};
-	FName UserSceneTexture;
-};
-
-struct FShadingModel : TValue<VK_ShadingModel>
-{
-	EMaterialShadingModel Id;
-};
-
 /*------------------------------------ Instructions ------------------------------------*/
 
 // A block of instructions.
@@ -769,210 +565,64 @@ struct FSetMaterialOutput : TInstruction<VK_SetMaterialOutput, 1>
 	int Priority;
 };
 
-// Inline HLSL instruction.
-//
-// A value that is the result of an arbitrary snippet of HLSL code.
-//
-// Note: See BaseMaterialExpressions.ini file.
-struct FInlineHLSL : TInstruction<VK_InlineHLSL, 0>
-{
-	static constexpr int32 MaxNumArguments = 16;
-
-	// Array of argument values.
-	FValue* Arguments[MaxNumArguments];
-
-	// Custom analysis function.
-	using TOnAnalyzeCallback = void ( *)(FMaterialIRModule*, uintptr_t UserData);
-	TOnAnalyzeCallback OnAnalyze;
-
-	// User specified data passed to the OnAnalyze callback.
-	uintptr_t OnAnalyzeUserData;
-
-	union
-	{
-		// The declaration this instruction refers to.
-		const FMaterialExternalCodeDeclaration* ExternalCodeDeclaration;
-
-		// The arbitrary inlined HLSL code snippet
-		FSimpleStringView Code;
-	};
-
-	// Number of arguments from another expression node.
-	int32 NumArguments;
-};
-
-// Material IR representation of the FSubstrateData shader parameter to transition legacy materials into Substrate.
-struct FSubstrateDefaultSlab : TInstruction<VK_SubstrateDefaultSlab, 0>
-{
-	FValue* Dummy;
-};
-
-struct FSubstrateSlab : TInstruction<VK_SubstrateSlab, 27> 
-{
-	FValue* Normal;
-	FValue* DiffuseAlbedo;
-	FValue* F0;
-	FValue* F90;
-	FValue* Roughness;
-	FValue* Anisotropy;
-	FValue* SSSProfileId;
-	FValue* SSSMFP;
-	FValue* SSSMFPScale;
-	FValue* SSSPhaseAniso;
-	FValue* SSSType;
-	FValue* EmissiveColor;
-	FValue* SecondRoughness;
-	FValue* SecondRoughnessWeight;
-	FValue* SecondRoughnessAsSimpleClearCoat;
-	FValue* ClearCoatUseSecondNormal;
-	FValue* ClearCoatBottomNormal;
-	FValue* FuzzAmount;
-	FValue* FuzzColor;
-	FValue* FuzzRoughness;
-	FValue* GlintValue;
-	FValue* GlintUV;
-	FValue* SpecularProfileId;
-	FValue* Thickness;
-	FValue* IsThin;
-	FValue* IsAtBottom;
-	FValue* LocalBasisIndex;
-
-	bool bFullySimplifiedInstruction;
-};
-
-struct FSubstrateShadingModels : TInstruction<VK_SubstrateShadingModels, 24>
+// Provides an implementation of a user-specified Extern definition.  
+struct FExternImpl
 {
-	FValue* BaseColor;
-	FValue* Specular;
-	FValue* Metallic;
-	FValue* Roughness;
-	FValue* Anisotropy;
-	FValue* SubSurfaceColor;
-	FValue* SubSurfaceProfileId;
-	FValue* ClearCoat;
-	FValue* ClearCoatRoughness;
-	FValue* EmissiveColor;
-	FValue* Opacity;
-	FValue* ThinTranslucentTransmittanceColor;
-	FValue* ThinTranslucentSurfaceCoverage;
-	FValue* WaterScatteringCoefficients;
-	FValue* WaterAbsorptionCoefficients;
-	FValue* WaterPhaseG;
-	FValue* ColorScaleBehindWater;
-	FValue* ShadingModel;
-	FValue* Normal;
-	FValue* Tangent;
-	FValue* ClearCoatNormal;
-	FValue* CustomTangent;
-	FValue* BasisIndexMacro;
-	FValue* ClearCoat_BasisIndexMacro;
-
-	bool bHasDynamicShadingModels;
-	bool bFullySimplifiedInstruction;
-};
-
-struct FSubstrateHorizontalMixing : TInstruction<VK_SubstrateHorizontalMixing, 9>
-{
-	FValue* Background;
-	FValue* Foreground;
-	FValue* Mix;
-
-	// ParameterBlending = 0
-	FValue* OperatorIndex;
-	FValue* MaxDistanceFromLeaves;
-
-	// ParameterBlending = 1
-	FValue* NormalMix;
-	FValue* SharedLocalBasisIndexMacro;
-	FValue* BackgroundNoV;
-	FValue* ForegroundNoV;
+	// The byte size of the whole Extern value.
+	uint32 ByteSize;
 
-	// Settings
-	bool bParameterBlendingEnabled;
-	bool bFullySimplifiedInstruction;
+	// Pointers to the trampoline functions to the user definition methods.
+	// Note: it implements the CExternDefinition plus all optional functions
+	TConstArrayView<FValue*> (*GetArguments)(const FExtern* Extern);
+	MIR::FExternInfo (*GetInfo)(const FExtern* Extern);
+	void (*Analyze)(FExtern* Extern, FExternAnalysisContext& Context);
+	void (*AnalyzeInStage)(FExtern* Extern, FExternAnalysisContext& Context, EStage Stage);
+	void (*ToHLSL)(const FExtern* Extern, FExternPrinterHLSL& Printer);
+	void (*EmitDebugInfo)(const FExtern* Extern, FString& Out);
 };
 
-struct FSubstrateVerticalLayering : TInstruction<VK_SubstrateVerticalLayering, 7>
+// An Extern represents a user-defined, backend-facing operation in the Material IR.
+// It is a generic instruction whose behavior is fully described by a user provided
+// struct that implements the CExternDefinition concept, which provides methods for argument access,
+// type/info description, analysis, stage-specific analysis, HLSL emission, and optional debug output.
+// The Extern instruction provides the extensibility point for users to model material-specific functionality
+// not directly natively provided by MIR.
+// See `MIR::CExternDefinition` for more information.
+struct FExtern : TInstruction<VK_Extern, 0>
 {
-	FValue* Top;
-	FValue* Base;
+	// Pointer to the implementation table.
+	const FExternImpl* Impl;
 
-	// ParameterBlending = 0
-	FValue* OperatorIndex;
-	FValue* MaxDistanceFromLeaves;
-
-	// ParameterBlending = 1
-	FValue* SharedLocalBasisIndexMacro;
-	FValue* TopNoV;
-	FValue* BaseNoV;
+	// Returns the number of argument values which correspond to this instruction's uses.
+	TConstArrayView<FValue*> GetArguments() const { return Impl->GetArguments(this); }
 
-	// Settings
-	bool bParameterBlendingEnabled;
-	bool bFullySimplifiedInstruction;
-};
+	// See `MIR::CExternDefinition`
+	FExternInfo GetInfo() const { return Impl->GetInfo(this); }
 
-struct FSubstrateCoverageWeight : TInstruction<VK_SubstrateCoverageWeight, 4>
-{
-	FValue* A;
-	FValue* Weight;
+	// See `MIR::CExternDefinition`
+	void Analyze(FExternAnalysisContext& Context) { Impl->Analyze(this, Context); }
 
-	// ParameterBlending = 0
-	FValue* OperatorIndex;
-	FValue* MaxDistanceFromLeaves;
+	// See `MIR::CExternDefinition`
+	void AnalyzeInStage(FExternAnalysisContext& Context, EStage Stage) { Impl->AnalyzeInStage(this, Context, Stage); }
+	
+	// See `MIR::CExternDefinition`
+	void ToHLSL(FExternPrinterHLSL& Printer) const { Impl->ToHLSL(this, Printer); }
 
-	// Settings
-	bool bParameterBlendingEnabled;
-	bool bFullySimplifiedInstruction;
+	// See `MIR::CExternDefinition`
+	void EmitDebugInfo(FString& Out) const { Impl->EmitDebugInfo(this, Out); }
 };
 
-struct FSubstrateAdd : TInstruction<VK_SubstrateAdd, 8>
+// An auxiliary instruction used to generate the differential (ddx/ddy) of an Extern.
+struct FExternDifferential : TInstruction<VK_ExternDifferential, 1>
 {
-	FValue* A;
-	FValue* B;
+	// The extern to differentiate
+	FValue* ExternArg;
 
-	// ParameterBlending = 0
-	FValue* OperatorIndex;
-	FValue* MaxDistanceFromLeaves;
-
-	// ParameterBlending = 1
-	FValue* NormalMix;
-	FValue* SharedLocalBasisIndexMacro;
-	FValue* ANoV;
-	FValue* BNoV;
-
-	// Settings
-	bool bParameterBlendingEnabled;
-	bool bFullySimplifiedInstruction;
-};
-
-struct FSubstrateSelect : TInstruction<VK_SubstrateSelect, 6>
-{
-	FValue* A;
-	FValue* B;
-	FValue* SelectValue;
-
-	// ParameterBlending = 0
-	FValue* OperatorIndex;
-	FValue* MaxDistanceFromLeaves;
-
-	// ParameterBlending = 1
-	FValue* SharedLocalBasisIndexMacro;
-
-	// Settings
-	bool bParameterBlendingEnabled;
-	bool bFullySimplifiedInstruction;
-};
-
-struct FSubstratePromoteToOperator : TInstruction<VK_SubstratePromoteToOperator, 5>
-{
-	FValue* SubstrateDataInput;
-	FValue* OperatorIndex;
-	FValue* BSDFIndex;
-	FValue* LayerDepth;
-	FValue* bIsBottom;
+	// The differential axis to generate for
+	EDerivativeAxis Axis;
 
-	// Settings
-	bool bFullySimplifiedInstruction;
+	// Convenience getter of the extern.
+	const FExtern* GetExtern() const { return MIR::As<FExtern>(ExternArg); }
 };
 
 // Operator enumeration.
@@ -1443,5 +1093,186 @@ struct FPreshaderParameter : TInstruction<VK_PreshaderParameter, 1>
 	uint32 Analysis_PreshaderOffset;
 };
 
+
+
+//------------------------------------------------------------------------
+//                        SUBSTRATE INSTRUCTIONS
+//------------------------------------------------------------------------
+
+// Material IR representation of the FSubstrateData shader parameter to transition legacy materials into Substrate.
+struct FSubstrateDefaultSlab : TInstruction<VK_SubstrateDefaultSlab, 0>
+{
+	FValue* Dummy;
+};
+
+struct FSubstrateSlab : TInstruction<VK_SubstrateSlab, 27> 
+{
+	FValue* Normal;
+	FValue* DiffuseAlbedo;
+	FValue* F0;
+	FValue* F90;
+	FValue* Roughness;
+	FValue* Anisotropy;
+	FValue* SSSProfileId;
+	FValue* SSSMFP;
+	FValue* SSSMFPScale;
+	FValue* SSSPhaseAniso;
+	FValue* SSSType;
+	FValue* EmissiveColor;
+	FValue* SecondRoughness;
+	FValue* SecondRoughnessWeight;
+	FValue* SecondRoughnessAsSimpleClearCoat;
+	FValue* ClearCoatUseSecondNormal;
+	FValue* ClearCoatBottomNormal;
+	FValue* FuzzAmount;
+	FValue* FuzzColor;
+	FValue* FuzzRoughness;
+	FValue* GlintValue;
+	FValue* GlintUV;
+	FValue* SpecularProfileId;
+	FValue* Thickness;
+	FValue* IsThin;
+	FValue* IsAtBottom;
+	FValue* LocalBasisIndex;
+
+	bool bFullySimplifiedInstruction;
+};
+
+struct FSubstrateShadingModels : TInstruction<VK_SubstrateShadingModels, 24>
+{
+	FValue* BaseColor;
+	FValue* Specular;
+	FValue* Metallic;
+	FValue* Roughness;
+	FValue* Anisotropy;
+	FValue* SubSurfaceColor;
+	FValue* SubSurfaceProfileId;
+	FValue* ClearCoat;
+	FValue* ClearCoatRoughness;
+	FValue* EmissiveColor;
+	FValue* Opacity;
+	FValue* ThinTranslucentTransmittanceColor;
+	FValue* ThinTranslucentSurfaceCoverage;
+	FValue* WaterScatteringCoefficients;
+	FValue* WaterAbsorptionCoefficients;
+	FValue* WaterPhaseG;
+	FValue* ColorScaleBehindWater;
+	FValue* ShadingModel;
+	FValue* Normal;
+	FValue* Tangent;
+	FValue* ClearCoatNormal;
+	FValue* CustomTangent;
+	FValue* BasisIndexMacro;
+	FValue* ClearCoat_BasisIndexMacro;
+
+	bool bHasDynamicShadingModels;
+	bool bFullySimplifiedInstruction;
+};
+
+struct FSubstrateHorizontalMixing : TInstruction<VK_SubstrateHorizontalMixing, 9>
+{
+	FValue* Background;
+	FValue* Foreground;
+	FValue* Mix;
+
+	// ParameterBlending = 0
+	FValue* OperatorIndex;
+	FValue* MaxDistanceFromLeaves;
+
+	// ParameterBlending = 1
+	FValue* NormalMix;
+	FValue* SharedLocalBasisIndexMacro;
+	FValue* BackgroundNoV;
+	FValue* ForegroundNoV;
+
+	// Settings
+	bool bParameterBlendingEnabled;
+	bool bFullySimplifiedInstruction;
+};
+
+struct FSubstrateVerticalLayering : TInstruction<VK_SubstrateVerticalLayering, 7>
+{
+	FValue* Top;
+	FValue* Base;
+
+	// ParameterBlending = 0
+	FValue* OperatorIndex;
+	FValue* MaxDistanceFromLeaves;
+
+	// ParameterBlending = 1
+	FValue* SharedLocalBasisIndexMacro;
+	FValue* TopNoV;
+	FValue* BaseNoV;
+
+	// Settings
+	bool bParameterBlendingEnabled;
+	bool bFullySimplifiedInstruction;
+};
+
+struct FSubstrateCoverageWeight : TInstruction<VK_SubstrateCoverageWeight, 4>
+{
+	FValue* A;
+	FValue* Weight;
+
+	// ParameterBlending = 0
+	FValue* OperatorIndex;
+	FValue* MaxDistanceFromLeaves;
+
+	// Settings
+	bool bParameterBlendingEnabled;
+	bool bFullySimplifiedInstruction;
+};
+
+struct FSubstrateAdd : TInstruction<VK_SubstrateAdd, 8>
+{
+	FValue* A;
+	FValue* B;
+
+	// ParameterBlending = 0
+	FValue* OperatorIndex;
+	FValue* MaxDistanceFromLeaves;
+
+	// ParameterBlending = 1
+	FValue* NormalMix;
+	FValue* SharedLocalBasisIndexMacro;
+	FValue* ANoV;
+	FValue* BNoV;
+
+	// Settings
+	bool bParameterBlendingEnabled;
+	bool bFullySimplifiedInstruction;
+};
+
+struct FSubstrateSelect : TInstruction<VK_SubstrateSelect, 6>
+{
+	FValue* A;
+	FValue* B;
+	FValue* SelectValue;
+
+	// ParameterBlending = 0
+	FValue* OperatorIndex;
+	FValue* MaxDistanceFromLeaves;
+
+	// ParameterBlending = 1
+	FValue* SharedLocalBasisIndexMacro;
+
+	// Settings
+	bool bParameterBlendingEnabled;
+	bool bFullySimplifiedInstruction;
+};
+
+struct FSubstratePromoteToOperator : TInstruction<VK_SubstratePromoteToOperator, 5>
+{
+	FValue* SubstrateDataInput;
+	FValue* OperatorIndex;
+	FValue* BSDFIndex;
+	FValue* LayerDepth;
+	FValue* bIsBottom;
+
+	// Settings
+	bool bFullySimplifiedInstruction;
+};
+
+
 } // namespace MIR
 #endif // WITH_EDITOR
diff --git a/Engine/Source/Runtime/Engine/Public/Materials/MaterialIRCommon.h b/Engine/Source/Runtime/Engine/Public/Materials/MaterialIRCommon.h
--- a/Engine/Source/Runtime/Engine/Public/Materials/MaterialIRCommon.h
+++ b/Engine/Source/Runtime/Engine/Public/Materials/MaterialIRCommon.h
@@ -24,7 +24,6 @@ struct FMaterialIRModuleBuilder;
 struct FMaterialIRModuleBuilderImpl;
 struct FShaderCompilerEnvironment;
 struct FStaticParameterSet;
-
 class FMaterial;
 class FMaterialIRModule;
 class ITargetPlatform;
@@ -33,6 +32,7 @@ class UMaterialAggregate;
 class UMaterialExpression;
 class UMaterialExpressionMaterialFunctionCall;
 class UTexture;
+class UMaterialInterface;
 
 namespace Substrate
 {
@@ -49,27 +49,103 @@ namespace ERHIFeatureLevel { enum Type : int; }
 namespace MIR
 {
 
-/* Types*/
-struct FType;
-struct FPrimitive;
-struct FObjectType;
-struct FAggregateType;
-
-/* IR */
-struct FValue;
+class FEmitter;
+struct FBlock;
+struct FExtern;
+struct FExternImpl;
+struct FExternPrinterHLSL;
 struct FInstruction;
+struct FPrimitive;
 struct FSetMaterialOutput;
-enum class EExternalInput;
+struct FType;
+struct FValue;
 
-/* Others */
-class FEmitter;
-struct FBlock;
+// A simple string view without constructors, to be used inside FValues.
+struct FSimpleStringView
+{
+	const TCHAR* Ptr; // The string characters
+	int32        Len; // The string length (excluding null byte)
+
+	FStringView ToView() const { return { Ptr, Len }; }
+	FString ToString() const { return FString{ ToView() }; }
+	operator FStringView() const { return ToView(); }
+	FSimpleStringView& operator=(FStringView View) { Ptr = View.GetData(); Len = View.Len(); return *this; }
+	const TCHAR* operator*() const { return Ptr; }
+};
 
-namespace Internal
+// Identifies the block of execution in which an instruction runs.
+// Note: If you introduce a new stage, make sure you update EValueFlags accordingly.
+enum EStage
 {
-	uint32 HashBytes(const char* Ptr, uint32 Size);
+	Stage_Vertex,
+	Stage_Pixel,
+	Stage_Compute,
+	NumStages
+};
+
+//
+enum class EStageMask : uint8
+{
+	None    = 0,
+	Vertex  = 1 << 0,
+	Pixel   = 1 << 1,
+	Compute = 1 << 2,
+	PixelAndCompute = Pixel | Compute,
+	Any     = Vertex | Pixel | Compute
+};
+
+ENUM_CLASS_FLAGS(EStageMask);
+
+//
+static inline bool IsStageInMask(EStage Stage, EStageMask Mask)
+{
+	return (1 << (uint32)Stage & (uint32)Mask) != 0;
 }
 
+// Returns the string representation of given stage.
+static inline const TCHAR* LexToString(EStage Stage)
+{
+	switch (Stage)
+	{
+		case Stage_Vertex:	return TEXT("Vertex");
+		case Stage_Pixel: 	return TEXT("Pixel");
+		case Stage_Compute: return TEXT("Compute");
+		case NumStages: 	check(false);
+	}
+	return nullptr;
+}
+
+// A collection of graph properties used by a value.
+//
+// If a graph property is true, it means that either the value itself makes has of that
+// property, or one of its dependencies (direct or indirect) has it. This entails these
+// flags are automatically propagated to all dependant values (values that depend on a given one).
+// As an example, if "ReadsPixelNormal" is true on a specific value it means that either
+// that value itself or some other upstream value that value is dependent on reads
+// the pixel normal.
+enum class EGraphProperties : uint8
+{
+	None = 0,
+
+	// This value is only meaningful and available in the pixel stage.
+	PixelStageOnly = 1 << 0,
+
+	// Some value reads the pixel normal.
+	ReadsPixelNormal = 1 << 1,
+};
+
+ENUM_CLASS_FLAGS(EGraphProperties);
+
+// Specifies the axis for computing screen-space derivatives.
+//
+// This enum is used to indicate whether a partial derivative should be taken
+// along the X or Y screen-space direction, corresponding to HLSL's ddx and ddy
+// functions.
+enum class EDerivativeAxis
+{
+	X, Y,
+};
+
 /* Utilities */
 
 // Use this to efficiently allocate a temporary array using MemStack, instead
@@ -131,6 +207,12 @@ static inline bool IsMaterialPropertyCustomizedUV(EMaterialProperty Property)
 
 static constexpr int MaxNumCustomizedUVs = 8;
 
+// @massimo.tristano this needs to be moved off to a better place
+namespace Internal
+{
+	uint32 HashBytes(const char* Ptr, uint32 Size);
+}
+
 } // namespace MIR
 
 #define UE_MIR_UNREACHABLE() { check(!"Unreachable"); UE_ASSUME(false); }
diff --git a/Engine/Source/Runtime/Engine/Public/Materials/MaterialIREmitter.h b/Engine/Source/Runtime/Engine/Public/Materials/MaterialIREmitter.h
--- a/Engine/Source/Runtime/Engine/Public/Materials/MaterialIREmitter.h
+++ b/Engine/Source/Runtime/Engine/Public/Materials/MaterialIREmitter.h
@@ -4,6 +4,8 @@
 #include "Materials/MaterialIRCommon.h"
 #include "Materials/MaterialIR.h"
 #include "Materials/MaterialIRTypes.h"
+#include "Materials/MaterialIRModule.h"
+#include "Materials/MaterialIRExtern.h"
 
 #if WITH_EDITOR
 
@@ -27,9 +29,15 @@ struct FSwizzleMask
 
 	// How many components have been defined for swizzle should be made of (maximum four).
 	int NumComponents{};
+	
+	// Convenience ".xy" swizzle mask.
+	static FSwizzleMask XY() { return { EVectorComponent::X, EVectorComponent::Y }; }
+	
+	// Convenience ".zw" swizzle mask.
+	static FSwizzleMask ZW() { return { EVectorComponent::Z, EVectorComponent::W }; }
 
 	// Convenience ".xyz" swizzle mask.
-	static FSwizzleMask XYZ();
+	static FSwizzleMask XYZ() { return { EVectorComponent::X, EVectorComponent::Y, EVectorComponent::Z }; }
 
 	FSwizzleMask() {}
 	FSwizzleMask(EVectorComponent X);
@@ -175,7 +183,7 @@ struct FFunctionHLSLDesc : FFunctionDesc
 };
 
 // Helper structure to simplify passing attributes to texture sampling emitter functions that share the same attributes.
-struct FTextureSampleBaseAttributes
+struct FTextureSampleAttributes
 {
 	// Specifies which sampler source mode to use. By default SSM_FromTextureAsset.
 	ESamplerSourceMode SamplerSourceMode { SSM_FromTextureAsset };
@@ -500,8 +508,8 @@ public:
 	// on specific stage.
 	FValueRef AnalysisError(FType Type, FStringView Message);
 
-	// Emits an external input value of given input Id and returns it.
-	FValueRef ExternalInput(EExternalInput Id, uint32 UserData = 0);
+	// Emits an builtin value of given id and returns it.
+	FValueRef Builtin(EBuiltin Id);
 
 	// Emits a material parameter collection.
 	FValueRef MaterialParameterCollection(UMaterialParameterCollection* Collection);
@@ -521,18 +529,6 @@ public:
 	// Emits a custom primitive data value with given index.
 	FValueRef CustomPrimitiveData(uint32 PrimitiveDataIndex);
 
-	// Emit a screen texture with the given Scene Texture id.
-	FValueRef SceneTexture(ESceneTextureId SceneTextureId);
-
-	// Emit a screen texture with the given User Scene Texture name.
-	FValueRef UserSceneTexture(FName UserSceneTexture);
-
-	// Emit a screen texture with the given DBuffer Texture id.
-	FValueRef DBufferTexture(EDBufferTextureId DBufferId);
-
-	// Emit another type of screen texture (not one of the three types above).
-	FValueRef ScreenTexture(EScreenTexture TextureKind);
-
 	/*----------------------------- Instructions emission -----------------------------*/
 	
 	// Emits an instruction that sets the material attribute with given property to given argument.
@@ -645,7 +641,7 @@ public:
 	// @param TexCoord The UV coordinates at which to sample the texture.
 	// @param GatherMode What channel to gather (any ::GatherXXX value).
 	// @param BaseAttributes Base sampling attributes for sampler source and type as well as VT attributes.
-	FValueRef TextureGather(FValueRef Texture, FValueRef TexCoord, ETextureReadMode GatherMode, const FTextureSampleBaseAttributes& BaseAttributes);
+	FValueRef TextureGather(FValueRef Texture, FValueRef TexCoord, ETextureReadMode GatherMode, const FTextureSampleAttributes& BaseAttributes);
 
 	// Emits a texture sample instruction.
 	// 
@@ -653,7 +649,7 @@ public:
 	// @param TexCoord The UV coordinates at which to sample the texture.
 	// @param bAddViewMaterialMipBias Whether to add the view material texture mip bias to the mip level.
 	// @param BaseAttributes Base sampling attributes for sampler source and type as well as VT attributes.
-	FValueRef TextureSample(FValueRef Texture, FValueRef TexCoord, bool bAddViewMaterialMipBias, const FTextureSampleBaseAttributes& BaseAttributes);
+	FValueRef TextureSample(FValueRef Texture, FValueRef TexCoord, bool bAddViewMaterialMipBias, const FTextureSampleAttributes& BaseAttributes);
 
 	// Emits a texture sample instruction with a manually given mip level.
 	// 
@@ -662,7 +658,7 @@ public:
 	// @param MipLevel Which mip level to use for this sample.
 	// @param bAddViewMaterialMipBias Whether to add the view material texture mip bias to the mip level.
 	// @param BaseAttributes Base sampling attributes for sampler source and type as well as VT attributes.
-	FValueRef TextureSampleLevel(FValueRef Texture, FValueRef TexCoord, FValueRef MipLevel, bool bAddViewMaterialMipBias, const FTextureSampleBaseAttributes& BaseAttributes);
+	FValueRef TextureSampleLevel(FValueRef Texture, FValueRef TexCoord, FValueRef MipLevel, bool bAddViewMaterialMipBias, const FTextureSampleAttributes& BaseAttributes);
 
 	// Emits a texture sample instruction with bias added to the automatically selected mip level.
 	// 
@@ -671,7 +667,7 @@ public:
 	// @param MipBias The bias to add to the automatically selected mip level.
 	// @param bAddViewMaterialMipBias Whether to add the view material texture mip bias to the mip level.
 	// @param BaseAttributes Base sampling attributes for sampler source and type as well as VT attributes.
-	FValueRef TextureSampleBias(FValueRef Texture, FValueRef TexCoord, FValueRef MipBias, bool bAddViewMaterialMipBias, const FTextureSampleBaseAttributes& BaseAttributes);
+	FValueRef TextureSampleBias(FValueRef Texture, FValueRef TexCoord, FValueRef MipBias, bool bAddViewMaterialMipBias, const FTextureSampleAttributes& BaseAttributes);
 	
 	// Emits a texture sample instruction using user provided partial derivatives.
 	// 
@@ -681,7 +677,7 @@ public:
 	// @param TexCoordDdy The partial derivative of the texture coordinates along the screen space y axis.
 	// @param bAddViewMaterialMipBias Whether to add the view material texture mip bias to the mip level.
 	// @param BaseAttributes Base sampling attributes for sampler source and type as well as VT attributes.
-	FValueRef TextureSampleGrad(FValueRef Texture, FValueRef TexCoord, FValueRef TexCoordDdx, FValueRef TexCoordDdy, bool bAddViewMaterialMipBias, const FTextureSampleBaseAttributes& BaseAttributes);
+	FValueRef TextureSampleGrad(FValueRef Texture, FValueRef TexCoord, FValueRef TexCoordDdx, FValueRef TexCoordDdy, bool bAddViewMaterialMipBias, const FTextureSampleAttributes& BaseAttributes);
 
 	// Emits the virtual texture page load instruction.
 	//
@@ -721,25 +717,20 @@ public:
 	// and exposed in the pixel stage. The returned value is thus valid only in the pixel stage.
 	FValueRef VertexInterpolator(FValueRef VertexValue);
 
-	// Emit a new HLSL function inline within the material shader.
-	// @param Type The expected output type of the HLSL code.
-	// @param Code The HLSL code to be emitted.
-	// @param InputValues An array of input values required by the HLSL code.
-	// @param ValueFlags Flags that specify additional properties for this value (default is None).
-	// @param UsedGraphProperties Properties used during graph processing (default is None).
-	FValueRef InlineHLSL(FType Type, FStringView Code, TConstArrayView<FValueRef> InputValues = {}, EValueFlags ValueFlags = EValueFlags::None, EGraphProperties UsedGraphProperties = EGraphProperties::None, FInlineHLSL::TOnAnalyzeCallback OnAnalyze = {}, uintptr_t OnAnalyzeUserData = {});
-
-	// Emit a new HLSL function inline within the material shader based on an external code declaration.
-	// @param InExternalCodeDeclarationName The name of the external HLSL code declaration to use for emission.
-	// @param InputValues An array of input values required by the HLSL code.
-	// @param UsedGraphProperties Properties used during graph processing (default is None).
-	FValueRef InlineHLSL(FName InExternalCodeDeclarationName, TConstArrayView<FValueRef> InputValues = {}, EValueFlags ValueFlags = EValueFlags::None, EGraphProperties UsedGraphProperties = EGraphProperties::None);
-
-	// Emit a new HLSL function inline within the material shader based on an external code declaration.
-	// @param InExternalCodeDeclaration The external HLSL code declaration to use for emission.
-	// @param InputValues An array of input values required by the HLSL code.
-	// @param UsedGraphProperties Properties used during graph processing (default is None).
-	FValueRef InlineHLSL(const FMaterialExternalCodeDeclaration* InExternalCodeDeclaration, TConstArrayView<FValueRef> InputValues, EValueFlags ValueFlags = EValueFlags::None, EGraphProperties UsedGraphProperties = EGraphProperties::None);
+	// Emits an extern instruction based on a user-specified definition struct.
+	// An extern instruction is how you can extend MIR capabilities from the outside. An extern instruction implements a well-defined
+	// interface (specified by the MIR::CExternDefinition concept) and provides a way to build functionality not natively offered by the
+	// MIR instruction set. For example, material-specific rendering features are implemented through externs rather than by extending MIR
+	// itself. The MIR module builder treats extern instructions as black boxes, but can still work with them through their interface.
+	// To emit an extern instruction, you need to provide a "definition", meaning a concrete POD struct that implements the extern concept.
+	// Parameters:
+	//    - `DefinitionPrototype` is an instance of the TDefinition struct you provide. This is moved into the emitted extern instruction's
+	//      own buffer.
+	//    - With `Arguments`, you can specify zero or more FValueRef arguments to provide to the extern instruction.
+	// For more usage and general information, see `MIR::CExternDefinition`.
+	template <CExternDefinition TDefinition, typename... TArguments>
+	requires (std::same_as<TArguments, FValueRef> && ...)
+	FValueRef Extern(TDefinition DefinitionPrototype, TArguments... Arguments);
 
 	// Defines a user-specified HLSL function to embed in the generated HLSL when generated IR Module is translated to HLSL.
 	// This will create a local function with the specified description. You can issue calls to the returned functions using the Call() API.
@@ -901,6 +892,9 @@ private:
 	// Initializes the emitter, called by the friend builder.
 	void Initialize();
 
+	//
+	FValueRef EmitPrototype_Internal(const FValue& Value);
+
 private:
 	// Pointer to the builder internal implementation.
 	FMaterialIRModuleBuilderImpl* BuilderImpl{};
@@ -941,6 +935,136 @@ private:
 	friend FMaterialIRModuleBuilderImpl;
 };
 
+
+//------------------------------------------------------------------------
+//                        IMPLEMENTATION DETAILS
+//------------------------------------------------------------------------
+
+namespace Internal
+{
+
+template <int32 NumArguments>
+struct TValueArgumentsArray
+{
+	FValue* Arguments[NumArguments];
+};
+
+template <>
+struct TValueArgumentsArray<0>
+{
+};
+
+template <typename TDefinition, int32 TNumArguments>
+struct TConcreteExtern : FExtern, TValueArgumentsArray<TNumArguments>, TDefinition
+{
+	TConcreteExtern(TDefinition&& Definition) : TDefinition { MoveTemp(Definition) } {}
+
+	static TDefinition& Get(FExtern *Self)
+	{
+		return *static_cast<TDefinition *>(static_cast<TConcreteExtern*>(Self));
+	}
+
+	static const TDefinition& Get(const FExtern *Self)
+	{
+		return *static_cast<const TDefinition *>(static_cast<const TConcreteExtern*>(Self));
+	}
+};
+
+template <typename TDefinition, int32 TNumArguments>
+struct TConcreteExternInterface
+{
+	using TMyConcreteExtern = TConcreteExtern<TDefinition, TNumArguments>;
+
+	inline static const FExternImpl Instance =
+	{
+		.ByteSize = (uint32)sizeof(TMyConcreteExtern),
+
+		.GetArguments = +[] (const FExtern* Extern) -> TConstArrayView<FValue*>
+		{
+			if constexpr (TNumArguments > 0)
+			{
+				return { static_cast<const TMyConcreteExtern*>(Extern)->Arguments, TNumArguments };
+			}
+			else
+			{
+				return {};
+			}
+		},
+
+		.GetInfo = +[] (const FExtern* Extern) -> MIR::FExternInfo
+		{
+			return TMyConcreteExtern::Get(Extern).GetInfo();
+		},
+
+		.Analyze = +[] (FExtern* Extern, FExternAnalysisContext& Context) {
+			if constexpr (CExternHasAnalyze<TDefinition>)
+			{
+				TMyConcreteExtern::Get(Extern).Analyze(Context);
+			}
+		},
+
+		.AnalyzeInStage = +[] (FExtern* Extern, FExternAnalysisContext& Context, EStage Stage)
+		{
+			if constexpr (CExternHasAnalyzeInStage<TDefinition>)
+			{
+				TMyConcreteExtern::Get(Extern).AnalyzeInStage(Context, Stage);
+			}
+		},
+		
+		.ToHLSL = +[] (const FExtern* Extern, FExternPrinterHLSL& Printer)
+		{
+			TMyConcreteExtern::Get(Extern).ToHLSL(Printer);
+		},
+
+		.EmitDebugInfo = +[] (const FExtern* Extern, FString& Out)
+		{
+			if constexpr (CExternHasEmitDebugInfo<TDefinition>)
+			{
+				TMyConcreteExtern::Get(Extern).EmitDebugInfo(Out);
+			}
+		},
+	};
+};
+
+} // namespace Internal
+
+// Implementation of FEmitter::Extern
+template <CExternDefinition TDefinition, typename... TArguments>
+requires (std::same_as<TArguments, FValueRef> && ...)
+FValueRef FEmitter::Extern(TDefinition DefinitionPrototype, TArguments... Arguments)
+{
+	if constexpr (sizeof...(Arguments) > 0)
+	{
+		if ((!FValueRef{ Arguments }.IsValid() || ...))
+		{
+			return Poison();
+		}
+	}
+
+	using TConcreteExtern = typename Internal::TConcreteExtern<TDefinition, sizeof...(Arguments)>;
+
+	// Read the value type from the definition prototype
+	MIR::FType Type = DefinitionPrototype.GetInfo().Type;
+
+	// Create the concrete extern instruction prototype and pass in the prototype
+	TTypeCompatibleBytes<TConcreteExtern> ProtoStorage;
+	FMemory::Memzero(ProtoStorage);
+
+	TConcreteExtern& ConcreteProto = * new (&ProtoStorage) TConcreteExtern{ MoveTemp(DefinitionPrototype) };
+
+	// Copy over all the provided arguments
+	int32 Index = 0;
+	((ConcreteProto.Arguments[Index++] = Arguments), ...);
+
+	// Setup the base extern instruction members
+	MIR::FExtern& Proto = ConcreteProto;
+	Proto.Kind = VK_Extern;
+	Proto.Type = Type;
+	Proto.Impl = &Internal::TConcreteExternInterface<TDefinition, sizeof...(Arguments)>::Instance;
+
+	return EmitPrototype_Internal(Proto);
+}
+
 } // namespace MIR
 
 #endif // #if WITH_EDITOR
