Make sure to check if the async animation task has completed before waiting on it

#Synced-CL 51092785

diff --git a/Engine/Source/Runtime/Engine/Private/Components/SkeletalMeshComponent.cpp b/Engine/Source/Runtime/Engine/Private/Components/SkeletalMeshComponent.cpp
--- a/Engine/Source/Runtime/Engine/Private/Components/SkeletalMeshComponent.cpp
+++ b/Engine/Source/Runtime/Engine/Private/Components/SkeletalMeshComponent.cpp
@@ -4927,12 +4927,15 @@ void USkeletalMeshComponent::CompleteParallelAnimationEvaluation(bool bDoPostAni
 
 bool USkeletalMeshComponent::HandleExistingParallelEvaluationTask(bool bBlockOnTask, bool bPerformPostAnimEvaluation)
 {
-	// We are already processing eval and we are not inside it
-	if (IsRunningParallelEvaluation()
-		&& ParallelAnimationEvaluationTask->IsAwaitable()
-	)
+	// If we have an async task handle
+	if (IsRunningParallelEvaluation())
 	{
-		if (bBlockOnTask)
+		// If we want to block and wait for it
+		if (bBlockOnTask
+			// and it hasn't completed yet
+			&& !ParallelAnimationEvaluationTask->IsCompleted()
+			// and we aren't currently executing it on this thread
+			&& ParallelAnimationEvaluationTask->IsAwaitable())
 		{
 			// If we try and complete existing parallel evaluation in a transaction, we have to abort the transaction
 			// and retry up to this point. This is because we can only have started parallel evaluation outside of a
@@ -4960,8 +4963,10 @@ bool USkeletalMeshComponent::HandleExistingParallelEvaluationTask(bool bBlockOnT
 			FTaskGraphInterface::Get().WaitUntilTaskCompletes(ParallelAnimationEvaluationTask, ENamedThreads::GameThread);
 			CompleteParallelAnimationEvaluation(bPerformPostAnimEvaluation); //Perform completion now
 		}
+
 		return true;
 	}
+
 	return false;
 }
 
