Removed unused static function from MaterialIRValueAnalyzer.cpp

#rnx
#Synced-CL 51020780

diff --git a/Engine/Source/Runtime/Engine/Private/Materials/MaterialIRValueAnalyzer.cpp b/Engine/Source/Runtime/Engine/Private/Materials/MaterialIRValueAnalyzer.cpp
--- a/Engine/Source/Runtime/Engine/Private/Materials/MaterialIRValueAnalyzer.cpp
+++ b/Engine/Source/Runtime/Engine/Private/Materials/MaterialIRValueAnalyzer.cpp
@@ -30,17 +30,6 @@ static void AnalyzeAnalysisError(FMaterialIRValueAnalyzer& Analyzer, MIR::FAnaly
 	Analyzer.Module->AddError(AnalysisError->Expression, AnalysisError->Message.ToString());
 }
 
-static void ErrorUnlessFeatureLevelSupported(FMaterialIRModule* Module, ERHIFeatureLevel::Type RequiredFeatureLevel, const TCHAR* Message)
-{
-	ERHIFeatureLevel::Type FeatureLevel = Module->GetFeatureLevel();
-	if (FeatureLevel < RequiredFeatureLevel)
-	{
-		FString FeatureLevelName;
-		GetFeatureLevelName(FeatureLevel, FeatureLevelName);
-		Module->AddError(nullptr, FString::Printf(TEXT("%s  Current feature level is %s."), Message, *FeatureLevelName));
-	}
-}
-
 static void AnalyzeInStageHardwarePartialDerivative(FMaterialIRValueAnalyzer& Analyzer, MIR::FHardwarePartialDerivative* PartialDerivative, MIR::EStage Stage)
 {
 	if (Stage == MIR::Stage_Vertex)
