[NetPhysics] Optimize network overhead for replicated Input and State properties via the legacy "Last Resort" replication in NetworkPhysicsComponent.
- Remove Owner* replication (24+8 reduced to 0 bits)
- Replace DeltaSourceFrame with DeltaSourceIndex (16 or 24 reduced to 2 bits)
- Clamp NumFrames (8 reduced to 4 bits)
- Send data array ordered to improve ServerFrame delta serialization (18 bits reduced to 1 bit on arrays larger than 1)

#Synced-CL 51043916

diff --git a/Engine/Source/Runtime/Engine/Private/PhysicsEngine/NetworkPhysicsComponent.cpp b/Engine/Source/Runtime/Engine/Private/PhysicsEngine/NetworkPhysicsComponent.cpp
--- a/Engine/Source/Runtime/Engine/Private/PhysicsEngine/NetworkPhysicsComponent.cpp
+++ b/Engine/Source/Runtime/Engine/Private/PhysicsEngine/NetworkPhysicsComponent.cpp
@@ -11,8 +11,9 @@
 #include "PhysicsProxy/SingleParticlePhysicsProxy.h"
 #include "Chaos/PhysicsObjectInternalInterface.h"
 #include "Net/Core/PushModel/PushModel.h"
-
 #include "Iris/ReplicationState/PropertyNetSerializerInfoRegistry.h"
+#include "Iris/ReplicationState/ReplicationStateDescriptor.h"
+#include "Iris/ReplicationSystem/ReplicationOperations.h"
 
 #if UE_WITH_REMOTE_OBJECT_HANDLE
 #include "UObject/UObjectMigrationContext.h"
@@ -20,21 +21,144 @@
 
 #include UE_INLINE_GENERATED_CPP_BY_NAME(NetworkPhysicsComponent)
 
+namespace UE::Net
+{
+	FReplicationFragment* CreateAndRegisterNetworkPhysicsRewindDataProxyReplicationFragment(UObject* Owner, const FReplicationStateDescriptor* Descriptor, FFragmentRegistrationContext& Context)
+	{
+		if (FNetworkPhysicsRewindDataProxyReplicationFragment* Fragment = new FNetworkPhysicsRewindDataProxyReplicationFragment(
+			Context.GetFragmentTraits() | EReplicationFragmentTraits::DeleteWithInstanceProtocol, Owner, Descriptor))
+		{
+			Fragment->Register(Context);
+			return Fragment;
+		}
+
+		return nullptr;
+	}
+
+	FNetworkPhysicsRewindDataProxyReplicationFragment::FNetworkPhysicsRewindDataProxyReplicationFragment(EReplicationFragmentTraits InTraits, UObject* InOwner, const FReplicationStateDescriptor* InDescriptor)
+		: FReplicationFragment(InTraits)
+		, ReplicationStateDescriptor(InDescriptor)
+		, Owner(InOwner)
+	{
+		// We don't want to create temporary states when applying replicated state as we want to update the replicated fields.
+		Traits |= EReplicationFragmentTraits::HasPersistentTargetStateBuffer;
+
+		if (EnumHasAnyFlags(InTraits, EReplicationFragmentTraits::CanReplicate))
+		{
+			SrcReplicationState = MakeUnique<FPropertyReplicationState>(InDescriptor);
+
+			// We mark the property as dirty via overriding PollReplicatedState
+			Traits |= EReplicationFragmentTraits::NeedsPoll;
+		}
+
+#if WITH_PUSH_MODEL
+		if (EnumHasAnyFlags(InDescriptor->Traits, EReplicationStateTraits::HasPushBasedDirtiness))
+		{
+			Traits |= EReplicationFragmentTraits::HasPushBasedDirtiness;
+			if (EnumHasAnyFlags(InDescriptor->Traits, EReplicationStateTraits::HasFullPushBasedDirtiness))
+			{
+				Traits |= EReplicationFragmentTraits::HasFullPushBasedDirtiness;
+			}
+		}
+#endif
+
+		if (EnumHasAnyFlags(InTraits, EReplicationFragmentTraits::CanReceive))
+		{
+			if (EnumHasAnyFlags(InDescriptor->Traits, EReplicationStateTraits::HasRepNotifies))
+			{
+				ensureMsgf(!EnumHasAnyFlags(InDescriptor->Traits, EReplicationStateTraits::KeepPreviousState), TEXT("FNetworkPhysicsRewindDataProxyReplicationFragment doesn't support OnRep calls with previous value."));
+			}
+
+			// We do custom callbacks in CallRepNotifies rather than in ApplyReplicatedState.
+			Traits |= EReplicationFragmentTraits::HasRepNotifies | EReplicationFragmentTraits::NeedsLegacyCallbacks;
+		}
+	}
+
+	void FNetworkPhysicsRewindDataProxyReplicationFragment::Register(FFragmentRegistrationContext& Context)
+	{
+		FPropertyReplicationState* ReplicationState = SrcReplicationState.Get();
+		Context.RegisterReplicationFragment(this, ReplicationStateDescriptor.GetReference(), (ReplicationState ? ReplicationState->GetStateBuffer() : nullptr));
+	}
+
+	void FNetworkPhysicsRewindDataProxyReplicationFragment::ApplyReplicatedState(FReplicationStateApplyContext& ApplyContext) const
+	{
+		uint8* ExternalStatePointer = reinterpret_cast<uint8*>(Owner) + ReplicationStateDescriptor->MemberProperties[0]->GetOffset_ForGC();
+
+		// Get struct instance from owner and assign owner pointer
+		FNetworkPhysicsRewindDataProxy* ExternalSourceState = reinterpret_cast<FNetworkPhysicsRewindDataProxy*>(ExternalStatePointer);
+		ExternalSourceState->Owner = static_cast<UNetworkPhysicsComponent*>(Owner);
+
+		// Dequantize replicated members to external struct instance
+		// Need to call the serializer directly with the appropriate arguments as the external state doesn't have a ReplicationStateHeader in front of it.
+		FNetDequantizeArgs DequantizeArgs = {};
+		DequantizeArgs.Source = NetSerializerValuePointer(ApplyContext.StateBufferData.RawStateBuffer);
+		DequantizeArgs.Target = NetSerializerValuePointer(ExternalSourceState);
+		DequantizeArgs.NetSerializerConfig = ReplicationStateDescriptor->MemberSerializerDescriptors[0].SerializerConfig;
+
+		const FNetSerializer* Serializer = ReplicationStateDescriptor->MemberSerializerDescriptors[0].Serializer;
+		Serializer->Dequantize(*ApplyContext.NetSerializationContext, DequantizeArgs);
+	}
+
+	bool FNetworkPhysicsRewindDataProxyReplicationFragment::PollReplicatedState(EReplicationFragmentPollFlags PollOption)
+	{
+		if (!SrcReplicationState)
+		{ 
+			return true;
+		}
+
+		// We can early out if we are pushbased and not dirty for polling
+		const bool bPoll = EnumHasAnyFlags(PollOption, EReplicationFragmentPollFlags::PollAllState) ||
+			(EnumHasAnyFlags(PollOption, EReplicationFragmentPollFlags::PollDirtyState) && (!EnumHasAnyFlags(EReplicationFragmentTraits::HasPushBasedDirtiness, Traits) || SrcReplicationState->IsDirtyForPolling()));
+
+		if (bPoll)
+		{
+			return SrcReplicationState->PollPropertyReplicationState(Owner);
+		}
+
+		return SrcReplicationState->IsDirty(0);
+	}
+
+	void FNetworkPhysicsRewindDataProxyReplicationFragment::CallRepNotifies(FReplicationStateApplyContext& Context)
+	{
+		if (const UFunction* RepNotifyFunction = ReplicationStateDescriptor->MemberPropertyDescriptors[0].RepNotifyFunction)
+		{
+			if (Context.bIsInit)
+			{
+				const uint8* ReceivedState = Context.StateBufferData.RawStateBuffer;
+				const uint8* DefaultState = ReplicationStateDescriptor->DefaultStateBuffer;
+
+				if (FReplicationStateOperations::IsEqualQuantizedState(*Context.NetSerializationContext, ReceivedState, DefaultState, ReplicationStateDescriptor))
+				{
+					return;
+				}
+			}
+
+			Owner->ProcessEvent(const_cast<UFunction*>(RepNotifyFunction), nullptr);
+		}
+	}
+
+	void FNetworkPhysicsRewindDataProxyReplicationFragment::CollectOwner(FReplicationStateOwnerCollector* Owners) const
+	{
+		Owners->AddOwner(Owner);
+	}
+}
+
+
 namespace PhysicsReplicationCVars
 {
 	namespace ResimulationCVars
 	{
 		int32 RedundantInputs = 2;
-		static FAutoConsoleVariableRef CVarResimRedundantInputs(TEXT("np2.Resim.RedundantInputs"), RedundantInputs, TEXT("How many extra inputs to send with each unreliable network message, to account for packetloss. From owning client to server and back to owning client. NOTE: This is disabled while np2.Resim.DynamicInputReplicationScaling.Enabled is enabled."));
+		static FAutoConsoleVariableRef CVarResimRedundantInputs(TEXT("np2.Resim.RedundantInputs"), RedundantInputs, TEXT("How many extra inputs to send with each unreliable network message, to account for packetloss. From owning client to server and back to owning client. NOTE: This is disabled while np2.Resim.DynamicInputReplicationScaling.Enabled is enabled. Clamped by NetworkPhysicsComponentConstants::MaxNumberOfElementsToNetwork."));
 		int32 RedundantRemoteInputs = 1;
-		static FAutoConsoleVariableRef CVarResimRedundantRemoteInputs(TEXT("np2.Resim.RedundantRemoteInputs"), RedundantRemoteInputs, TEXT("How many extra inputs to send with each unreliable network message, to account for packetloss. From server to remote clients."));
+		static FAutoConsoleVariableRef CVarResimRedundantRemoteInputs(TEXT("np2.Resim.RedundantRemoteInputs"), RedundantRemoteInputs, TEXT("How many extra inputs to send with each unreliable network message, to account for packetloss. From server to remote clients. Clamped by NetworkPhysicsComponentConstants::MaxNumberOfElementsToNetwork."));
 		int32 RedundantStates = 0;
-		static FAutoConsoleVariableRef CVarResimRedundantStates(TEXT("np2.Resim.RedundantStates"), RedundantStates, TEXT("How many extra states to send with each unreliable network message, to account for packetloss."));
+		static FAutoConsoleVariableRef CVarResimRedundantStates(TEXT("np2.Resim.RedundantStates"), RedundantStates, TEXT("How many extra states to send with each unreliable network message, to account for packetloss. Clamped by NetworkPhysicsComponentConstants::MaxNumberOfElementsToNetwork."));
 
 		bool bDynamicInputReplicationScalingEnabled = true;
-		static FAutoConsoleVariableRef CVarDynamicInputReplicationScalingEnabled(TEXT("np2.Resim.DynamicInputReplicationScaling.Enabled"), bDynamicInputReplicationScalingEnabled, TEXT("Enable dynmic scaling of number of inputs sent from owning client to the server to account for packet loss. The server will control the value based on how often the server has a hole in its input buffer. NOTE: This overrides np2.Resim.RedundantInputs"));
+		static FAutoConsoleVariableRef CVarDynamicInputReplicationScalingEnabled(TEXT("np2.Resim.DynamicInputReplicationScaling.Enabled"), bDynamicInputReplicationScalingEnabled, TEXT("Enable dynmic scaling of number of inputs sent from owning client to the server to account for packet loss. The server will control the value based on how often the server has a hole in its input buffer. NOTE: This overrides np2.Resim.RedundantInputs. Clamped by NetworkPhysicsComponentConstants::MaxNumberOfElementsToNetwork."));
 		float DynamicInputReplicationScalingMaxInputsPercent = 0.1f;
-		static FAutoConsoleVariableRef CVarDynamicInputReplicationScalingMaxInputsPercent(TEXT("np2.Resim.DynamicInputReplicationScaling.MaxInputsPercent"), DynamicInputReplicationScalingMaxInputsPercent, TEXT("Default 0.1 (= 10%, value in percent as multiplier). Sets the max scalable number of inputs to network from owning client to server as a percentage of the physics fixed tick-rate. 10% of 30Hz = 3 inputs at max."));
+		static FAutoConsoleVariableRef CVarDynamicInputReplicationScalingMaxInputsPercent(TEXT("np2.Resim.DynamicInputReplicationScaling.MaxInputsPercent"), DynamicInputReplicationScalingMaxInputsPercent, TEXT("Default 0.1 (= 10%, value in percent as multiplier). Sets the max scalable number of inputs to network from owning client to server as a percentage of the physics fixed tick-rate. 10% of 30Hz = 3 inputs at max. Clamped by NetworkPhysicsComponentConstants::MaxNumberOfElementsToNetwork."));
 		int32 DynamicInputReplicationScalingMinInputs = 2;
 		static FAutoConsoleVariableRef CVarDynamicInputReplicationScalingMinInputs(TEXT("np2.Resim.DynamicInputReplicationScaling.MinInputs"), DynamicInputReplicationScalingMinInputs, TEXT("Default 2. Sets the minimum scalable number of inputs to network from owning client to server."));
 		float DynamicInputReplicationScalingIncreaseAverageMultiplier = 0.2f;
@@ -147,19 +271,64 @@ namespace Chaos
 	extern CHAOS_API int32 RewindBeforeAdvance;
 }
 
-bool FNetworkPhysicsRewindDataProxy::NetSerializeBase(FArchive& Ar, class UPackageMap* Map, bool& bOutSuccess, TUniqueFunction<TUniquePtr<Chaos::FBaseRewindHistory> ()> CreateHistoryFunction , TUniqueFunction<FNetworkPhysicsData* (const int32)> GetDeltaSourceData)
+bool FNetworkPhysicsRewindDataProxy::NetSerializeBase(FArchive& Ar, class UPackageMap* Map, bool& bOutSuccess
+	, TUniqueFunction<TUniquePtr<Chaos::FBaseRewindHistory> ()> CreateHistoryFunction
+	, TUniqueFunction<const uint32 ()> GetLatestDeltaSourceIndex
+	, TUniqueFunction<FNetworkPhysicsData* (const int32)> GetDeltaSourceData)
 {
 	bDeltaSerializationIssue = false;
-	Ar << Owner;
+	bool bHasData = false;
 
-	bool bHasData = History.IsValid();
-	Ar.SerializeBits(&bHasData, 1);
+	if (Ar.IsSaving())
+	{
+		if (History.IsValid())
+		{
+			bHasData = History->GetHistorySize() > 0;
+		}
 
-	if (bHasData)
+		Ar.SerializeBits(&bHasData, 1);
+
+		if (bHasData)
+		{
+			History->NetSerialize(Ar, Map, [&](void* Data, const int32 DataIndex)
+				{
+					if (FNetworkPhysicsData* NetData = static_cast<FNetworkPhysicsData*>(Data))
+					{
+						if (Owner)
+						{
+							// Set the component pointer to the implementation that uses this data
+							NetData->SetImplementationComponent(Owner.Get()->ActorComponent.Get());
+
+							// Only use stateful delta source for the first entry in history, the following entries will use the previous entry as delta source
+							if (PhysicsReplicationCVars::ResimulationCVars::bEnableStatefulDeltaSerialization && DataIndex == 0 && GetLatestDeltaSourceIndex && GetDeltaSourceData)
+							{
+								// Stateful Delta Serialization
+								{
+									uint32 DeltaSourceIndex = GetLatestDeltaSourceIndex();
+									if (FNetworkPhysicsData* DeltaSourceData = GetDeltaSourceData(DeltaSourceIndex))
+									{
+										NetData->SetDeltaSourceData(DeltaSourceData);
+									}
+									else
+									{
+										ensureMsgf(false, TEXT("Delta Serialization failed to get the latest delta source when sending, should not happen. On the first send the latest index should be populated with a default value, not null."));
+										NetData->SetDeltaSourceData(GetDeltaSourceData(/*Default*/ -2));
+									}
+
+									constexpr uint32 NumBitsDeltaBufferSize = FMath::CeilLogTwo(NetworkPhysicsComponentConstants::DeltaSourceBufferSize);
+									Ar.SerializeBits(&DeltaSourceIndex, NumBitsDeltaBufferSize);
+								}
+							}
+						}
+					}
+				});
+		}
+	}
+	else // IsLoading
 	{
-		if (Ar.IsLoading() && !History.IsValid())
+		if (!History.IsValid())
 		{
-			if(ensureMsgf(Owner, TEXT("FNetRewindDataBase::NetSerialize: owner is null")))
+			if (ensureMsgf(Owner, TEXT("FNetRewindDataBase::NetSerialize: owner is null")))
 			{
 				History = CreateHistoryFunction();
 				if (!ensureMsgf(History.IsValid(), TEXT("FNetRewindDataBase::NetSerialize: failed to create history. Owner: %s"), *GetFullNameSafe(Owner)))
@@ -177,71 +346,52 @@ bool FNetworkPhysicsRewindDataProxy::NetSerializeBase(FArchive& Ar, class UPacka
 			}
 		}
 
-		History->NetSerialize(Ar, Map, [&](void* Data, const int32 DataIndex)
-			{
-				if (FNetworkPhysicsData* NetData = static_cast<FNetworkPhysicsData*>(Data))
+		Ar.SerializeBits(&bHasData, 1);
+
+		if (bHasData)
+		{
+			History->NetSerialize(Ar, Map, [&](void* Data, const int32 DataIndex)
 				{
-					if (Owner)
+					if (FNetworkPhysicsData* NetData = static_cast<FNetworkPhysicsData*>(Data))
 					{
-						// Set the component pointer to the implementation that uses this data
-						NetData->SetImplementationComponent(Owner.Get()->ActorComponent.Get());
-
-						// Only use stateful delta source for the first entry in history, the following entries will use the previous entry as delta source
-						if (PhysicsReplicationCVars::ResimulationCVars::bEnableStatefulDeltaSerialization && DataIndex == 0 && GetDeltaSourceData)
+						if (Owner)
 						{
-							// Stateful Delta Serialization
-							uint32 DeltaSourceFrame = 0;
-							if (Ar.IsLoading())
+							// Set the component pointer to the implementation that uses this data
+							NetData->SetImplementationComponent(Owner.Get()->ActorComponent.Get());
+
+							// Only use stateful delta source for the first entry in history, the following entries will use the previous entry as delta source
+							if (PhysicsReplicationCVars::ResimulationCVars::bEnableStatefulDeltaSerialization && DataIndex == 0 && GetDeltaSourceData)
 							{
-								Ar.SerializeIntPacked(DeltaSourceFrame);
-								
-								FNetworkPhysicsData* DeltaSourceData = nullptr;
-								if (DeltaSourceFrame > 0)
-								{
-									// Try get a valid delta source for frame
-									DeltaSourceData = GetDeltaSourceData(static_cast<int32>(DeltaSourceFrame) - 1);
-								}
-								else
-								{
-									// Sender used default as delta source
-									DeltaSourceData = GetDeltaSourceData(/*Default*/ -2);
-								}
-								
-								if (!DeltaSourceData)
+								// Stateful Delta Serialization
 								{
+									uint32 DeltaSourceIndex = 0;
+									constexpr uint32 NumBitsDeltaBufferSize = FMath::CeilLogTwo(NetworkPhysicsComponentConstants::DeltaSourceBufferSize);
+									Ar.SerializeBits(&DeltaSourceIndex, NumBitsDeltaBufferSize);
+
+									FNetworkPhysicsData* DeltaSourceData = GetDeltaSourceData(static_cast<int32>(DeltaSourceIndex));;
+									if (!DeltaSourceData)
+									{
+										// Sender used default as delta source
+										DeltaSourceData = GetDeltaSourceData(/*Default*/ -2);
+									}
+
+									if (!DeltaSourceData)
+									{
 #if DEBUG_NETWORK_PHYSICS_DELTASERIALIZATION
-									const int32 Frame = static_cast<int32>(DeltaSourceFrame) - 1;
-									UE_LOG(LogChaos, Warning, TEXT("[DEBUG Delta Serialization] %s ISSUE, did not find delta source with frame: %d  --  Name: %s")
-										, (Owner->HasServerWorld() ? TEXT("[SERVER]    ") : (Owner->IsLocallyControlled() ? TEXT("[AUTONOMOUS]") : TEXT("[SIMULATED] "))), Frame, *AActor::GetDebugName(Owner->GetOwner()));
+										UE_LOG(LogChaos, Warning, TEXT("[DEBUG Delta Serialization] %s ISSUE, did not find delta source at index: %d  --  Name: %s")
+											, (Owner->HasServerWorld() ? TEXT("[SERVER]    ") : (Owner->IsLocallyControlled() ? TEXT("[AUTONOMOUS]") : TEXT("[SIMULATED] "))), DeltaSourceIndex, *AActor::GetDebugName(Owner->GetOwner()));
 #endif
-									bDeltaSerializationIssue = true;
-									DeltaSourceData = GetDeltaSourceData(/*Default*/ -2);
-								}
-								
-								NetData->SetDeltaSourceData(DeltaSourceData);
-							}
-							else // IsSaving
-							{
-								if (FNetworkPhysicsData* DeltaSourceData = GetDeltaSourceData(/*Latest*/ -1))
-								{
-									NetData->SetDeltaSourceData(DeltaSourceData);
+										bDeltaSerializationIssue = true;
+										DeltaSourceData = GetDeltaSourceData(/*Default*/ -2);
+									}
 
-									ensure((DeltaSourceData->ServerFrame + 1) >= 0);
-									DeltaSourceFrame = static_cast<uint32>(DeltaSourceData->ServerFrame + 1); // +1 since ServerFrame has a default value of -1 and it needs to be serialized unsigned
-								}
-								else
-								{
-									ensureMsgf(false, TEXT("Delta Serialization failed to get the latest delta source when sending, should not happen. On the first send the latest index should be populated with a default value, not null."));
-									NetData->SetDeltaSourceData(GetDeltaSourceData(/*Default*/ -2));
-									DeltaSourceFrame = 0; // Set DeltaSourceFrame to 0 to indicate that default delta source was used
+									NetData->SetDeltaSourceData(DeltaSourceData);
 								}
-
-								Ar.SerializeIntPacked(DeltaSourceFrame);
 							}
 						}
 					}
-				}
-			});
+				});
+		}
 	}
 
 	return true;
@@ -258,21 +408,56 @@ FNetworkPhysicsRewindDataProxy& FNetworkPhysicsRewindDataProxy::operator=(const
 	return *this;
 }
 
+FNetworkPhysicsRewindDataProxyRPC& FNetworkPhysicsRewindDataProxyRPC::operator=(const FNetworkPhysicsRewindDataProxyRPC& Other)
+{
+	if (&Other != this)
+	{
+		Owner = Other.Owner;
+		History = Other.History ? Other.History->Clone() : nullptr;
+	}
+
+	return *this;
+}
+
+// Replicated Properties, register ReplicationFragment to inject Owner (without replicating it) into FNetworkPhysicsRewindDataProxy 
+UE_NET_IMPLEMENT_NAMED_STRUCT_LASTRESORT_NETSERIALIZER_AND_REGISTRY_DELEGATES(NetworkPhysicsRewindDataOwnerInputProxy, { .CreateAndRegisterReplicationFragmentFunction = UE::Net::CreateAndRegisterNetworkPhysicsRewindDataProxyReplicationFragment })
+UE_NET_IMPLEMENT_NAMED_STRUCT_LASTRESORT_NETSERIALIZER_AND_REGISTRY_DELEGATES(NetworkPhysicsRewindDataRemoteInputProxy, { .CreateAndRegisterReplicationFragmentFunction = UE::Net::CreateAndRegisterNetworkPhysicsRewindDataProxyReplicationFragment })
+UE_NET_IMPLEMENT_NAMED_STRUCT_LASTRESORT_NETSERIALIZER_AND_REGISTRY_DELEGATES(NetworkPhysicsRewindDataStateProxy, { .CreateAndRegisterReplicationFragmentFunction = UE::Net::CreateAndRegisterNetworkPhysicsRewindDataProxyReplicationFragment })
+UE_NET_IMPLEMENT_NAMED_STRUCT_LASTRESORT_NETSERIALIZER_AND_REGISTRY_DELEGATES(NetworkPhysicsRewindDataDeltaSourceInputProxy, { .CreateAndRegisterReplicationFragmentFunction = UE::Net::CreateAndRegisterNetworkPhysicsRewindDataProxyReplicationFragment });
+UE_NET_IMPLEMENT_NAMED_STRUCT_LASTRESORT_NETSERIALIZER_AND_REGISTRY_DELEGATES(NetworkPhysicsRewindDataDeltaSourceStateProxy, { .CreateAndRegisterReplicationFragmentFunction = UE::Net::CreateAndRegisterNetworkPhysicsRewindDataProxyReplicationFragment });
+
+// Replicated RPC Parameters, doesn't support ReplicationFragment so they replicate Owner via FNetworkPhysicsRewindDataProxyRPC
 UE_NET_IMPLEMENT_NAMED_STRUCT_LASTRESORT_NETSERIALIZER_AND_REGISTRY_DELEGATES(NetworkPhysicsRewindDataInputProxy);
-UE_NET_IMPLEMENT_NAMED_STRUCT_LASTRESORT_NETSERIALIZER_AND_REGISTRY_DELEGATES(NetworkPhysicsRewindDataRemoteInputProxy);
-UE_NET_IMPLEMENT_NAMED_STRUCT_LASTRESORT_NETSERIALIZER_AND_REGISTRY_DELEGATES(NetworkPhysicsRewindDataStateProxy);
 UE_NET_IMPLEMENT_NAMED_STRUCT_LASTRESORT_NETSERIALIZER_AND_REGISTRY_DELEGATES(NetworkPhysicsRewindDataImportantInputProxy);
 UE_NET_IMPLEMENT_NAMED_STRUCT_LASTRESORT_NETSERIALIZER_AND_REGISTRY_DELEGATES(NetworkPhysicsRewindDataImportantStateProxy);
-UE_NET_IMPLEMENT_NAMED_STRUCT_LASTRESORT_NETSERIALIZER_AND_REGISTRY_DELEGATES(NetworkPhysicsRewindDataDeltaSourceInputProxy);
-UE_NET_IMPLEMENT_NAMED_STRUCT_LASTRESORT_NETSERIALIZER_AND_REGISTRY_DELEGATES(NetworkPhysicsRewindDataDeltaSourceStateProxy);
 
 bool FNetworkPhysicsRewindDataInputProxy::NetSerialize(FArchive& Ar, class UPackageMap* Map, bool& bOutSuccess)
 {
-	const bool bSuccess = NetSerializeBase(Ar, Map, bOutSuccess, [this]() { return Owner->InputHelper->CreateUniqueRewindHistory(0); }, [this](const int32 Value) -> FNetworkPhysicsData* { return Owner->GetDeltaSourceInput(Value, /*bValueIsIndex*/ false); });
+	const bool bSuccess = NetSerializeBase(Ar, Map, bOutSuccess
+		/*CreateHistoryFunction*/, [this]() { return Owner->InputHelper->CreateUniqueRewindHistory(0); }
+		/*GetLatestDeltaSourceIndex*/, [this]() -> const uint32 { return Owner->GetLatestAcknowledgedDeltaSourceInputIndex(); }
+		/*GetDeltaSourceData*/, [this](const int32 Value) -> FNetworkPhysicsData* { return Owner->GetDeltaSourceInput(Value, /*bValueIsIndex*/ true); });
+
+#if DEBUG_NETWORK_PHYSICS_DELTASERIALIZATION
+	if (bDeltaSerializationIssue)
+	{
+		UE_LOG(LogChaos, Warning, TEXT("		[DEBUG Delta Serialization] INPUT RPC"));
+	}
+#endif
+	return bSuccess;
+}
+
+bool FNetworkPhysicsRewindDataOwnerInputProxy::NetSerialize(FArchive& Ar, class UPackageMap* Map, bool& bOutSuccess)
+{
+	const bool bSuccess = NetSerializeBase(Ar, Map, bOutSuccess
+		/*CreateHistoryFunction*/, [this]() { return Owner->InputHelper->CreateUniqueRewindHistory(0); }
+		/*GetLatestDeltaSourceIndex*/, [this]() -> const uint32 { return Owner->GetLatestAcknowledgedDeltaSourceInputIndex(); }
+		/*GetDeltaSourceData*/, [this](const int32 Value) -> FNetworkPhysicsData* { return Owner->GetDeltaSourceInput(Value, /*bValueIsIndex*/ true); });
+
 #if DEBUG_NETWORK_PHYSICS_DELTASERIALIZATION
 	if (bDeltaSerializationIssue)
 	{
-		UE_LOG(LogChaos, Warning, TEXT("		[DEBUG Delta Serialization] INPUT"));
+		UE_LOG(LogChaos, Warning, TEXT("		[DEBUG Delta Serialization] OWNER INPUT"));
 	}
 #endif
 	return bSuccess;
@@ -280,7 +465,10 @@ bool FNetworkPhysicsRewindDataInputProxy::NetSerialize(FArchive& Ar, class UPack
 
 bool FNetworkPhysicsRewindDataRemoteInputProxy::NetSerialize(FArchive& Ar, class UPackageMap* Map, bool& bOutSuccess)
 {
-	const bool bSuccess = NetSerializeBase(Ar, Map, bOutSuccess, [this]() { return Owner->InputHelper->CreateUniqueRewindHistory(0); }, [this](const int32 Value) -> FNetworkPhysicsData* { return Owner->GetDeltaSourceInput(Value, /*bValueIsIndex*/ false); });
+	const bool bSuccess = NetSerializeBase(Ar, Map, bOutSuccess
+		/*CreateHistoryFunction*/, [this]() { return Owner->InputHelper->CreateUniqueRewindHistory(0); }
+		/*GetLatestDeltaSourceIndex*/, [this]() -> const uint32 { return Owner->GetLatestAcknowledgedDeltaSourceInputIndex(); }
+		/*GetDeltaSourceData*/, [this](const int32 Value) -> FNetworkPhysicsData* { return Owner->GetDeltaSourceInput(Value, /*bValueIsIndex*/ true); });
 
 #if DEBUG_NETWORK_PHYSICS_DELTASERIALIZATION
 	if (bDeltaSerializationIssue) { UE_LOG(LogChaos, Warning, TEXT("		[DEBUG Delta Serialization] REMOTE INPUT")); }
@@ -291,7 +479,10 @@ bool FNetworkPhysicsRewindDataRemoteInputProxy::NetSerialize(FArchive& Ar, class
 
 bool FNetworkPhysicsRewindDataStateProxy::NetSerialize(FArchive& Ar, class UPackageMap* Map, bool& bOutSuccess)
 {
-	const bool bSuccess = NetSerializeBase(Ar, Map, bOutSuccess, [this]() { return Owner->StateHelper->CreateUniqueRewindHistory(0); }, [this](const int32 Value) -> FNetworkPhysicsData* { return Owner->GetDeltaSourceState(Value, /*bValueIsIndex*/ false); });
+	const bool bSuccess = NetSerializeBase(Ar, Map, bOutSuccess
+		/*CreateHistoryFunction*/, [this]() { return Owner->StateHelper->CreateUniqueRewindHistory(0); }
+		/*GetLatestDeltaSourceIndex*/, [this]() -> const uint32 { return Owner->GetLatestAcknowledgedDeltaSourceStateIndex(); }
+		/*GetDeltaSourceData*/, [this](const int32 Value) -> FNetworkPhysicsData* { return Owner->GetDeltaSourceState(Value, /*bValueIsIndex*/ true); });
 
 #if DEBUG_NETWORK_PHYSICS_DELTASERIALIZATION
 	if (bDeltaSerializationIssue) { UE_LOG(LogChaos, Warning, TEXT("		[DEBUG Delta Serialization] STATE")); }
@@ -302,10 +493,13 @@ bool FNetworkPhysicsRewindDataStateProxy::NetSerialize(FArchive& Ar, class UPack
 
 bool FNetworkPhysicsRewindDataImportantInputProxy::NetSerialize(FArchive& Ar, class UPackageMap* Map, bool& bOutSuccess)
 {
-	const bool bSuccess = NetSerializeBase(Ar, Map, bOutSuccess, [this]() { return Owner->InputHelper->CreateUniqueRewindHistory(0); }, [this](const int32 Value) -> FNetworkPhysicsData* { return Owner->GetDeltaSourceInput(Value, /*bValueIsIndex*/ false); });
+	const bool bSuccess = NetSerializeBase(Ar, Map, bOutSuccess
+		/*CreateHistoryFunction*/, [this]() { return Owner->InputHelper->CreateUniqueRewindHistory(0); }
+		/*GetLatestDeltaSourceIndex*/, [this]() -> const uint32 { return Owner->GetLatestAcknowledgedDeltaSourceInputIndex(); }
+		/*GetDeltaSourceData*/, [this](const int32 Value) -> FNetworkPhysicsData* { return Owner->GetDeltaSourceInput(Value, /*bValueIsIndex*/ true); });
 
 #if DEBUG_NETWORK_PHYSICS_DELTASERIALIZATION
-	if (bDeltaSerializationIssue) { UE_LOG(LogChaos, Warning, TEXT("		[DEBUG Delta Serialization] IMPORTANT INPUT")); }
+	if (bDeltaSerializationIssue) { UE_LOG(LogChaos, Warning, TEXT("		[DEBUG Delta Serialization] IMPORTANT INPUT RPC")); }
 #endif
 
 	return bSuccess;
@@ -313,10 +507,13 @@ bool FNetworkPhysicsRewindDataImportantInputProxy::NetSerialize(FArchive& Ar, cl
 
 bool FNetworkPhysicsRewindDataImportantStateProxy::NetSerialize(FArchive& Ar, class UPackageMap* Map, bool& bOutSuccess)
 {
-	const bool bSuccess = NetSerializeBase(Ar, Map, bOutSuccess, [this]() { return Owner->StateHelper->CreateUniqueRewindHistory(0); }, [this](const int32 Value) -> FNetworkPhysicsData* { return Owner->GetDeltaSourceState(Value, /*bValueIsIndex*/ false); });
+	const bool bSuccess = NetSerializeBase(Ar, Map, bOutSuccess
+		/*CreateHistoryFunction*/, [this]() { return Owner->StateHelper->CreateUniqueRewindHistory(0); }
+		/*GetLatestDeltaSourceIndex*/, [this]() -> const uint32 { return Owner->GetLatestAcknowledgedDeltaSourceStateIndex(); }
+		/*GetDeltaSourceData*/, [this](const int32 Value) -> FNetworkPhysicsData* { return Owner->GetDeltaSourceState(Value, /*bValueIsIndex*/ true); });
 
 #if DEBUG_NETWORK_PHYSICS_DELTASERIALIZATION
-	if (bDeltaSerializationIssue) { UE_LOG(LogChaos, Warning, TEXT("		[DEBUG Delta Serialization] IMPORTANT STATE")); }
+	if (bDeltaSerializationIssue) { UE_LOG(LogChaos, Warning, TEXT("		[DEBUG Delta Serialization] IMPORTANT STATE RPC")); }
 #endif
 
 	return bSuccess;
@@ -325,15 +522,24 @@ bool FNetworkPhysicsRewindDataImportantStateProxy::NetSerialize(FArchive& Ar, cl
 bool FNetworkPhysicsRewindDataDeltaSourceInputProxy::NetSerialize(FArchive& Ar, class UPackageMap* Map, bool& bOutSuccess)
 {
 	bool bSuccess = true;
+	constexpr uint32 NumBitsDeltaBufferSize = FMath::CeilLogTwo(NetworkPhysicsComponentConstants::DeltaSourceBufferSize);
+	Ar.SerializeBits(&Index, NumBitsDeltaBufferSize);
+
 	if (PhysicsReplicationCVars::ResimulationCVars::bUseDefaultDeltaForDeltaSourceReplication)
 	{
 		// Use default as base for delta serialization when sending delta source
-		bSuccess = NetSerializeBase(Ar, Map, bOutSuccess, [this]() { return Owner->InputHelper->CreateUniqueRewindHistory(0); }, [this](const int32 Value) -> FNetworkPhysicsData* { return Owner->GetDeltaSourceInput(/*Default*/ -2, /*bValueIsIndex*/ false); });
+		bSuccess = NetSerializeBase(Ar, Map, bOutSuccess
+			/*CreateHistoryFunction*/, [this]() { return Owner->InputHelper->CreateUniqueRewindHistory(0); }
+			/*GetLatestDeltaSourceIndex*/, [this]() -> const uint32 { return Owner->GetLatestAcknowledgedDeltaSourceInputIndex(); }
+			/*GetDeltaSourceData*/, [this](const int32 Value) -> FNetworkPhysicsData* { return Owner->GetDeltaSourceInput(/*Default*/ -2, /*bValueIsIndex*/ true); });
 	}
 	else
 	{
 		// Standard serialization for delta source
-		bSuccess = NetSerializeBase(Ar, Map, bOutSuccess, [this]() { return Owner->InputHelper->CreateUniqueRewindHistory(0); }, nullptr);
+		bSuccess = NetSerializeBase(Ar, Map, bOutSuccess
+			/*CreateHistoryFunction*/, [this]() { return Owner->InputHelper->CreateUniqueRewindHistory(0); }
+			/*GetLatestDeltaSourceIndex*/, nullptr
+			/*GetDeltaSourceData*/, nullptr);
 	}
 
 #if DEBUG_NETWORK_PHYSICS_DELTASERIALIZATION
@@ -346,15 +552,24 @@ bool FNetworkPhysicsRewindDataDeltaSourceInputProxy::NetSerialize(FArchive& Ar,
 bool FNetworkPhysicsRewindDataDeltaSourceStateProxy::NetSerialize(FArchive& Ar, class UPackageMap* Map, bool& bOutSuccess)
 {
 	bool bSuccess = true;
+	constexpr uint32 NumBitsDeltaBufferSize = FMath::CeilLogTwo(NetworkPhysicsComponentConstants::DeltaSourceBufferSize);
+	Ar.SerializeBits(&Index, NumBitsDeltaBufferSize);
+
 	if (PhysicsReplicationCVars::ResimulationCVars::bUseDefaultDeltaForDeltaSourceReplication)
 	{
 		// Use default as base for delta serialization when sending delta source
-		bSuccess = NetSerializeBase(Ar, Map, bOutSuccess, [this]() { return Owner->StateHelper->CreateUniqueRewindHistory(0); }, [this](const int32 Value) -> FNetworkPhysicsData* { return Owner->GetDeltaSourceState(/*Default*/ -2, /*bValueIsIndex*/ false); });
+		bSuccess = NetSerializeBase(Ar, Map, bOutSuccess
+			/*CreateHistoryFunction*/, [this]() { return Owner->StateHelper->CreateUniqueRewindHistory(0); }
+			/*GetLatestDeltaSourceIndex*/, [this]() -> const uint32 { return Owner->GetLatestAcknowledgedDeltaSourceStateIndex(); }
+			/*GetDeltaSourceData*/, [this](const int32 Value) -> FNetworkPhysicsData* { return Owner->GetDeltaSourceState(/*Default*/ -2, /*bValueIsIndex*/ true); });
 	}
 	else
 	{
 		// Standard serialization for delta source
-		bSuccess = NetSerializeBase(Ar, Map, bOutSuccess, [this]() { return Owner->StateHelper->CreateUniqueRewindHistory(0); }, nullptr);
+		bSuccess = NetSerializeBase(Ar, Map, bOutSuccess
+			/*CreateHistoryFunction*/, [this]() { return Owner->StateHelper->CreateUniqueRewindHistory(0); }
+			/*GetLatestDeltaSourceIndex*/, nullptr
+			/*GetDeltaSourceData*/, nullptr);
 	}
 
 #if DEBUG_NETWORK_PHYSICS_DELTASERIALIZATION
@@ -600,9 +815,9 @@ void UNetworkPhysicsComponent::InitializePhysicsReplication()
 			bEnableReliableFlow = SettingsData.NetworkPhysicsComponentSettings.GetEnableReliableFlow();
 			bValidateDataOnGameThread = SettingsData.NetworkPhysicsComponentSettings.GetValidateDataOnGameThread();
 
-			if (ReplicatedInputs.History)
+			if (ReplicatedOwnerInputs.History)
 			{
-				ReplicatedInputs.History->ResizeDataHistory(InputsToNetwork_OwnerDefault);
+				ReplicatedOwnerInputs.History->ResizeDataHistory(InputsToNetwork_OwnerDefault);
 			}
 			if (ReplicatedRemoteInputs.History)
 			{
@@ -766,7 +981,7 @@ void UNetworkPhysicsComponent::GetLifetimeReplicatedProps(TArray< FLifetimePrope
 	// RepGraph / Legacy
 	DOREPLIFETIME_WITH_PARAMS_FAST(UNetworkPhysicsComponent, ReplicatedDeltaSourceInput, ReplicatedParamsAll);
 	DOREPLIFETIME_WITH_PARAMS_FAST(UNetworkPhysicsComponent, ReplicatedDeltaSourceState, ReplicatedParamsAll);
-	DOREPLIFETIME_WITH_PARAMS_FAST(UNetworkPhysicsComponent, ReplicatedInputs, ReplicatedParamsOwner);
+	DOREPLIFETIME_WITH_PARAMS_FAST(UNetworkPhysicsComponent, ReplicatedOwnerInputs, ReplicatedParamsOwner);
 	DOREPLIFETIME_WITH_PARAMS_FAST(UNetworkPhysicsComponent, ReplicatedRemoteInputs, ReplicatedParamsRemote);
 	DOREPLIFETIME_WITH_PARAMS_FAST(UNetworkPhysicsComponent, ReplicatedStates, ReplicatedParamsAll);
 
@@ -783,7 +998,7 @@ void UNetworkPhysicsComponent::PreReplication(IRepChangedPropertyTracker& Change
 	// RepGraph / Legacy
 	DOREPLIFETIME_ACTIVE_OVERRIDE_FAST(UNetworkPhysicsComponent, ReplicatedDeltaSourceInput, bIsUsingLegacyData);
 	DOREPLIFETIME_ACTIVE_OVERRIDE_FAST(UNetworkPhysicsComponent, ReplicatedDeltaSourceState, bIsUsingLegacyData);
-	DOREPLIFETIME_ACTIVE_OVERRIDE_FAST(UNetworkPhysicsComponent, ReplicatedInputs, bIsUsingLegacyData);
+	DOREPLIFETIME_ACTIVE_OVERRIDE_FAST(UNetworkPhysicsComponent, ReplicatedOwnerInputs, bIsUsingLegacyData);
 	DOREPLIFETIME_ACTIVE_OVERRIDE_FAST(UNetworkPhysicsComponent, ReplicatedRemoteInputs, bIsUsingLegacyData);
 	DOREPLIFETIME_ACTIVE_OVERRIDE_FAST(UNetworkPhysicsComponent, ReplicatedStates, bIsUsingLegacyData);
 
@@ -830,15 +1045,16 @@ void UNetworkPhysicsComponent::NetworkMarshaledData()
 	// Replicate source data for input delta serialization
 	auto DeltaSourceInputSyncHelper = [&](const TUniquePtr<Chaos::FBaseRewindHistory>& InputData)
 	{
-		if (bShouldSyncDeltaSourceInput && !bHasSyncedDeltaSourceInput && IsValidNextDeltaSourceInput(InputData->GetLatestFrame()))
+		if (bShouldSyncDeltaSourceInput && !bHasSyncedDeltaSourceInput)
 		{
 			ReplicatedDeltaSourceInput.History->SetRecordDataIncremental(true); // Only record data that is newer than already cached data
 			if (InputData->CopyAllData(*ReplicatedDeltaSourceInput.History, /*bIncludeUnimportant*/ true, /*bIncludeImportant*/ true))
 			{
-				MARK_PROPERTY_DIRTY_FROM_NAME(UNetworkPhysicsComponent, ReplicatedDeltaSourceInput, this);
+				ReplicatedDeltaSourceInput.Index = GetNextDeltaSourceInputIndex();
 				bHasSyncedDeltaSourceInput = true;
 				TimeToSyncDeltaSourceInput = World->GetRealTimeSeconds() + PhysicsReplicationCVars::ResimulationCVars::TimeToSyncStatefulDeltaSource;
 				AddDeltaSourceInput();
+				MARK_PROPERTY_DIRTY_FROM_NAME(UNetworkPhysicsComponent, ReplicatedDeltaSourceInput, this);
 			}
 		}
 	};
@@ -846,15 +1062,16 @@ void UNetworkPhysicsComponent::NetworkMarshaledData()
 	// Replicate source data for state delta serialization
 	auto DeltaSourceStateSyncHelper = [&](const TUniquePtr<Chaos::FBaseRewindHistory>& StateData)
 	{
-		if (bShouldSyncDeltaSourceState && !bHasSyncedDeltaSourceState && IsValidNextDeltaSourceState(StateData->GetLatestFrame()))
+		if (bShouldSyncDeltaSourceState && !bHasSyncedDeltaSourceState)
 		{
 			ReplicatedDeltaSourceState.History->SetRecordDataIncremental(true); // Only record data that is newer than already cached data
 			if (StateData->CopyAllData(*ReplicatedDeltaSourceState.History, /*bIncludeUnimportant*/ true, /*bIncludeImportant*/ true))
 			{
-				MARK_PROPERTY_DIRTY_FROM_NAME(UNetworkPhysicsComponent, ReplicatedDeltaSourceState, this);
+				ReplicatedDeltaSourceState.Index = GetNextDeltaSourceStateIndex();
 				bHasSyncedDeltaSourceState = true;
 				TimeToSyncDeltaSourceState = World->GetRealTimeSeconds() + PhysicsReplicationCVars::ResimulationCVars::TimeToSyncStatefulDeltaSource;
 				AddDeltaSourceState();
+				MARK_PROPERTY_DIRTY_FROM_NAME(UNetworkPhysicsComponent, ReplicatedDeltaSourceState, this);
 			}
 		}
 	};
@@ -912,13 +1129,13 @@ void UNetworkPhysicsComponent::NetworkMarshaledData()
 				if (bIsUsingLegacyData)
 				{
 					// Only replicate data to owning client if bDataAltered is true i.e. the input has been altered by the server
-					ReplicatedInputs.History->ResizeDataHistory(AsyncOutput->InputData->CountAlteredData(/*bIncludeUnimportant*/ true, /*bIncludeImportant*/ bEnableReliableFlow == false), EAllowShrinking::Yes);
-					ReplicatedInputs.History->SetRecordDataIncremental(true); // Only record data that is newer than already cached data
+					ReplicatedOwnerInputs.History->ResizeDataHistory(AsyncOutput->InputData->CountAlteredData(/*bIncludeUnimportant*/ true, /*bIncludeImportant*/ bEnableReliableFlow == false), EAllowShrinking::Yes);
+					ReplicatedOwnerInputs.History->SetRecordDataIncremental(true); // Only record data that is newer than already cached data
 
 					// Server sends inputs through property replication to owning client
-					if (AsyncOutput->InputData->CopyAlteredData(*ReplicatedInputs.History, /*bIncludeUnimportant*/ true, /*bIncludeImportant*/ bEnableReliableFlow == false))
+					if (AsyncOutput->InputData->CopyAlteredData(*ReplicatedOwnerInputs.History, /*bIncludeUnimportant*/ true, /*bIncludeImportant*/ bEnableReliableFlow == false))
 					{
-						MARK_PROPERTY_DIRTY_FROM_NAME(UNetworkPhysicsComponent, ReplicatedInputs, this);
+						MARK_PROPERTY_DIRTY_FROM_NAME(UNetworkPhysicsComponent, ReplicatedOwnerInputs, this);
 					}
 				}
 				else
@@ -936,11 +1153,11 @@ void UNetworkPhysicsComponent::NetworkMarshaledData()
 			{
 				if (bIsUsingLegacyData)
 				{
-					ReplicatedInputs.History->ResizeDataHistory(AsyncOutput->InputData->GetHistorySize(), EAllowShrinking::Yes);
-					if (AsyncOutput->InputData->CopyAllData(*ReplicatedInputs.History, /*bIncludeUnimportant*/ true, /*bIncludeImportant*/ bEnableReliableFlow == false))
+					ReplicateClientInputs.History->ResizeDataHistory(AsyncOutput->InputData->GetHistorySize(), EAllowShrinking::Yes);
+					if (AsyncOutput->InputData->CopyAllData(*ReplicateClientInputs.History, /*bIncludeUnimportant*/ true, /*bIncludeImportant*/ bEnableReliableFlow == false))
 					{
 						// Clients send inputs through an RPC to the server
-						ServerReceiveInputData(ReplicatedInputs);
+						ServerReceiveInputData(ReplicateClientInputs);
 					}
 				}
 				else
@@ -1068,31 +1285,47 @@ void UNetworkPhysicsComponent::OnRep_SetReplicatedDeltaSourceInput()
 	AddDeltaSourceInput();
 }
 
-void UNetworkPhysicsComponent::ServerReceiveDeltaSourceInputFrame_Implementation(const int32 Frame)
+void UNetworkPhysicsComponent::ServerReceiveDeltaSourceInputFrame_Implementation(const uint32 Index, const int32 Frame)
 {
 	ensure(bIsUsingLegacyData);
-	for (int32 I = 0; I < DeltaSourceInputs.Num(); I++)
-	{
-		TUniquePtr<FNetworkPhysicsData>& Data = DeltaSourceInputs[I];
-		if (Data->ServerFrame == Frame)
-		{
-			// Set latest delta source index acknowledged by the client so that we can start using this delta source
-			LatestAcknowledgedDeltaSourceInputIndex = I;
 
+	if (Index >= static_cast<uint32>(DeltaSourceInputs.Num()))
+	{
 #if DEBUG_NETWORK_PHYSICS_DELTASERIALIZATION
-			UE_LOG(LogChaos, Log, TEXT("[DEBUG Delta Serialization] [SERVER]     Owner acknowledged delta source INPUT frame: %d at index: %d  --  Name: %s"), Frame, LatestAcknowledgedDeltaSourceInputIndex, *AActor::GetDebugName(GetOwner()));
+		UE_LOG(LogChaos, Log, TEXT("[DEBUG Delta Serialization] [SERVER]     Owner failed to acknowledged delta source INPUT for frame: %d, index larger than collection: %d >= %d  --  Name: %s"), Frame, Index, DeltaSourceInputs.Num(), *AActor::GetDebugName(GetOwner()));
 #endif
+		return;
+	}
 
-			break;
-		}
+	TUniquePtr<FNetworkPhysicsData>& Data = DeltaSourceInputs[Index];
+	if (Data->ServerFrame == Frame)
+	{
+		// Set latest delta source index acknowledged by the client so that we can start using this delta source
+		LatestAcknowledgedDeltaSourceInputIndex = Index;
+
+#if DEBUG_NETWORK_PHYSICS_DELTASERIALIZATION
+		UE_LOG(LogChaos, Log, TEXT("[DEBUG Delta Serialization] [SERVER]     Owner acknowledged delta source INPUT frame: %d at index: %d  --  Name: %s"), Frame, LatestAcknowledgedDeltaSourceInputIndex, *AActor::GetDebugName(GetOwner()));
+#endif
+	}
+	else
+	{
+#if DEBUG_NETWORK_PHYSICS_DELTASERIALIZATION
+		UE_LOG(LogChaos, Log, TEXT("[DEBUG Delta Serialization] [SERVER]     Owner failed to acknowledged delta source INPUT for Index: %d - AckFrame %d is not equal to CachedFrame: %d  --  Name: %s"), LatestAcknowledgedDeltaSourceInputIndex, Frame, Data->ServerFrame, *AActor::GetDebugName(GetOwner()));
+#endif
 	}
+
 }
 
 void UNetworkPhysicsComponent::AddDeltaSourceInput()
 {
 	// Get the data entry for the correct index in the data sources array
-	const int32 Index = GetDeltaSourceIndexForFrame(ReplicatedDeltaSourceInput.History->GetLatestFrame());
-	check(Index < DeltaSourceInputs.Num());
+	const int32 Index = ReplicatedDeltaSourceInput.Index;
+	if (Index >= DeltaSourceInputs.Num())
+	{
+		check(false);
+		return;
+	}
+
 	FNetworkPhysicsData* PhysicsData = DeltaSourceInputs[Index].Get();
 
 	// Extract the data from the replicated DeltaSources property
@@ -1111,7 +1344,7 @@ void UNetworkPhysicsComponent::AddDeltaSourceInput()
 			if (IsLocallyControlled())
 			{
 				// If this client is the one controlling this entity, send back an acknowledgment to the server we have received the delta source for ServerFrame
-				ServerReceiveDeltaSourceInputFrame(PhysicsData->ServerFrame);
+				ServerReceiveDeltaSourceInputFrame(Index, PhysicsData->ServerFrame);
 			}
 		}
 		else if (IsLocallyControlled())
@@ -1148,31 +1381,47 @@ void UNetworkPhysicsComponent::OnRep_SetReplicatedDeltaSourceState()
 	AddDeltaSourceState();
 }
 
-void UNetworkPhysicsComponent::ServerReceiveDeltaSourceStateFrame_Implementation(const int32 Frame)
+void UNetworkPhysicsComponent::ServerReceiveDeltaSourceStateFrame_Implementation(const uint32 Index, const int32 Frame)
 {
 	ensure(bIsUsingLegacyData);
-	for (int32 I = 0; I < DeltaSourceStates.Num(); I++)
-	{
-		TUniquePtr<FNetworkPhysicsData>& Data = DeltaSourceStates[I];
-		if (Data->ServerFrame == Frame)
-		{
-			// Set latest delta source index acknowledged by the client so that we can start using this delta source
-			LatestAcknowledgedDeltaSourceStateIndex = I;
 
+	if (Index >= static_cast<uint32>(DeltaSourceStates.Num()))
+	{
 #if DEBUG_NETWORK_PHYSICS_DELTASERIALIZATION
-			UE_LOG(LogChaos, Log, TEXT("[DEBUG Delta Serialization] [SERVER]     Owner acknowledged delta source STATE frame: %d at index: %d  --  Name: %s"), Frame, LatestAcknowledgedDeltaSourceStateIndex, *AActor::GetDebugName(GetOwner()));
+		UE_LOG(LogChaos, Log, TEXT("[DEBUG Delta Serialization] [SERVER]     Owner failed to acknowledged delta source STATE for frame: %d, index larger than collection: %d >= %d  --  Name: %s"), Frame, Index, DeltaSourceStates.Num(), *AActor::GetDebugName(GetOwner()));
 #endif
+		return;
+	}
 
-			break;
-		}
+	TUniquePtr<FNetworkPhysicsData>& Data = DeltaSourceStates[Index];
+	if (Data->ServerFrame == Frame)
+	{
+		// Set latest delta source index acknowledged by the client so that we can start using this delta source
+		LatestAcknowledgedDeltaSourceStateIndex = Index;
+
+#if DEBUG_NETWORK_PHYSICS_DELTASERIALIZATION
+		UE_LOG(LogChaos, Log, TEXT("[DEBUG Delta Serialization] [SERVER]     Owner acknowledged delta source STATE frame: %d at index: %d  --  Name: %s"), Frame, LatestAcknowledgedDeltaSourceStateIndex, *AActor::GetDebugName(GetOwner()));
+#endif
+	}
+	else
+	{
+#if DEBUG_NETWORK_PHYSICS_DELTASERIALIZATION
+		UE_LOG(LogChaos, Log, TEXT("[DEBUG Delta Serialization] [SERVER]     Owner failed to acknowledged delta source STATE for Index: %d - AckFrame %d is not equal to CachedFrame: %d --  Name: %s"), LatestAcknowledgedDeltaSourceStateIndex, Frame, Data->ServerFrame, *AActor::GetDebugName(GetOwner()));
+#endif
 	}
 }
 
 void UNetworkPhysicsComponent::AddDeltaSourceState()
 {
 	// Get the data entry for the correct index in the data sources array
-	const int32 Index = GetDeltaSourceIndexForFrame(ReplicatedDeltaSourceState.History->GetLatestFrame());
-	check(Index < DeltaSourceStates.Num());
+	const int32 Index = ReplicatedDeltaSourceState.Index;
+	
+	if (Index >= DeltaSourceStates.Num())
+	{
+		check(false);
+		return;
+	}
+	
 	FNetworkPhysicsData* PhysicsData = DeltaSourceStates[Index].Get();
 
 	// Extract the data from the replicated DeltaSources property
@@ -1192,7 +1441,7 @@ void UNetworkPhysicsComponent::AddDeltaSourceState()
 			if (IsLocallyControlled())
 			{
 				// If this client is the one controlling this entity, send back an acknowledgment to the server we have received the delta source for ServerFrame
-				ServerReceiveDeltaSourceStateFrame(PhysicsData->ServerFrame);
+				ServerReceiveDeltaSourceStateFrame(Index, PhysicsData->ServerFrame);
 			}
 		}
 		else if (IsLocallyControlled())
@@ -1283,7 +1532,7 @@ void UNetworkPhysicsComponent::OnRep_SetReplicatedStates()
 	if (ReplicatedStates.bDeltaSerializationIssue)
 	{
 #if DEBUG_NETWORK_PHYSICS_DELTASERIALIZATION
-		ensureMsgf(!IsLocallyControlled(), TEXT("OnRep_SetReplicatedStates failed delta serialization, should not happen on the owning client unless the pawn just got possessed."));
+		ensureMsgf(!IsLocallyControlled(), TEXT("OnRep_SetReplicatedStates failed delta serialization for locally controlled object, should not happen unless the pawn just got possessed."));
 #endif
 		return;
 	}
@@ -1305,17 +1554,17 @@ void UNetworkPhysicsComponent::OnRep_SetReplicatedStates()
 	}
 }
 
-void UNetworkPhysicsComponent::OnRep_SetReplicatedInputs()
+void UNetworkPhysicsComponent::OnRep_SetReplicatedOwnerInputs()
 {
-	if (ReplicatedInputs.bDeltaSerializationIssue)
+	if (ReplicatedOwnerInputs.bDeltaSerializationIssue)
 	{
 #if DEBUG_NETWORK_PHYSICS_DELTASERIALIZATION
-		ensureMsgf(!IsLocallyControlled(), TEXT("OnRep_SetReplicatedInputs failed delta serialization, should not happen on the owning client unless the pawn just got possessed."));
+		ensureMsgf(!IsLocallyControlled(), TEXT("OnRep_SetReplicatedOwnerInputs failed delta serialization for locally controlled object, should not happen unless the pawn just got possessed."));
 #endif
 		return;
 	}
 
-	if (!NetworkPhysicsComponent_Internal || !InputHelper || !ReplicatedInputs.History)
+	if (!NetworkPhysicsComponent_Internal || !InputHelper || !ReplicatedOwnerInputs.History)
 	{
 		return;
 	}
@@ -1325,10 +1574,10 @@ void UNetworkPhysicsComponent::OnRep_SetReplicatedInputs()
 	{
 		if (!AsyncInput->InputData)
 		{
-			AsyncInput->InputData = InputHelper->CreateUniqueRewindHistory(ReplicatedInputs.History->GetHistorySize());
+			AsyncInput->InputData = InputHelper->CreateUniqueRewindHistory(ReplicatedOwnerInputs.History->GetHistorySize());
 		}
 
-		ReplicatedInputs.History->CopyAllDataGrowingOrdered(*AsyncInput->InputData.Get());
+		ReplicatedOwnerInputs.History->CopyAllDataGrowingOrdered(*AsyncInput->InputData.Get());
 	}
 }
 
@@ -1336,9 +1585,6 @@ void UNetworkPhysicsComponent::OnRep_SetReplicatedRemoteInputs()
 {
 	if (ReplicatedRemoteInputs.bDeltaSerializationIssue)
 	{
-#if DEBUG_NETWORK_PHYSICS_DELTASERIALIZATION
-		ensureMsgf(!IsLocallyControlled(), TEXT("OnRep_SetReplicatedRemoteInputs failed delta serialization, should not happen on the owning client unless the pawn just got possessed."));
-#endif
 		return;
 	}
 
@@ -1436,7 +1682,7 @@ void UNetworkPhysicsComponent::MulticastReceiveImportantInputData_Implementation
 	if (ServerInputs.bDeltaSerializationIssue)
 	{
 #if DEBUG_NETWORK_PHYSICS_DELTASERIALIZATION
-		ensureMsgf(!IsLocallyControlled(), TEXT("MulticastReceiveImportantInputData_Implementation failed delta serialization, should not happen on the owning client unless the pawn just got possessed."));
+		ensureMsgf(!IsLocallyControlled(), TEXT("MulticastReceiveImportantInputData_Implementation failed delta serialization for locally controlled object, should not happen unless the pawn just got possessed."));
 #endif
 		return;
 	}
@@ -1468,7 +1714,7 @@ void UNetworkPhysicsComponent::MulticastReceiveImportantStateData_Implementation
 	if (ServerStates.bDeltaSerializationIssue)
 	{
 #if DEBUG_NETWORK_PHYSICS_DELTASERIALIZATION
-		ensureMsgf(!IsLocallyControlled(), TEXT("MulticastReceiveImportantStateData_Implementation failed delta serialization, should not happen on the owning client unless the pawn just got possessed."));
+		ensureMsgf(!IsLocallyControlled(), TEXT("MulticastReceiveImportantStateData_Implementation failed delta serialization for locally controlled object, should not happen unless the pawn just got possessed."));
 #endif
 		return;
 	}
@@ -3045,7 +3291,8 @@ void FAsyncNetworkPhysicsComponent::UpdateDynamicInputReplicationScaling()
 	{
 		if (TimeSinceLastDynamicScaling > PhysicsReplicationCVars::ResimulationCVars::DynamicInputReplicationScalingIncreaseTimeInterval)
 		{
-			const uint16 MaxInputsValue = static_cast<uint16>(FMath::CeilToInt32(PhysicsReplicationCVars::ResimulationCVars::DynamicInputReplicationScalingMaxInputsPercent / (RigidSolver->GetAsyncDeltaTime())));
+			const int32 MaxScalable = FMath::CeilToInt32(PhysicsReplicationCVars::ResimulationCVars::DynamicInputReplicationScalingMaxInputsPercent / RigidSolver->GetAsyncDeltaTime());
+			const uint16 MaxInputsValue = FMath::Min(NetworkPhysicsComponentConstants::MaxNumberOfElementsToNetwork, static_cast<uint16>(MaxScalable));
 
 			// Increase the amount of inputs the owner sends
 			InputsToNetwork_Owner++;
@@ -3069,7 +3316,8 @@ void FAsyncNetworkPhysicsComponent::UpdateDynamicInputReplicationScaling()
 		DynamicInputReplicationScalingAverageInputs += (static_cast<float>(InputsToNetwork_Owner) - DynamicInputReplicationScalingAverageInputs) * PhysicsReplicationCVars::ResimulationCVars::DynamicInputReplicationScalingDecreaseAverageMultiplier;
 
 		// Clamp to minimum valid value
-		const uint16 MinInputsValue = static_cast<uint16>(FMath::Max(static_cast<uint16>(FMath::RoundToInt(DynamicInputReplicationScalingAverageInputs)), static_cast<uint16>(PhysicsReplicationCVars::ResimulationCVars::DynamicInputReplicationScalingMinInputs)));
+		const int32 MinScalable = FMath::Max(1, PhysicsReplicationCVars::ResimulationCVars::DynamicInputReplicationScalingMinInputs);
+		const uint16 MinInputsValue = static_cast<uint16>(FMath::Max(FMath::RoundToInt(DynamicInputReplicationScalingAverageInputs), MinScalable));
 		InputsToNetwork_Owner = FMath::Max(InputsToNetwork_Owner, MinInputsValue);
 
 		TimeOfLastDynamicInputReplicationScaling = RigidSolver->GetSolverTime();
diff --git a/Engine/Source/Runtime/Engine/Public/Physics/NetworkPhysicsComponent.h b/Engine/Source/Runtime/Engine/Public/Physics/NetworkPhysicsComponent.h
--- a/Engine/Source/Runtime/Engine/Public/Physics/NetworkPhysicsComponent.h
+++ b/Engine/Source/Runtime/Engine/Public/Physics/NetworkPhysicsComponent.h
@@ -13,6 +13,8 @@
 #include "Chaos/PhysicsObject.h"
 #include "Physics/NetworkPhysicsSettingsComponent.h"
 #include "StructUtils/InstancedStruct.h"
+#include "Iris/ReplicationState/PropertyReplicationState.h"
+#include "Iris/ReplicationSystem/ReplicationFragment.h"
 
 #include "NetworkPhysicsComponent.generated.h"
 
@@ -28,6 +30,47 @@ DECLARE_MULTICAST_DELEGATE_OneParam(FOnPostProcessInputsInternal, const int32);
 DECLARE_MULTICAST_DELEGATE_TwoParams(FOnInjectInputsExternal, const int32 /* PhysicsStep */, const int32 /* NumSteps */);
 DECLARE_MULTICAST_DELEGATE_OneParam(FAddResimulationRequestInternal, const int32 /* PhysicsStep */);
 
+
+namespace NetworkPhysicsComponentConstants
+{
+	/** Size of the array caching Delta Sources for delta serialization. */
+	static constexpr int32 DeltaSourceBufferSize = 4;
+
+	/** Max number of elements when replicating inputs and states. Example, latest input along with the last N inputs will be sent in an array to cover for packet loss.
+	* See CVar: np2.Resim.RedundantInputs or np2.Resim.DynamicInputReplicationScaling.Enabled */
+	static constexpr uint16 MaxNumberOfElementsToNetwork = 16;
+}
+
+namespace UE::Net
+{
+	FReplicationFragment* CreateAndRegisterNetworkPhysicsRewindDataProxyReplicationFragment(UObject* Owner, const FReplicationStateDescriptor* Descriptor, FFragmentRegistrationContext& Context);
+
+	/** ReplicationFragment to handle assigning Owner to replicated properties based on FNetworkPhysicsRewindDataProxy
+	* Info: A ReplicationFragment is an extension / override to the NetSerializer logic */
+	class FNetworkPhysicsRewindDataProxyReplicationFragment : public FReplicationFragment
+	{
+	public:
+		FNetworkPhysicsRewindDataProxyReplicationFragment(EReplicationFragmentTraits InTraits, UObject* InOwner, const FReplicationStateDescriptor* InDescriptor);
+
+		void Register(FFragmentRegistrationContext& Context);
+
+	protected:
+		// Set replicated property as dirty if needed
+		virtual bool PollReplicatedState(EReplicationFragmentPollFlags PollOption) override;
+		// On receive, assign Owner to the target struct and then call Dequantize
+		virtual void ApplyReplicatedState(FReplicationStateApplyContext& Context) const override;
+		// Manually call OnRep_ functions
+		virtual void CallRepNotifies(FReplicationStateApplyContext& Context) override;
+		virtual void CollectOwner(FReplicationStateOwnerCollector* Owners) const override;
+
+	private:
+		TRefCountPtr<const FReplicationStateDescriptor> ReplicationStateDescriptor;
+		TUniquePtr<FPropertyReplicationState> SrcReplicationState;
+		UObject* Owner = nullptr;
+	};
+}
+
+
 namespace UE::NetworkPhysicsUtils
 {
 	// Returns the index of the frame about to be simulated, in the server timeline
@@ -251,30 +294,30 @@ struct TNetRewindHistory : public Chaos::TDataRewindHistory<DataType>
 		TNetRewindHistory& NetNewData = static_cast<TNetRewindHistory&>(NewData);
 
 		int32 RewindFrame = INDEX_NONE;
-		DataType& ReceiveData = NetNewData.GetAndLoadEarliestData();
+		DataType* ReceiveData = &NetNewData.GetAndLoadEarliestData();
 
 		for (int32 Itr = 0; Itr < NetNewData.GetHistorySize(); Itr++)
 		{
 			DataType* NextReceiveData = NetNewData.GetAndLoadNextIncrementalData();
 
-			ReceiveData.bImportant = bImportant;
-			ReceiveData.bReceivedData = true; // Received data is marked to differentiate from locally predicted data
+			ReceiveData->bImportant = bImportant;
+			ReceiveData->bReceivedData = true; // Received data is marked to differentiate from locally predicted data
 
-			ReceiveData.LocalFrame = ReceiveData.ServerFrame - FrameOffset;
+			ReceiveData->LocalFrame = ReceiveData->ServerFrame - FrameOffset;
 
-			if (ShouldRecordReceivedDataOnFrame(ReceiveData, NextReceiveData))
+			if (ShouldRecordReceivedDataOnFrame(*ReceiveData, NextReceiveData))
 			{
-				if (CompareDataForRewind && ReceiveData.LocalFrame > RewindFrame && TriggerRewindFromNewData(ReceiveData))
+				if (CompareDataForRewind && ReceiveData->LocalFrame > RewindFrame && TriggerRewindFromNewData(*ReceiveData))
 				{
-					RewindFrame = ReceiveData.LocalFrame;
+					RewindFrame = ReceiveData->LocalFrame;
 				}
 
-				Super::RecordData(ReceiveData.LocalFrame, &ReceiveData);
+				Super::RecordData(ReceiveData->LocalFrame, ReceiveData);
 			}
 
 			if (NextReceiveData)
 			{
-				ReceiveData = *NextReceiveData;
+				ReceiveData = NextReceiveData;
 			}
 			else
 			{
@@ -286,7 +329,7 @@ struct TNetRewindHistory : public Chaos::TDataRewindHistory<DataType>
 						, TryInjectAtFrame, ReceiveData.LocalFrame, ReceiveData.ServerFrame, ReceiveData.bDataAltered, *ReceiveData.DebugData());
 #endif
 
-					DataType InjectData = ReceiveData;
+					DataType InjectData = *ReceiveData;
 					InjectData.bReceivedData = false;
 					InjectData.bDataAltered = true;
 					InjectData.LocalFrame = TryInjectAtFrame;
@@ -365,60 +408,132 @@ struct TNetRewindHistory : public Chaos::TDataRewindHistory<DataType>
 	virtual void NetSerialize(FArchive& Ar, UPackageMap* InPackageMap, TUniqueFunction<void(void* Data, const int32 DataIndex)> DataSetupFunction) override
 	{
 		bool bOneEntry = Super::NumFrames == 1;
-		Ar.SerializeBits(&bOneEntry, 1);
-
-		if (!bOneEntry)
-		{
-			uint32 NumFramesUnsigned = static_cast<uint32>(Super::NumFrames);
-			Ar.SerializeIntPacked(NumFramesUnsigned);
-			Super::NumFrames = static_cast<int32>(NumFramesUnsigned);
-		}
-		else
+		if (Ar.IsSaving())
 		{
-			Super::NumFrames = 1;
-		}
+			Ar.SerializeBits(&bOneEntry, 1);
 
-		if (Super::NumFrames > GetMaxArraySize())
-		{
-			UE_LOG(LogTemp, Warning, TEXT("TNetRewindHistory: serialized array of size %d exceeds maximum size %d."), Super::NumFrames, GetMaxArraySize());
-			Ar.SetError();
-			return;
-		}
+			if (!bOneEntry)
+			{
+				uint32 NumFramesUnsigned = static_cast<uint32>(Super::NumFrames);
 
-		if (Ar.IsLoading())
-		{
-			Super::DataHistory.SetNum(Super::NumFrames);
-		}
+				constexpr uint32 NumBitsNumFramesUnsigned = FMath::CeilLogTwo(NetworkPhysicsComponentConstants::MaxNumberOfElementsToNetwork);
+				Ar.SerializeBits(&NumFramesUnsigned, NumBitsNumFramesUnsigned);
 
-		for (int32 I = 0; I < Super::DataHistory.Num(); I++)
-		{
-			DataType& Data = Super::DataHistory[I];
+				Super::NumFrames = static_cast<int32>(NumFramesUnsigned);
+			}
+			else
+			{
+				Super::NumFrames = 1;
+			}
 
-			// Set the implementation component pointer and stateful delta serialization source
-			if (DataSetupFunction)
+			if (Super::NumFrames > GetMaxArraySize() || Super::NumFrames <= 0)
 			{
-				DataSetupFunction(&Data, I);
+				UE_LOG(LogChaos, Warning, TEXT("TNetRewindHistory: serialized array of size %d exceeds maximum size %d."), Super::NumFrames, GetMaxArraySize());
+				Ar.SetError();
+				return;
 			}
 
-			// Set the internal delta serialization source (between data entries in the collection)
-			if (I > 0)
+			// Note that we use a for-loop to iterate over all entries instead of relying on GetAndLoadNextIncrementalData because the data can have Frame 0 which makes incremental not work.
+			DataType* LastData = nullptr;
+			DataType* CurrentData = &Super::GetAndLoadEarliestData();
+			for (int32 Itr = 0; Itr < Super::GetHistorySize(); Itr++)
 			{
+				// Set the implementation component pointer and stateful delta serialization source
+				if (DataSetupFunction)
+				{
+					DataSetupFunction(CurrentData, Itr);
+				}
+
+				// Set the internal delta serialization source (between data entries in the collection)
 				if constexpr (bLegacyData)
 				{
-					Data.SetDeltaSourceData(&Super::DataHistory[I - 1]);
+					if (LastData)
+					{
+						CurrentData->SetDeltaSourceData(LastData);
+					}
+				}
+
+				NetSerializeData(*CurrentData, Ar, InPackageMap);
+
+				// Get next incremental data
+				if (DataType* NextData = Super::GetAndLoadNextData())
+				{
+					LastData = CurrentData;
+					CurrentData = NextData;
+				}
+				else
+				{
+					ensureMsgf(false, TEXT("TNetRewindHistory: Serialization failed to get the next data."));
 				}
 			}
-		}
 
-		for (DataType& Data : Super::DataHistory)
+			for (DataType& Data : Super::DataHistory)
+			{
+				// Clear delta source and strong pointer to implementation component after serialization
+				if constexpr (bLegacyData)
+				{
+					Data.ClearImplementationComponent();
+					Data.ClearDeltaSourceData();
+				}
+			}
+		}
+		else // IsLoading
 		{
-			NetSerializeData(Data, Ar, InPackageMap);
+			Ar.SerializeBits(&bOneEntry, 1);
+
+			if (!bOneEntry)
+			{
+				uint32 NumFramesUnsigned = static_cast<uint32>(Super::NumFrames);
+
+				constexpr uint32 NumBitsNumFramesUnsigned = FMath::CeilLogTwo(NetworkPhysicsComponentConstants::MaxNumberOfElementsToNetwork);
+				Ar.SerializeBits(&NumFramesUnsigned, NumBitsNumFramesUnsigned);
+
+				Super::NumFrames = static_cast<int32>(NumFramesUnsigned);
+			}
+			else
+			{
+				Super::NumFrames = 1;
+			}
+
+			if (Super::NumFrames > GetMaxArraySize())
+			{
+				UE_LOG(LogChaos, Warning, TEXT("TNetRewindHistory: serialized array of size %d exceeds maximum size %d."), Super::NumFrames, GetMaxArraySize());
+				Ar.SetError();
+				return;
+			}
+
+			Super::DataHistory.SetNum(Super::NumFrames);
+
+			for (int32 I = 0; I < Super::DataHistory.Num(); I++)
+			{
+				DataType& Data = Super::DataHistory[I];
+
+				// Set the implementation component pointer and stateful delta serialization source
+				if (DataSetupFunction)
+				{
+					DataSetupFunction(&Data, I);
+				}
+
+				// Set the internal delta serialization source (between data entries in the collection)
+				if constexpr (bLegacyData)
+				{
+					if (I > 0)
+					{
+						Data.SetDeltaSourceData(&Super::DataHistory[I - 1]);
+					}
+				}
+
+				NetSerializeData(Data, Ar, InPackageMap);
+			}
 
-			// Clear delta source and strong pointer to implementation component after serialization
-			if constexpr (bLegacyData)
+			for (DataType& Data : Super::DataHistory)
 			{
-				Data.ClearImplementationComponent();
-				Data.ClearDeltaSourceData();
+				// Clear delta source and strong pointer to implementation component after serialization
+				if constexpr (bLegacyData)
+				{
+					Data.ClearImplementationComponent();
+					Data.ClearDeltaSourceData();
+				}
 			}
 		}
 
@@ -489,7 +604,7 @@ private :
 		}
 		else
 		{
-			UE_LOG(LogTemp, Error, TEXT("TNetRewindHistory::NetSerializeData called on data struct %s without a native NetSerialize"), *ScriptStruct->GetName());
+			UE_LOG(LogChaos, Error, TEXT("TNetRewindHistory::NetSerializeData called on data struct %s without a native NetSerialize"), *ScriptStruct->GetName());
 
 			// Not working for now since the packagemap could be null
 			// UNetConnection* Connection = CastChecked<UPackageMapClient>(PackageMap)->GetConnection();
@@ -515,17 +630,14 @@ USTRUCT()
 struct FNetworkPhysicsRewindDataProxy
 {
 	GENERATED_BODY()
-
+	virtual ~FNetworkPhysicsRewindDataProxy(){};
 	FNetworkPhysicsRewindDataProxy& operator=(const FNetworkPhysicsRewindDataProxy& Other);
 
 	/** Causes the history to be serialized every time. If implemented, would prevent serializing if the history hasn't changed. */
 	bool operator==(const FNetworkPhysicsRewindDataProxy& Other) const { return false; }
 
 protected:
-	UE_DEPRECATED(5.6, "Deprecated, use the NetSerializeBase call with parameter that takes a function GetDeltaSourceData, pass in nullptr to opt out of implementing a source for delta compression")
-	bool NetSerializeBase(FArchive& Ar, class UPackageMap* Map, bool& bOutSuccess, TUniqueFunction<TUniquePtr<Chaos::FBaseRewindHistory>()> CreateHistoryFunction) { return NetSerializeBase(Ar, Map, bOutSuccess, [&](){ return CreateHistoryFunction(); }, nullptr); };
-
-	bool NetSerializeBase(FArchive& Ar, class UPackageMap* Map, bool& bOutSuccess, TUniqueFunction<TUniquePtr<Chaos::FBaseRewindHistory>()> CreateHistoryFunction, TUniqueFunction<FNetworkPhysicsData*(const int32)> GetDeltaSourceData);
+	virtual bool NetSerializeBase(FArchive& Ar, class UPackageMap* Map, bool& bOutSuccess, TUniqueFunction<TUniquePtr<Chaos::FBaseRewindHistory>()> CreateHistoryFunction, TUniqueFunction<const uint32()> GetLatestDeltaSourceIndex, TUniqueFunction<FNetworkPhysicsData*(const int32)> GetDeltaSourceData);
 
 public:
 	/** The history to be serialized */
@@ -540,10 +652,32 @@ public:
 };
 
 /**
- * Struct suitable for use as a replicated property to replicate input rewind history
+ * Base struct for replicated rewind history RPCs
  */
 USTRUCT()
-struct FNetworkPhysicsRewindDataInputProxy : public FNetworkPhysicsRewindDataProxy
+struct FNetworkPhysicsRewindDataProxyRPC : public FNetworkPhysicsRewindDataProxy
+{
+	GENERATED_BODY()
+	virtual ~FNetworkPhysicsRewindDataProxyRPC() {};
+	FNetworkPhysicsRewindDataProxyRPC& operator=(const FNetworkPhysicsRewindDataProxyRPC& Other);
+
+	/** Causes the history to be serialized every time. If implemented, would prevent serializing if the history hasn't changed. */
+	bool operator==(const FNetworkPhysicsRewindDataProxyRPC& Other) const { return false; }
+
+protected:
+	virtual bool NetSerializeBase(FArchive& Ar, class UPackageMap* Map, bool& bOutSuccess
+		, TUniqueFunction<TUniquePtr<Chaos::FBaseRewindHistory>()> CreateHistoryFunction, TUniqueFunction<const uint32()> GetLatestDeltaSourceIndex, TUniqueFunction<FNetworkPhysicsData* (const int32)> GetDeltaSourceData) override
+	{
+		Ar << Owner;
+		return Super::NetSerializeBase(Ar, Map, bOutSuccess, MoveTemp(CreateHistoryFunction), MoveTemp(GetLatestDeltaSourceIndex), MoveTemp(GetDeltaSourceData));
+	}
+};
+
+/**
+ * Struct suitable for use as a replicated RPC parameter to replicate input rewind history
+ */
+USTRUCT()
+struct FNetworkPhysicsRewindDataInputProxy : public FNetworkPhysicsRewindDataProxyRPC
 {
 	GENERATED_BODY()
 		
@@ -560,6 +694,27 @@ struct TStructOpsTypeTraits<FNetworkPhysicsRewindDataInputProxy> : public TStruc
 	};
 };
 
+/**
+ * Struct suitable for use as a replicated property to replicate input rewind history
+ */
+USTRUCT()
+struct FNetworkPhysicsRewindDataOwnerInputProxy : public FNetworkPhysicsRewindDataProxy
+{
+	GENERATED_BODY()
+		
+	bool NetSerialize(FArchive& Ar, class UPackageMap* Map, bool& bOutSuccess);
+};
+
+template<>
+struct TStructOpsTypeTraits<FNetworkPhysicsRewindDataOwnerInputProxy> : public TStructOpsTypeTraitsBase2<FNetworkPhysicsRewindDataOwnerInputProxy>
+{
+	enum
+	{
+		WithNetSerializer = true,
+		WithIdenticalViaEquality = true
+	};
+};
+
 /**
  * Struct suitable for use as a replicated property to replicate input rewind history
  */
@@ -602,10 +757,10 @@ struct TStructOpsTypeTraits<FNetworkPhysicsRewindDataStateProxy> : public TStruc
 	};
 };
 /**
- * Struct suitable for use as a replicated property to replicate input rewind history
+ * Struct suitable for use as a replicated RPC parameter to replicate input rewind history
  */
 USTRUCT()
-struct FNetworkPhysicsRewindDataImportantInputProxy : public FNetworkPhysicsRewindDataProxy
+struct FNetworkPhysicsRewindDataImportantInputProxy : public FNetworkPhysicsRewindDataProxyRPC
 {
 	GENERATED_BODY()
 		
@@ -623,10 +778,10 @@ struct TStructOpsTypeTraits<FNetworkPhysicsRewindDataImportantInputProxy> : publ
 };
 
 /**
- * Struct suitable for use as a replicated property to replicate state rewind history
+ * Struct suitable for use as a replicated RPC parameter to replicate state rewind history
  */
 USTRUCT()
-struct FNetworkPhysicsRewindDataImportantStateProxy : public FNetworkPhysicsRewindDataProxy
+struct FNetworkPhysicsRewindDataImportantStateProxy : public FNetworkPhysicsRewindDataProxyRPC
 {
 	GENERATED_BODY()
 
@@ -653,6 +808,9 @@ struct FNetworkPhysicsRewindDataDeltaSourceStateProxy : public FNetworkPhysicsRe
 	GENERATED_BODY()
 
 	bool NetSerialize(FArchive& Ar, class UPackageMap* Map, bool& bOutSuccess);
+
+	UPROPERTY()
+	uint32 Index = 0;
 };
 
 template<>
@@ -674,6 +832,9 @@ struct FNetworkPhysicsRewindDataDeltaSourceInputProxy : public FNetworkPhysicsRe
 	GENERATED_BODY()
 
 	bool NetSerialize(FArchive& Ar, class UPackageMap* Map, bool& bOutSuccess);
+
+	UPROPERTY()
+	uint32 Index = 0;
 };
 
 template<>
@@ -1019,7 +1180,27 @@ struct FNetworkPhysicsData : public FNetworkPhysicsPayload
 			uint32 ServerFrameUnsigned = 0;
 
 			bool bIncrememtalFrame = false;
-			if (Ar.IsLoading())
+			if (Ar.IsSaving())
+			{
+				bIncrememtalFrame = ServerFrame == (DeltaSourceData->ServerFrame + 1);
+				{
+					Ar.SerializeBits(&bIncrememtalFrame, 1); // Write if the frame delta is just +1, which is most common for internal deltas
+				}
+
+				if (!bIncrememtalFrame)
+				{
+					int32 FrameDelta = ServerFrame - DeltaSourceData->ServerFrame; // Get the frame delta
+					bool bFrameDeltaNegative = FrameDelta < 0;
+					{
+						Ar.SerializeBits(&bFrameDeltaNegative, 1); // Write if the delta is negative
+					}
+					uint32 FrameDeltaUnsigned = FMath::Abs(FrameDelta);
+					{
+						Ar.SerializeIntPacked(FrameDeltaUnsigned); // Write the unsigned delta frame
+					}
+				}
+			}
+			else // IsLoading
 			{
 				Ar.SerializeBits(&bIncrememtalFrame, 1);
 				if (!bIncrememtalFrame)
@@ -1040,34 +1221,23 @@ struct FNetworkPhysicsData : public FNetworkPhysicsPayload
 
 				LocalFrame = ServerFrame; // Temporarily set LocalFrame to ServerFrame, it will get recalculated later in TRewindHistory::ReceiveNewData
 			}
-			else
-			{
-				bIncrememtalFrame = ServerFrame == (DeltaSourceData->ServerFrame + 1);
-				Ar.SerializeBits(&bIncrememtalFrame, 1); // Write if the frame delta is just +1, which is most common for internal deltas
-				if (!bIncrememtalFrame)
-				{
-					int32 FrameDelta = ServerFrame - DeltaSourceData->ServerFrame; // Get the frame delta
-					bool bFrameDeltaNegative = FrameDelta < 0;
-					Ar.SerializeBits(&bFrameDeltaNegative, 1); // Write if the delta is negative
-					uint32 FrameDeltaUnsigned = FMath::Abs(FrameDelta);
-					Ar.SerializeIntPacked(FrameDeltaUnsigned); // Write the unsigned delta frame
-				}
-			}
 		}
 		else // Standard Serialization
 		{
 			uint32 ServerFrameUnsigned = 0;
-			if (Ar.IsLoading()) // Deserializing
-			{
-				Ar.SerializeIntPacked(ServerFrameUnsigned);
-				ServerFrame = static_cast<int32>(ServerFrameUnsigned) - 1;
-				LocalFrame = ServerFrame; // Temporarily set LocalFrame to ServerFrame, it will get recalculated later in TRewindHistory::ReceiveNewData
-			}
-			else // Serializing
+			if (Ar.IsSaving())
 			{
 				check((ServerFrame + 1) >= 0);
 				ServerFrameUnsigned = static_cast<uint32>(ServerFrame + 1);
+				{
+					Ar.SerializeIntPacked(ServerFrameUnsigned);
+				}
+			}
+			else // IsLoading
+			{
 				Ar.SerializeIntPacked(ServerFrameUnsigned);
+				ServerFrame = static_cast<int32>(ServerFrameUnsigned) - 1;
+				LocalFrame = ServerFrame; // Temporarily set LocalFrame to ServerFrame, it will get recalculated later in TRewindHistory::ReceiveNewData
 			}
 		}
 	}
@@ -1455,11 +1625,11 @@ protected:
 
 	/** Server RPC to acknowledge which Delta Source Input Frame the owning client has received */
 	UFUNCTION(Server, reliable)
-	ENGINE_API void ServerReceiveDeltaSourceInputFrame(const int32 Frame);
+	ENGINE_API void ServerReceiveDeltaSourceInputFrame(const uint32 Index, const int32 Frame);
 	
 	/** Server RPC to acknowledge which Delta Source State Frame the owning client has received */
 	UFUNCTION(Server, reliable)
-	ENGINE_API void ServerReceiveDeltaSourceStateFrame(const int32 Frame);
+	ENGINE_API void ServerReceiveDeltaSourceStateFrame(const uint32 Index, const int32 Frame);
 
 
 	// Server RPC to receive inputs from client
@@ -1478,6 +1648,10 @@ protected:
 	UFUNCTION(NetMulticast, reliable)
 	ENGINE_API void MulticastReceiveImportantStateData(const FNetworkPhysicsRewindDataImportantStateProxy& ServerStates);
 
+	// replicated client inputs for rpc
+	UPROPERTY(Transient)
+	FNetworkPhysicsRewindDataInputProxy ReplicateClientInputs;
+
 	// replicated important physics input
 	UPROPERTY(Transient)
 	FNetworkPhysicsRewindDataImportantInputProxy ReplicatedImportantInput;
@@ -1488,7 +1662,7 @@ protected:
 
 	// repnotify for inputs on owner client
 	UFUNCTION()
-	ENGINE_API void OnRep_SetReplicatedInputs();
+	ENGINE_API void OnRep_SetReplicatedOwnerInputs();
 
 	// repnotify for inputs on remote clients
 	UFUNCTION()
@@ -1499,8 +1673,8 @@ protected:
 	ENGINE_API void OnRep_SetReplicatedStates();
 
 	// replicated physics inputs for owner client
-	UPROPERTY(Transient, ReplicatedUsing = OnRep_SetReplicatedInputs)
-	FNetworkPhysicsRewindDataInputProxy ReplicatedInputs;
+	UPROPERTY(Transient, ReplicatedUsing = OnRep_SetReplicatedOwnerInputs)
+	FNetworkPhysicsRewindDataOwnerInputProxy ReplicatedOwnerInputs;
 
 	// replicated physics inputs for remote clients
 	UPROPERTY(Transient, ReplicatedUsing = OnRep_SetReplicatedRemoteInputs)
@@ -1510,7 +1684,6 @@ protected:
 	UPROPERTY(Transient, ReplicatedUsing = OnRep_SetReplicatedStates)
 	FNetworkPhysicsRewindDataStateProxy ReplicatedStates;
 
-
 	/** Iris Replication */
 	// Server RPC to receive inputs from client
 	UFUNCTION(Server, unreliable)
@@ -1555,26 +1728,30 @@ private:
 	/** Setup data type and history in async input and trigger the creation of data history in the async component */
 	ENGINE_API void CreateAsyncDataHistory();
 
-	/** Set default number of inputs to send over the network with each message, clamped to 1 as minimum 
+	/** Set default number of inputs to send over the network with each message, clamped between 1 and NetworkPhysicsComponentConstants::MaxNumberOfElementsToNetwork.
 	* From owning client to server.
 	* NOTE: this will be overridden if np2.Resim.DynamicInputReplicationScaling.Enabled is enabled. */
-	void SetNumberOfInputsToNetwork(const uint16 NumInputs)
+	void SetNumberOfInputsToNetwork(uint16 NumInputs)
 	{
-		InputsToNetwork_OwnerDefault = FMath::Max(NumInputs, static_cast<uint16>(1));
+		NumInputs = FMath::Clamp(NumInputs, 1, NetworkPhysicsComponentConstants::MaxNumberOfElementsToNetwork);
+		InputsToNetwork_OwnerDefault = NumInputs;
 		InputsToNetwork_Owner = InputsToNetwork_OwnerDefault;
 	}
 
-	/** Set number of inputs to send over the network with each message, clamped to 1 as minimum.
+	/** Set number of inputs to send over the network with each message, clamped between 1 and NetworkPhysicsComponentConstants::MaxNumberOfElementsToNetwork.
 	* From server to remote clients (simulated proxies) */
-	void SetNumberOfRemoteInputsToNetwork(const uint16 NumInputs)
+	void SetNumberOfRemoteInputsToNetwork(uint16 NumInputs)
 	{
-		InputsToNetwork_Simulated = FMath::Max(NumInputs, static_cast<uint16>(1));
+		NumInputs = FMath::Clamp(NumInputs, 1, NetworkPhysicsComponentConstants::MaxNumberOfElementsToNetwork);
+		InputsToNetwork_Simulated = NumInputs;
 	}
 
-	/** Set number of states to send over the network with each message, clamped to 1 as minimum */
-	void SetNumberOfStatesToNetwork(const uint16 NumInputs)
+	/** Set number of states to send over the network with each message, clamped between 1 and NetworkPhysicsComponentConstants::MaxNumberOfElementsToNetwork.
+	* From server to all clients */
+	void SetNumberOfStatesToNetwork(uint16 NumInputs)
 	{
-		StatesToNetwork = FMath::Max(NumInputs, static_cast<uint16>(1));
+		NumInputs = FMath::Clamp(NumInputs, 1, NetworkPhysicsComponentConstants::MaxNumberOfElementsToNetwork);
+		StatesToNetwork = NumInputs;
 	}
 
 	/** Inject the current delta source data into array */
@@ -1583,6 +1760,12 @@ private:
 	/** Inject the current delta source data into array */
 	void AddDeltaSourceState();
 
+	/** Returns the latest acknowledged delta source index  */
+	const uint32 GetLatestAcknowledgedDeltaSourceInputIndex() const { return LatestAcknowledgedDeltaSourceInputIndex; }
+
+	/** Returns the latest acknowledged delta source index  */
+	const uint32 GetLatestAcknowledgedDeltaSourceStateIndex() const { return LatestAcknowledgedDeltaSourceStateIndex; }
+
 	/** Returns the next index in the delta source input array */
 	const int32 GetNextDeltaSourceInputIndex() const { return GetDeltaSourceIndexForFrame(LatestCachedDeltaSourceInputIndex + 1); }
 
@@ -1612,11 +1795,8 @@ public:
 	* Note, passing in @param Value as - 2 returns default data */
 	FNetworkPhysicsData* GetDeltaSourceState(const int32 Value, const bool bValueIsIndexElseFrame);
 
-	/** Size of the array caching Delta Sources for delta serialization. */
-	static const int32 DeltaSourceBufferSize = 10;
-
 	/** Convert frame number to its corresponding index it would hold in the delta sources array */
-	static const int32 GetDeltaSourceIndexForFrame(const int32 Frame) { return FMath::Abs(Frame % DeltaSourceBufferSize); }
+	static const int32 GetDeltaSourceIndexForFrame(const int32 Frame) { return FMath::Abs(Frame % NetworkPhysicsComponentConstants::DeltaSourceBufferSize); }
 
 	/** Register and create both state and input to be both networked and cached in history */
 	template<typename PhysicsTraits>
@@ -1636,6 +1816,7 @@ private:
 
 	friend FNetworkPhysicsCallback;
 	friend struct FNetworkPhysicsRewindDataInputProxy;
+	friend struct FNetworkPhysicsRewindDataOwnerInputProxy;
 	friend struct FNetworkPhysicsRewindDataRemoteInputProxy;
 	friend struct FNetworkPhysicsRewindDataStateProxy;
 	friend struct FNetworkPhysicsRewindDataImportantInputProxy;
@@ -1669,17 +1850,17 @@ private:
 	// Local default state data in legacy type
 	TUniquePtr<FNetworkPhysicsData> StateDataDefault_Legacy;
 
-	// The number of inputs the owning client should send to the server with each RPC, replicated from the server. This is dynamically scaled based on when there are holes in the inputs buffer if np2.Resim.DynamicInputReplicationScaling.Enabled is enabled
+	// The number of inputs the owning client should send to the server with each RPC, replicated from the server. This is dynamically scaled based on when there are holes in the inputs buffer if np2.Resim.DynamicInputReplicationScaling.Enabled is enabled. Clamped by NetworkPhysicsComponentConstants::MaxNumberOfElementsToNetwork.
 	UPROPERTY( Replicated )
 	uint16 InputsToNetwork_Owner = 3;
 
-	// The default value for number for InputsToNetwork_Owner, acts as the initial value and the cap when dynamically adjusting InputsToNetwork_Owner
+	// The default value for number for InputsToNetwork_Owner, acts as the initial value and the cap when dynamically adjusting InputsToNetwork_Owner. Clamped by NetworkPhysicsComponentConstants::MaxNumberOfElementsToNetwork.
 	uint16 InputsToNetwork_OwnerDefault = 3;
 
-	// Send last N number of inputs each replication call from server to remote clients
+	// Send last N number of inputs each replication call from server to remote clients. Clamped by NetworkPhysicsComponentConstants::MaxNumberOfElementsToNetwork.
 	uint16 InputsToNetwork_Simulated = 2;
 
-	// Send last N number of states each replication call from server to remote clients
+	// Send last N number of states each replication call from server to remote clients. Clamped by NetworkPhysicsComponentConstants::MaxNumberOfElementsToNetwork.
 	uint16 StatesToNetwork = 1;
 
 	// Array of delta sources, used as a base for delta serialization
@@ -1717,15 +1898,6 @@ private:
 	bool bEnableUnreliableFlow = true;
 	bool bEnableReliableFlow = false;
 	bool bValidateDataOnGameThread = false;
-
-	/** ----- Deprecated API ----- */
-private:
-#if WITH_EDITORONLY_DATA
-	UE_DEPRECATED(5.6, "Deprecated, use InputsToNetwork_Owner instead")
-		uint16 InputsToNetwork = 3;
-	UE_DEPRECATED(5.6, "Deprecated, use InputsToNetwork_Simulated instead")
-		uint16 InputsToNetworkRemote = 2;
-#endif
 };
 
 template<typename Input, typename State>
@@ -1754,14 +1926,17 @@ inline void UNetworkPhysicsComponent::CreateDataHistory(UActorComponent* History
 	StateDataDefault_Legacy = MakeUnique<typename PhysicsTraits::StatesType>();
 
 	// Initialize delta source arrays
-	for (int32 I = 0; I < DeltaSourceBufferSize; I++)
+	for (int32 I = 0; I < NetworkPhysicsComponentConstants::DeltaSourceBufferSize; I++)
 	{
 		DeltaSourceInputs.Add(MakeUnique<typename PhysicsTraits::InputsType>());
 		DeltaSourceStates.Add(MakeUnique<typename PhysicsTraits::StatesType>());
 	}
 
-	ReplicatedInputs.History = InputHelper->CreateUniqueRewindHistory(InputsToNetwork_OwnerDefault);
-	ReplicatedInputs.Owner = this;
+	ReplicateClientInputs.History = InputHelper->CreateUniqueRewindHistory(InputsToNetwork_OwnerDefault);
+	ReplicateClientInputs.Owner = this;
+
+	ReplicatedOwnerInputs.History = InputHelper->CreateUniqueRewindHistory(InputsToNetwork_OwnerDefault);
+	ReplicatedOwnerInputs.Owner = this;
 
 	ReplicatedRemoteInputs.History = InputHelper->CreateUniqueRewindHistory(InputsToNetwork_Simulated);
 	ReplicatedRemoteInputs.Owner = this;
@@ -1796,13 +1971,16 @@ inline void UNetworkPhysicsComponent::CreateInputHistory(UActorComponent* Histor
 	InputDataDefault_Legacy = MakeUnique<InputsType>();
 
 	// Initialize delta source array
-	for (int32 I = 0; I < DeltaSourceBufferSize; I++)
+	for (int32 I = 0; I < NetworkPhysicsComponentConstants::DeltaSourceBufferSize; I++)
 	{
 		DeltaSourceInputs.Add(MakeUnique<InputsType>());
 	}
 
-	ReplicatedInputs.History = InputHelper->CreateUniqueRewindHistory(InputsToNetwork_OwnerDefault);
-	ReplicatedInputs.Owner = this;
+	ReplicateClientInputs.History = InputHelper->CreateUniqueRewindHistory(InputsToNetwork_OwnerDefault);
+	ReplicateClientInputs.Owner = this;
+
+	ReplicatedOwnerInputs.History = InputHelper->CreateUniqueRewindHistory(InputsToNetwork_OwnerDefault);
+	ReplicatedOwnerInputs.Owner = this;
 
 	ReplicatedRemoteInputs.History = InputHelper->CreateUniqueRewindHistory(InputsToNetwork_Simulated);
 	ReplicatedRemoteInputs.Owner = this;
diff --git a/Engine/Source/Runtime/Experimental/Chaos/Public/RewindData.h b/Engine/Source/Runtime/Experimental/Chaos/Public/RewindData.h
--- a/Engine/Source/Runtime/Experimental/Chaos/Public/RewindData.h
+++ b/Engine/Source/Runtime/Experimental/Chaos/Public/RewindData.h
@@ -478,6 +478,21 @@ public :
 		return nullptr;
 	}
 
+	/** Loads and returns the data on the next index */
+	DataType* GetAndLoadNextData()
+	{
+		if (NumFrames > 0)
+		{
+			const int32 NextIndex = GetFrameIndex(CurrentIndex + 1);
+			CurrentIndex = NextIndex;
+			CurrentFrame = DataHistory[NextIndex].LocalFrame;
+
+			return &DataHistory[NextIndex];
+		}
+
+		return nullptr;
+	}
+
 	/** Get the number of valid data in the buffer index range */
 	FORCEINLINE uint32 NumValidData(const uint32 StartFrame, const uint32 EndFrame) const
 	{
