Always preload animation sequence dependencies as MT async loading requires manual handling

#Synced-CL 51094291

diff --git a/Engine/Source/Runtime/Engine/Private/Animation/AnimCompressionTypes.cpp b/Engine/Source/Runtime/Engine/Private/Animation/AnimCompressionTypes.cpp
--- a/Engine/Source/Runtime/Engine/Private/Animation/AnimCompressionTypes.cpp
+++ b/Engine/Source/Runtime/Engine/Private/Animation/AnimCompressionTypes.cpp
@@ -1219,10 +1219,11 @@ template void FUECompressedAnimData::ByteSwapData(TArrayView<uint8> CompressedDa
 
 void ValidateUObjectLoaded(UObject* Obj, UObject* Source)
 {
-#if WITH_EDITOR
+	// Multi-threaded async loading does not enforce pre-load dependencies specified during cook
+	// and instead requires explicit dependencies
 	Obj->ConditionalPreload();
-#endif
-	checkf(!Obj->HasAnyFlags(RF_NeedLoad), TEXT("Failed to load %s in %s"), *Obj->GetFullName(), *Source->GetFullName()); // in non editor should have been preloaded by GetPreloadDependencies
+
+	checkf(!Obj->HasAnyFlags(RF_NeedLoad), TEXT("Failed to load %s in %s"), *Obj->GetFullName(), *Source->GetFullName());
 }
 
 void FUECompressedAnimDataMutable::BuildFinalBuffer(TArray<uint8>& OutCompressedByteStream)
@@ -1396,6 +1397,9 @@ void FCompressedAnimSequence::SerializeCompressedData(FArchive& Ar, bool bDDCDat
 		Ar.Serialize(CompressedCurveByteStream.GetData(), NumCurveBytes);
 
 		// Lookup our codecs in our settings assets
+		// We need them to be fully loaded in order to be able to look up the codec used to compress our data within
+		// The codec typically lives with the compressed data which ends up in the DDC, as such we cannot maintain
+		// a stable UObject reference. We use a lookup mechanism instead.
 		ValidateUObjectLoaded(BoneCompressionSettings, DataOwner);
 		ValidateUObjectLoaded(CurveCompressionSettings, DataOwner);
 		BoneCompressionCodec = BoneCompressionSettings->GetCodec(BoneCodecDDCHandle);
