[MegaLights] Fix PPV properties not being applied in FPostProcessUtils::OverridePostProcessSettings.
- this file is used by GameplayCameras plugin (the main PPV blending is implemented in SceneView.cpp).

#Synced-CL 51016705

diff --git a/Engine/Source/Runtime/Engine/Private/PostProcessUtils.cpp b/Engine/Source/Runtime/Engine/Private/PostProcessUtils.cpp
--- a/Engine/Source/Runtime/Engine/Private/PostProcessUtils.cpp
+++ b/Engine/Source/Runtime/Engine/Private/PostProcessUtils.cpp
@@ -229,6 +229,9 @@ bool FPostProcessUtils::OverridePostProcessSettings(FPostProcessSettings& ThisFr
 		UE_SET_PP(RayTracingAOIntensity);
 		UE_SET_PP(RayTracingAORadius);
 
+		UE_SET_PP(bMegaLights);
+		UE_SET_PP(MegaLightsFrontLayerTranslucency);
+
 		UE_SET_PP(PathTracingMaxBounces);
 		UE_SET_PP(PathTracingSamplesPerPixel);
 		UE_SET_PP(PathTracingMaxPathIntensity);
@@ -550,6 +553,9 @@ bool FPostProcessUtils::BlendPostProcessSettings(FPostProcessSettings& ThisFrom,
 		UE_SET_PP(RayTracingAOIntensity);
 		UE_SET_PP(RayTracingAORadius);
 
+		UE_SET_PP(bMegaLights);
+		UE_SET_PP(MegaLightsFrontLayerTranslucency);
+
 		UE_SET_PP(PathTracingMaxBounces);
 		UE_SET_PP(PathTracingSamplesPerPixel);
 		UE_LERP_PP(PathTracingMaxPathIntensity);
