Refactored spline mesh to not need to pass information through the FMeshBatchElement, instead using the LooseParametersUniformBuffer to pass the information to the shader.
- Changed VF to use (EVertexFactoryFlags::SupportsCachingMeshDrawCommands)
- (also removed deprecated TessellationDisablingShadowMapMeshSize in FMeshBatch)

#Synced-CL 51090859

diff --git a/Engine/Shaders/Private/LocalVertexFactory.ush b/Engine/Shaders/Private/LocalVertexFactory.ush
--- a/Engine/Shaders/Private/LocalVertexFactory.ush
+++ b/Engine/Shaders/Private/LocalVertexFactory.ush
@@ -361,10 +361,9 @@ FSplineMeshShaderParams GetSplineMeshParams(FInstanceSceneData InstanceData)
 #else
 
 // These currently still need to be bound loosely when not using GPUScene
-float4 SplineParams[SPLINE_MESH_PARAMS_FLOAT4_SIZE];
 FSplineMeshShaderParams GetSplineMeshParams(FInstanceSceneData InstanceData)
 {
-	return UnpackSplineMeshParams(SplineParams);
+	return UnpackSplineMeshParams(SplineMeshVFPrimitiveShaderParameters.SplineParams);
 }
 
 #endif // VF_USE_PRIMITIVE_SCENE_DATA && FEATURE_LEVEL > FEATURE_LEVEL_ES3_1
diff --git a/Engine/Source/Runtime/Engine/Private/Components/SplineMeshComponent.cpp b/Engine/Source/Runtime/Engine/Private/Components/SplineMeshComponent.cpp
--- a/Engine/Source/Runtime/Engine/Private/Components/SplineMeshComponent.cpp
+++ b/Engine/Source/Runtime/Engine/Private/Components/SplineMeshComponent.cpp
@@ -122,10 +122,6 @@ FSplineMeshInstanceData::FSplineMeshInstanceData(const USplineMeshComponent* Sou
 //////////////////////////////////////////////////////////////////////////
 // FSplineMeshVertexFactoryShaderParameters
 
-void FSplineMeshVertexFactoryShaderParameters::Bind(const FShaderParameterMap& ParameterMap)
-{
-	SplineMeshParams.Bind(ParameterMap, TEXT("SplineParams"), SPF_Optional);
-}
 
 void FSplineMeshVertexFactoryShaderParameters::GetElementShaderBindings(
 	const class FSceneInterface* Scene,
@@ -163,19 +159,12 @@ void FSplineMeshVertexFactoryShaderParameters::GetElementShaderBindings(
 		ShaderBindings.Add(Shader->GetUniformBufferParameter<FLocalVertexFactoryUniformShaderParameters>(), VertexFactoryUniformBuffer != nullptr ? VertexFactoryUniformBuffer : LocalVertexFactory->GetUniformBuffer());
 	}
 
-	// If we can't use GPU Scene instance data, we have to bind the params to the VS loosely
-	// NOTE: Mobile GPU scene can't support loading the spline params from instance data in VS
-	if (!bUseGPUScene || FeatureLevel == ERHIFeatureLevel::ES3_1)
+	if (BatchElement.LooseParametersUniformBuffer)
 	{
-		checkSlow(BatchElement.bIsSplineProxy);
-		const FSplineMeshSceneProxy* SplineProxy = BatchElement.SplineMeshSceneProxy;
-
-		FVector4f ParamData[SPLINE_MESH_PARAMS_FLOAT4_SIZE];
-		PackSplineMeshParams(SplineProxy->GetSplineMeshParams(), TArrayView<FVector4f>(ParamData, SPLINE_MESH_PARAMS_FLOAT4_SIZE));
+		ShaderBindings.Add(Shader->GetUniformBufferParameter<FSplineMeshVFPrimitiveShaderParameters>(), BatchElement.LooseParametersUniformBuffer);
+	}
 
-		ShaderBindings.Add(SplineMeshParams, ParamData);
 	}
-}
 
 //////////////////////////////////////////////////////////////////////////
 // SplineMeshVertexFactory
@@ -197,6 +186,7 @@ IMPLEMENT_VERTEX_FACTORY_TYPE(FSplineMeshVertexFactory, "/Engine/Private/LocalVe
 	| EVertexFactoryFlags::SupportsRayTracing
 	| EVertexFactoryFlags::SupportsRayTracingDynamicGeometry
 	| EVertexFactoryFlags::SupportsManualVertexFetch
+	| EVertexFactoryFlags::SupportsCachingMeshDrawCommands
 );
 
 void InitSplineMeshVertexFactoryComponents(
diff --git a/Engine/Source/Runtime/Engine/Private/SplineMeshSceneProxy.cpp b/Engine/Source/Runtime/Engine/Private/SplineMeshSceneProxy.cpp
--- a/Engine/Source/Runtime/Engine/Private/SplineMeshSceneProxy.cpp
+++ b/Engine/Source/Runtime/Engine/Private/SplineMeshSceneProxy.cpp
@@ -15,6 +15,8 @@
 
 IMPLEMENT_TYPE_LAYOUT(FSplineMeshVertexFactoryShaderParameters);
 
+IMPLEMENT_GLOBAL_SHADER_PARAMETER_STRUCT(FSplineMeshVFPrimitiveShaderParameters, "SplineMeshVFPrimitiveShaderParameters");
+
 static TAutoConsoleVariable<int32> CVarRayTracingSplineMeshes(
 	TEXT("r.RayTracing.Geometry.SplineMeshes"),
 	1,
@@ -49,6 +51,43 @@ void FSplineMeshVertexFactory::GetPSOPrecacheVertexFetchElements(EVertexInputStr
 	FLocalVertexFactory::GetPSOPrecacheVertexFetchElements(VertexInputStreamType, Elements);
 }
 
+void FSplineMeshSceneProxyCommon::SetSplineTextureCoord_RenderThread(uint32 SplineIndex, FUintVector2 TexCoord)
+{
+	check(SplineIndex == 0); // currently only support one spline type, but could do instanced later
+	SplineParams.TextureCoord = TexCoord;
+	RepackSplineMeshParamsInternal();
+}
+	
+void FSplineMeshSceneProxyCommon::UpdateSplineMeshParams_RenderThread(const FSplineMeshShaderParams& Params)
+{
+	// Don't replace the texture coord that was assigned to us in SetSplineTextureCoord
+	FUintVector2 TexCoord = SplineParams.TextureCoord;
+	SplineParams = Params;
+	SplineParams.TextureCoord = TexCoord;
+	RepackSplineMeshParamsInternal();
+}
+
+void FSplineMeshSceneProxyCommon::RepackSplineMeshParamsInternal()
+{
+	auto& SceneProxy = GetProxy();
+		
+	// NOTE: If the payload extension was not initialized, we're probably not using GPU Scene
+	if (SplineMeshInstanceData.Update(SplineParams))
+	{
+		// Request a GPU Scene update for this primitive so it updates its instance data
+		SceneProxy.GetScene().RequestGPUSceneUpdate(*SceneProxy.GetPrimitiveSceneInfo(), EPrimitiveDirtyState::ChangedOther);
+	}
+
+	if (SplineMeshVFPrimitiveShaderParametersUB)
+	{
+		FSplineMeshVFPrimitiveShaderParameters SplineMeshVFPrimitiveShaderParameters;
+		PackSplineMeshParams(SplineParams, SplineMeshVFPrimitiveShaderParameters.SplineParams);
+		SplineMeshVFPrimitiveShaderParametersUB.UpdateUniformBufferImmediate(FRHICommandListImmediate::Get(), SplineMeshVFPrimitiveShaderParameters);
+	}
+
+	bUpdateRayTracingGeometry = true;
+}
+
 FSplineMeshSceneProxy::FSplineMeshSceneProxy(USplineMeshComponent* InComponent) :
 	FSplineMeshSceneProxy(FStaticMeshSceneProxyDesc(InComponent), FSplineMeshSceneProxyDesc(InComponent))
 {}
@@ -95,6 +134,12 @@ FSplineMeshSceneProxy::FSplineMeshSceneProxy(const FStaticMeshSceneProxyDesc& In
 		SplineMeshInstanceData.Setup(SplineParams);
 		SetupInstanceSceneDataBuffers(&SplineMeshInstanceData);
 	}
+	else
+	{
+		FSplineMeshVFPrimitiveShaderParameters SplineMeshVFPrimitiveShaderParameters;
+		PackSplineMeshParams(SplineParams, SplineMeshVFPrimitiveShaderParameters.SplineParams);
+		SplineMeshVFPrimitiveShaderParametersUB = CreateUniformBufferImmediate(SplineMeshVFPrimitiveShaderParameters, EUniformBufferUsage::UniformBuffer_MultiFrame);
+	}
 
 	for (int32 LODIndex = 0; LODIndex < LODs.Num(); LODIndex++)
 	{
@@ -129,9 +174,8 @@ void FSplineMeshSceneProxy::SetupMeshBatchForSpline(int32 InLODIndex, FMeshBatch
 	check(OutMeshBatch.Elements.Num() == 1);
 	OutMeshBatch.VertexFactory = OutMeshBatch.Elements[0].bUserDataIsColorVertexBuffer ? VFs.SplineVertexFactoryOverrideColorVertexBuffer : VFs.SplineVertexFactory;
 	check(OutMeshBatch.VertexFactory);
-	OutMeshBatch.Elements[0].SplineMeshSceneProxy = this;
-	OutMeshBatch.Elements[0].bIsSplineProxy = true;
 	OutMeshBatch.Elements[0].PrimitiveUniformBuffer = GetUniformBuffer();
+	OutMeshBatch.Elements[0].LooseParametersUniformBuffer = SplineMeshVFPrimitiveShaderParametersUB;
 	OutMeshBatch.ReverseCulling ^= (SplineParams.StartScale.X < 0) ^ (SplineParams.StartScale.Y < 0);
 }
 
@@ -141,9 +185,8 @@ void FSplineMeshSceneProxy::SetupRayTracingMeshBatchForSpline(int32 InLODIndex,
 	check(OutMeshBatch.Elements.Num() == 1);
 	OutMeshBatch.VertexFactory = OutMeshBatch.Elements[0].bUserDataIsColorVertexBuffer ? VFs.SplineVertexFactoryOverrideColorVertexBuffer : VFs.SplineVertexFactory;
 	check(OutMeshBatch.VertexFactory);
-	OutMeshBatch.Elements[0].SplineMeshSceneProxy = this;
-	OutMeshBatch.Elements[0].bIsSplineProxy = true;
 	OutMeshBatch.Elements[0].PrimitiveUniformBuffer = GetUniformBuffer();
+	OutMeshBatch.Elements[0].LooseParametersUniformBuffer = SplineMeshVFPrimitiveShaderParametersUB;
 	OutMeshBatch.ReverseCulling ^= (SplineParams.StartScale.X < 0) ^ (SplineParams.StartScale.Y < 0);
 }
 
diff --git a/Engine/Source/Runtime/Engine/Public/MeshBatch.h b/Engine/Source/Runtime/Engine/Public/MeshBatch.h
--- a/Engine/Source/Runtime/Engine/Public/MeshBatch.h
+++ b/Engine/Source/Runtime/Engine/Public/MeshBatch.h
@@ -254,13 +254,9 @@ struct FMeshBatchElement
 	*/
 	FMeshBatchElementDynamicIndexBuffer DynamicIndexBuffer;
 
-	union 
-	{
-		/** If !bIsSplineProxy, Instance runs, where number of runs is specified by NumInstances.  Run structure is [StartInstanceIndex, EndInstanceIndex]. */
-		uint32* InstanceRuns;
-		/** If bIsSplineProxy, a pointer back to the proxy */
-		const class FSplineMeshSceneProxy* SplineMeshSceneProxy;
-	};
+	/** Instance runs, where number of runs is specified by NumInstances.  Run structure is [StartInstanceIndex, EndInstanceIndex]. */
+	uint32* InstanceRuns;
+
 	const void* UserData;
 
 	// Meaning depends on the vertex factory, e.g. FGPUSkinPassthroughVertexFactory: element index in FGPUSkinCache::CachedElements
@@ -288,7 +284,6 @@ struct FMeshBatchElement
 	uint32 InstancedLODIndex : 4;
 	uint32 InstancedLODRange : 4;
 	uint32 bUserDataIsColorVertexBuffer : 1;
-	uint32 bIsSplineProxy : 1;
 	uint32 bIsInstanceRuns : 1;
 	uint32 bForceInstanceCulling : 1;
 	uint32 bPreserveInstanceOrder : 1;
@@ -344,7 +339,6 @@ struct FMeshBatchElement
 	,	InstancedLODIndex(0)
 	,	InstancedLODRange(0)
 	,	bUserDataIsColorVertexBuffer(false)
-	,	bIsSplineProxy(false)
 	,	bIsInstanceRuns(false)
 	,	bForceInstanceCulling(false)
 	,	bPreserveInstanceOrder(false)
@@ -386,10 +380,6 @@ struct FMeshBatch
 	/** If set, the ID to use during culling when GPU skin cache is available to provide dynamic tight-fit bounds. */
 	int32 DynamicMeshBoundsIndex = INDEX_NONE;
 
-	/** This is the threshold that will be used to know if we should use this mesh batch or use one with no tessellation enabled */
-	UE_DEPRECATED(5.6, "This is not used by the renderer.")
-	float TessellationDisablingShadowMapMeshSize;
-
 	/** Segment index of the original mesh */
 	int32 SegmentIndex;
 
@@ -512,15 +502,11 @@ struct FMeshBatch
 	inline bool Validate(const FPrimitiveSceneProxy* SceneProxy, ERHIFeatureLevel::Type FeatureLevel) const { return true; }
 #endif
 
-	// Disable deprecation warnings for TessellationDisablingShadowMapMeshSize.
-	PRAGMA_DISABLE_DEPRECATION_WARNINGS
-
 	/** Default constructor. */
 	FMeshBatch()
 	:	VertexFactory(nullptr)
 	,	MaterialRenderProxy(nullptr)
 	,	LCI(nullptr)
-	,	TessellationDisablingShadowMapMeshSize(0.0f)
 	,	SegmentIndex(INDEX_NONE)
 	,	MeshIdInPrimitive(0)
 	,	LODIndex(INDEX_NONE)
@@ -558,16 +544,6 @@ struct FMeshBatch
 		// By default always add the first element.
 		Elements.AddDefaulted();
 	}
-
-	// Need to do these explicitly as clang otherwise produces deprecation warnings for the _implicit_ ctors & operators
-	FMeshBatch(FMeshBatch&) = default;
-	~FMeshBatch() = default;
-	FMeshBatch(const FMeshBatch&) = default;
-	FMeshBatch(FMeshBatch&&) = default;
-	FMeshBatch& operator=(const FMeshBatch&) = default;
-	FMeshBatch& operator=(FMeshBatch&&) = default;
-
-	PRAGMA_ENABLE_DEPRECATION_WARNINGS
 };
 
 struct FUniformBufferValue
diff --git a/Engine/Source/Runtime/Engine/Public/SplineMeshSceneProxy.h b/Engine/Source/Runtime/Engine/Public/SplineMeshSceneProxy.h
--- a/Engine/Source/Runtime/Engine/Public/SplineMeshSceneProxy.h
+++ b/Engine/Source/Runtime/Engine/Public/SplineMeshSceneProxy.h
@@ -17,6 +17,10 @@
 
 struct FSplineMeshSceneProxyDesc;
 
+BEGIN_GLOBAL_SHADER_PARAMETER_STRUCT(FSplineMeshVFPrimitiveShaderParameters, ENGINE_API)
+	SHADER_PARAMETER_ARRAY(FVector4f, SplineParams, [SPLINE_MESH_PARAMS_FLOAT4_SIZE])
+END_GLOBAL_SHADER_PARAMETER_STRUCT()
+
 //////////////////////////////////////////////////////////////////////////
 // SplineMeshVertexFactory
 
@@ -61,8 +65,6 @@ class FSplineMeshVertexFactoryShaderParameters : public FVertexFactoryShaderPara
 {
 	DECLARE_TYPE_LAYOUT(FSplineMeshVertexFactoryShaderParameters, NonVirtual);
 public:
-	void Bind(const FShaderParameterMap& ParameterMap);
-
 	void GetElementShaderBindings(
 		const class FSceneInterface* Scene,
 		const FSceneView* View,
@@ -76,7 +78,6 @@ public:
 		) const;
 
 private:
-	LAYOUT_FIELD(FShaderParameter, SplineMeshParams);
 };
 
 struct FSplineMeshSceneInstanceDataBuffers : public FSingleInstanceDataBuffers
@@ -92,55 +93,24 @@ struct FSplineMeshSceneInstanceDataBuffers : public FSingleInstanceDataBuffers
  * This interface ties common functionality to the two different spline mesh scene proxies without duplicating code or
  * introducing diamond inheritance.
  */
-template<typename TDerived>
-class TSplineMeshSceneProxyCommon
+class FSplineMeshSceneProxyCommon
 {
 public:
-	TSplineMeshSceneProxyCommon()
+	FSplineMeshSceneProxyCommon()
 	{
 		FMemory::Memzero(&SplineParams, sizeof(SplineParams));
 		SplineParams.TextureCoord = FUintVector2(INDEX_NONE, INDEX_NONE);
 	}
-	virtual ~TSplineMeshSceneProxyCommon() = default;
+	virtual ~FSplineMeshSceneProxyCommon() = default;
 
 	const FSplineMeshShaderParams& GetSplineMeshParams() const { return SplineParams; }
 
-	void SetSplineTextureCoord_RenderThread(uint32 SplineIndex, FUintVector2 TexCoord)
-	{
-		check(SplineIndex == 0); // currently only support one spline type, but could do instanced later
-		SplineParams.TextureCoord = TexCoord;
-		RepackSplineMeshParams();
-	}
+	ENGINE_API void SetSplineTextureCoord_RenderThread(uint32 SplineIndex, FUintVector2 TexCoord);
 	
-	void UpdateSplineMeshParams_RenderThread(const FSplineMeshShaderParams& Params)
-	{
-		// Don't replace the texture coord that was assigned to us in SetSplineTextureCoord
-		FUintVector2 TexCoord = SplineParams.TextureCoord;
-		SplineParams = Params;
-		SplineParams.TextureCoord = TexCoord;
-		RepackSplineMeshParams();
-	}
+	ENGINE_API void UpdateSplineMeshParams_RenderThread(const FSplineMeshShaderParams& Params);
 
 protected:
-	void RepackSplineMeshParams(bool bUpdateGPUScene = true)
-	{
-		auto& SceneProxy = Downcast();
-		
-		// NOTE: If the payload extension was not initialized, we're probably not using GPU Scene
-		if (SplineMeshInstanceData.Update(SplineParams))
-		{
-			if (bUpdateGPUScene)
-			{
-				// Request a GPU Scene update for this primitive so it updates its instance data
-				SceneProxy.GetScene().RequestGPUSceneUpdate(
-					*SceneProxy.GetPrimitiveSceneInfo(),
-					EPrimitiveDirtyState::ChangedOther
-				);
-			}
-		}
-
-		bUpdateRayTracingGeometry = true;
-	}
+	void RepackSplineMeshParamsInternal();
 
 protected:
 	/** Parameters that define the spline, used to deform mesh */
@@ -150,21 +120,18 @@ protected:
 
 	bool bUpdateRayTracingGeometry = true;
 	FRWBuffer RayTracingDynamicVertexBuffer;
+
+	TUniformBufferRef<FSplineMeshVFPrimitiveShaderParameters> SplineMeshVFPrimitiveShaderParametersUB;
 private:
-	/** implemented by derived to provide access by parent */
-	virtual TDerived& Downcast() = 0;
+	/** implemented by derived to provide access to proxy */
+	virtual FPrimitiveSceneProxy& GetProxy() = 0;
 };
 
-#define IMPL_SPLINE_MESH_SCENE_PROXY_COMMON(classname) \
-	friend class TSplineMeshSceneProxyCommon<classname>; \
-	private: \
-		virtual classname& Downcast() override { return *this; }
-
 /** Scene proxy for SplineMesh instance */
-class FSplineMeshSceneProxy final : public FStaticMeshSceneProxy, public TSplineMeshSceneProxyCommon<FSplineMeshSceneProxy>
+class FSplineMeshSceneProxy final : public FStaticMeshSceneProxy, public FSplineMeshSceneProxyCommon
 {
-	IMPL_SPLINE_MESH_SCENE_PROXY_COMMON(FSplineMeshSceneProxy)
-
+private:
+	virtual FPrimitiveSceneProxy& GetProxy() override { return *this; }
 public:
 	FSplineMeshSceneProxy(USplineMeshComponent* InComponent);
 
@@ -201,10 +168,10 @@ private:
 };
 
 /** Scene proxy for SplineMesh instance for Nanite */
-class FNaniteSplineMeshSceneProxy final : public Nanite::FSceneProxy, public TSplineMeshSceneProxyCommon<FNaniteSplineMeshSceneProxy>
+class FNaniteSplineMeshSceneProxy final : public Nanite::FSceneProxy, public FSplineMeshSceneProxyCommon
 {
-	IMPL_SPLINE_MESH_SCENE_PROXY_COMMON(FNaniteSplineMeshSceneProxy)
-
+private:
+	virtual FPrimitiveSceneProxy& GetProxy() override { return *this; }
 public:
 	FNaniteSplineMeshSceneProxy(const Nanite::FMaterialAudit& MaterialAudit, USplineMeshComponent* InComponent);
 	ENGINE_API FNaniteSplineMeshSceneProxy(const Nanite::FMaterialAudit& MaterialAudit, const FStaticMeshSceneProxyDesc& InMeshDesc, const FSplineMeshSceneProxyDesc& InSplineDesc);
