Fix build error from 51093404.

#rnx

#Synced-CL 51020780

diff --git a/Engine/Source/Runtime/Engine/Public/Materials/MaterialIRCommon.h b/Engine/Source/Runtime/Engine/Public/Materials/MaterialIRCommon.h
--- a/Engine/Source/Runtime/Engine/Public/Materials/MaterialIRCommon.h
+++ b/Engine/Source/Runtime/Engine/Public/Materials/MaterialIRCommon.h
@@ -24,6 +24,7 @@ struct FMaterialIRModuleBuilder;
 struct FMaterialIRModuleBuilderImpl;
 struct FShaderCompilerEnvironment;
 struct FStaticParameterSet;
+struct FMaterialExternalCodeDeclaration;
 class FMaterial;
 class FMaterialIRModule;
 class ITargetPlatform;
diff --git a/Engine/Source/Runtime/Engine/Public/Materials/MaterialIRExtern.h b/Engine/Source/Runtime/Engine/Public/Materials/MaterialIRExtern.h
--- a/Engine/Source/Runtime/Engine/Public/Materials/MaterialIRExtern.h
+++ b/Engine/Source/Runtime/Engine/Public/Materials/MaterialIRExtern.h
@@ -4,7 +4,6 @@
 #if WITH_EDITOR
 
 #include "Materials/MaterialIRCommon.h"
-#include "Materials/MaterialExternalCodeRegistry.h"
 #include "Materials/MaterialIRTypes.h"
 #include "MaterialDomain.h"
 
@@ -316,9 +315,9 @@ struct FExternInlineDeclaration
 	uint32 MaterialDomainMask = 0xffffffff;
 	static_assert(MD_MAX <= 32);
 
-	// Specifies the minimum feature level required by this inline extern instruction.
+	// Specifies the minimum EMaterialFeatureLevel required by this inline extern instruction.
 	// An error is reported if the translating material has a lower feature level.
-	EMaterialFeatureLevel MinimumFeatureLevel{};
+	uint32 MinimumFeatureLevel{};
 
 	// Array of this instruction MIR value graph properties.
 	MIR::EGraphProperties GraphProperties = MIR::EGraphProperties::None;
