Removed ensures both FBodyInstance::GetSimplePhysicalMaterial/GetComplexPhysicalMaterials since there are existing tests that fail.
Ensures were validating that calling GetSimplePhysicalMaterial/GetComplexPhysicalMaterials outside of the GT should only be allowed if PhysicsCreationInputs cache was valid.
AWaterBody::Serialize can currently call DeprecateData -> UWaterBodyComponent::DeprecateData -> FBodyInstance::SetPhysMaterialOverride and call both of these functions (from the ALT).

#Synced-CL 50924777

diff --git a/Engine/Source/Runtime/Engine/Private/PhysicsEngine/BodyInstance.cpp b/Engine/Source/Runtime/Engine/Private/PhysicsEngine/BodyInstance.cpp
--- a/Engine/Source/Runtime/Engine/Private/PhysicsEngine/BodyInstance.cpp
+++ b/Engine/Source/Runtime/Engine/Private/PhysicsEngine/BodyInstance.cpp
@@ -3171,12 +3171,9 @@ namespace UE::Private
 UPhysicalMaterial* FBodyInstance::GetSimplePhysicalMaterial(const FBodyInstance* BodyInstance, TWeakObjectPtr<UPrimitiveComponent> OwnerComp, TWeakObjectPtr<UBodySetup> BodySetupPtr, IPhysicsBodyInstanceOwner* BodyInstanceOwner)
 {
 	const FBodyInstancePhysicsCreationInputsPtr& PhysicsCreationInputs = BodyInstance->GetPhysicsCreationInputs();
-	if (!IsInGameThread() && !IsInParallelGameThread())
+	if (!IsInGameThread() && !IsInParallelGameThread() && PhysicsCreationInputs.IsValid())
 	{
-		if (ensureMsgf(PhysicsCreationInputs.IsValid(), TEXT("FBodyInstance::GetSimplePhysicalMaterial called outside of game thread, expected a valid cache for %s"), *BodyInstance->GetBodyDebugName()))
-		{
-			return PhysicsCreationInputs->ResolvedSimpleMaterial;
-		}
+		return PhysicsCreationInputs->ResolvedSimpleMaterial;
 	}
 
 	if(!GEngine || !GEngine->DefaultPhysMaterial)
@@ -3272,17 +3269,14 @@ void FBodyInstance::GetComplexPhysicalMaterials(const FBodyInstance* BodyInstanc
 	if (BodyInstance)
 	{
 		const FBodyInstancePhysicsCreationInputsPtr& PhysicsCreationInputs = BodyInstance->GetPhysicsCreationInputs();
-		if (!IsInGameThread() && !IsInParallelGameThread())
+		if (!IsInGameThread() && !IsInParallelGameThread() && PhysicsCreationInputs.IsValid())
 		{
-			if (ensureMsgf(PhysicsCreationInputs.IsValid(), TEXT("FBodyInstance::GetComplexPhysicalMaterials called outside of game thread, expected a valid cache for %s"), *BodyInstance->GetBodyDebugName()))
+			if (OutPhysMaterialMasks)
 			{
-				if (OutPhysMaterialMasks)
-				{
-					*OutPhysMaterialMasks = PhysicsCreationInputs->ResolvedComplexMasks;
-				}
-				OutPhysMaterials = PhysicsCreationInputs->ResolvedComplexMaterials;
-				return;
+				*OutPhysMaterialMasks = PhysicsCreationInputs->ResolvedComplexMasks;
 			}
+			OutPhysMaterials = PhysicsCreationInputs->ResolvedComplexMaterials;
+			return;
 		}
 	}
 
