Revert CL 51093662
Adding comments

#Synced-CL 51094291

diff --git a/Engine/Source/Runtime/Engine/Private/Components/SkeletalMeshComponent.cpp b/Engine/Source/Runtime/Engine/Private/Components/SkeletalMeshComponent.cpp
--- a/Engine/Source/Runtime/Engine/Private/Components/SkeletalMeshComponent.cpp
+++ b/Engine/Source/Runtime/Engine/Private/Components/SkeletalMeshComponent.cpp
@@ -4927,15 +4927,10 @@ void USkeletalMeshComponent::CompleteParallelAnimationEvaluation(bool bDoPostAni
 
 bool USkeletalMeshComponent::HandleExistingParallelEvaluationTask(bool bBlockOnTask, bool bPerformPostAnimEvaluation)
 {
-	// If we have an async task handle
-	if (IsRunningParallelEvaluation())
+	// We are already processing eval and we aren't currently executing it on this thread
+	if (IsRunningParallelEvaluation() && ParallelAnimationEvaluationTask->IsAwaitable())
 	{
-		// If we want to block and wait for it
-		if (bBlockOnTask
-			// and it hasn't completed yet
-			&& !ParallelAnimationEvaluationTask->IsCompleted()
-			// and we aren't currently executing it on this thread
-			&& ParallelAnimationEvaluationTask->IsAwaitable())
+		if (bBlockOnTask)
 		{
 			// If we try and complete existing parallel evaluation in a transaction, we have to abort the transaction
 			// and retry up to this point. This is because we can only have started parallel evaluation outside of a
@@ -4945,23 +4940,28 @@ bool USkeletalMeshComponent::HandleExistingParallelEvaluationTask(bool bBlockOnT
 			if (AutoRTFM::IsTransactional())
 			{
 				(void)AutoRTFM::Close([&]
-				{
-					AutoRTFM::OnRetry([strong_this = TObjectPtr<USkeletalMeshComponent>(this), bPerformPostAnimEvaluation]
 					{
-						UE_LOG(LogSkeletalMesh, Display, TEXT("A call to `HandleExistingParallelEvaluationTask` that is blocking is causing Verse to abort and retry."));
-						const bool bBlockOnTask = true; // Was true from the if block above!
-						strong_this->HandleExistingParallelEvaluationTask(bBlockOnTask, bPerformPostAnimEvaluation);
-						UE_LOG(LogSkeletalMesh, Display, TEXT("`HandleExistingParallelEvaluationTask` has completed after Verse aborted, and we are now retrying."));
+						AutoRTFM::OnRetry([strong_this = TObjectPtr<USkeletalMeshComponent>(this), bPerformPostAnimEvaluation]
+							{
+								UE_LOG(LogSkeletalMesh, Display, TEXT("A call to `HandleExistingParallelEvaluationTask` that is blocking is causing Verse to abort and retry."));
+								const bool bBlockOnTask = true; // Was true from the if block above!
+								strong_this->HandleExistingParallelEvaluationTask(bBlockOnTask, bPerformPostAnimEvaluation);
+								UE_LOG(LogSkeletalMesh, Display, TEXT("`HandleExistingParallelEvaluationTask` has completed after Verse aborted, and we are now retrying."));
+							});
+						AutoRTFM::CascadingRetryTransaction();
 					});
-					AutoRTFM::CascadingRetryTransaction();
-				});
 				return true;
 			}
 
 			TRACE_CPUPROFILER_EVENT_SCOPE(USkeletalMeshComponent::BlockOnParallelEvaluationTask);
 			check(IsInGameThread()); // Only attempt this from game thread!
 			FTaskGraphInterface::Get().WaitUntilTaskCompletes(ParallelAnimationEvaluationTask, ENamedThreads::GameThread);
-			CompleteParallelAnimationEvaluation(bPerformPostAnimEvaluation); //Perform completion now
+
+			// We had a valid task handle and now we are done waiting for it
+			// Perform completion now
+			// Note that because completion happens in a separate task from the one we waited on
+			// it could execute concurrently with this call to finalize evaluation
+			CompleteParallelAnimationEvaluation(bPerformPostAnimEvaluation);
 		}
 
 		return true;
