Log fixes

Reducing log spam on seek requests
Adding AutoDurationThreshold log on initialization

#Synced-CL 50891994

diff --git a/Engine/Source/Runtime/Engine/Private/AudioDecompress.cpp b/Engine/Source/Runtime/Engine/Private/AudioDecompress.cpp
--- a/Engine/Source/Runtime/Engine/Private/AudioDecompress.cpp
+++ b/Engine/Source/Runtime/Engine/Private/AudioDecompress.cpp
@@ -350,7 +350,7 @@ bool IStreamedCompressedInfo::StreamCompressedData(uint8* Destination, bool bLoo
 
 		uint32 ChunkSize = 0;
 		const uint8* NewlySeekedChunk = GetLoadedChunk(StreamingSoundWave, StreamSeekBlockIndex, ChunkSize); 
-		UE_LOG(LogAudio, Log, TEXT("Seek request for (%s): (%s) Chunk=%d (%s), Offset=%d, OffsetInAudioFrames=%u"),
+		UE_LOG(LogAudio, VeryVerbose, TEXT("Seek request for (%s): (%s) Chunk=%d (%s), Offset=%d, OffsetInAudioFrames=%u"),
 			*StreamingSoundWave->GetFName().ToString(),
 			StreamSeekToAudioFrames != INDEX_NONE ? TEXT("Using streaming seek-tables") : TEXT("Using chunk/offset pair"),
 			StreamSeekBlockIndex.load(),
diff --git a/Engine/Source/Runtime/Engine/Private/SoundWaveProxyDecodeCacheImpl.cpp b/Engine/Source/Runtime/Engine/Private/SoundWaveProxyDecodeCacheImpl.cpp
--- a/Engine/Source/Runtime/Engine/Private/SoundWaveProxyDecodeCacheImpl.cpp
+++ b/Engine/Source/Runtime/Engine/Private/SoundWaveProxyDecodeCacheImpl.cpp
@@ -58,8 +58,8 @@ FSoundWaveProxyDecodeCache::FSoundWaveProxyDecodeCache()
 		float ActualMaxMemoryMB = FUnitConversion::Convert((float)MaxMemoryBytes, EUnit::Bytes, EUnit::Megabytes);
 
 		UE_LOG(LogAudio, Display,
-			TEXT("Initializing decode cache to %.2f Megabytes (Desired: %.2f), MaxChunkCount: %d. Enabled: %s"),
-			ActualMaxMemoryMB, DesiredCacheSizeMB, MaxChunkCount, SoundWaveProxyDecodeCache::IsEnabled() ? TEXT("Yes") : TEXT("No"));
+			TEXT("Initializing decode cache to %.2f Megabytes (Desired: %.2f), MaxChunkCount: %d. AutoDurationThreshold: %.2f, Enabled: %s"),
+			ActualMaxMemoryMB, DesiredCacheSizeMB, MaxChunkCount, AutoDurationThreshold, SoundWaveProxyDecodeCache::IsEnabled() ? TEXT("Yes") : TEXT("No"));
 
 		// Note that we are tracking the LRU overhead here since the LRU cache is pre-allocated when initializing.
 		UsedMemoryBytes = MaxChunkCount * (sizeof(FDecodedSoundWaveChunk) + sizeof(FDecodedSoundWaveChunkKey));
