Improved SCW scaling when SCW memory spikes on multiple workers.
#rnx

#Synced-CL 50966172

diff --git a/Engine/Source/Runtime/Engine/Private/ShaderCompiler/ShaderCompilerThreadRunnable.cpp b/Engine/Source/Runtime/Engine/Private/ShaderCompiler/ShaderCompilerThreadRunnable.cpp
--- a/Engine/Source/Runtime/Engine/Private/ShaderCompiler/ShaderCompilerThreadRunnable.cpp
+++ b/Engine/Source/Runtime/Engine/Private/ShaderCompiler/ShaderCompilerThreadRunnable.cpp
@@ -76,6 +76,16 @@ static FAutoConsoleVariableRef CVarShaderCompilerWorkerSoftLimitInWorkers(
 	TEXT("See r.ShaderCompiler.SoftLimit.\n"),
 	ECVF_Default);
 
+float GShaderCompilerMemoryUsedPerWorkerBiasTowardsMax = 0.15f;
+static FAutoConsoleVariableRef CVarShaderCompilerWorkerMemoryBiasTowardsMax(
+	TEXT("r.ShaderCompiler.MemoryUsedPerWorkerBiasTowardsMax"),
+	GShaderCompilerMemoryUsedPerWorkerBiasTowardsMax,
+	TEXT("How much to bias the memory estimated required per worker from the current average across the workers towards the current max of any worker.\n") \
+	TEXT("Using just the average underestimates the risk of going over the limit when a small number of workers are using a lot of memory.\n")
+	TEXT("A value of 0 uses the pure average, while a value of 1 uses the max.\n")
+	TEXT("The higher the value, the more conservative scaling back workers is, but the slower overall shader compilation will be.\n"),
+	ECVF_Default);
+
 int32 GShaderCompilerWorkerScaleUpApproach = 1;
 static FAutoConsoleVariableRef CVarShaderCompilerWorkerScaleUpApproach(
 	TEXT("r.ShaderCompiler.WorkerScaleUpApproach"),
@@ -1161,14 +1171,21 @@ FString FShaderCompileThreadRunnable::GetWorkingDirectoryForWorker(int32 WorkerI
 	return FPaths::Combine(bRelativePath ? Manager->ShaderBaseWorkingDirectory : Manager->AbsoluteShaderBaseWorkingDirectory, FString::FromInt(WorkerIndex));
 }
 
+
 FShaderCompileMemoryUsage FShaderCompileThreadRunnable::CalcPerWorkerMemoryUsageTotal(const TArray<FWorkerMemoryUsage>& InPerWorkerMemoryUsage) const
 {
-	FShaderCompileMemoryUsage MemoryUsage{ 0ull, 0ull, 0 };
+	FShaderCompileMemoryUsage MemoryUsage{ 0ull, 0ull, 0ll, 0 };
+
 
 	for (const FWorkerMemoryUsage& WorkerMemUsage : InPerWorkerMemoryUsage)
 	{
 		MemoryUsage.VirtualMemory += WorkerMemUsage.VirtualMemory;
 		MemoryUsage.PhysicalMemory += WorkerMemUsage.PhysicalMemory;
+
+		if (WorkerMemUsage.VirtualMemory > MemoryUsage.MaxWorkerMemory)
+		{
+			MemoryUsage.MaxWorkerMemory = WorkerMemUsage.VirtualMemory;
+		}
 	}
 	MemoryUsage.ActiveWorkerCount = GetNumberOfAvailableWorkers();
 
@@ -1188,13 +1205,12 @@ void FShaderCompileThreadRunnable::CheckMemoryLimitViolation()
 	{
 		return;
 	}
-	FShaderCompileMemoryUsage TotalSCMemoryUsage = { 0ull, 0ull, 0 };
 
 	{
 		FScopeLock PerWorkerMemScopeLock(&PerWorkerMemUsageLock);
 		CachePerWorkerMemoryUsage();
 	}
-	TotalSCMemoryUsage = CalcPerWorkerMemoryUsageTotal(PerWorkerMemoryUsage);
+	FShaderCompileMemoryUsage TotalSCMemoryUsage = CalcPerWorkerMemoryUsageTotal(PerWorkerMemoryUsage);
 	MemoryMonitoringState.LastTimeOfMemoryLimitPoll = CurrentTime;
 
 	//	Query the memory violation information once, whether just checking suspend or resume as well.
@@ -1220,9 +1236,9 @@ void FShaderCompileThreadRunnable::CheckMemoryLimitViolation()
 		//	Not violating, check for a soft limit.
 		if(GShaderCompilerWorkerSoftLimit)
 		{
-			if(LimitInfo.ActiveWorkerCount)
+			if(LimitInfo.ActiveWorkerCount > 1)
 			{
-				const int64 MemoryUsedPerWorker = LimitInfo.MemoryUsed / LimitInfo.ActiveWorkerCount;
+				const int64 MemoryUsedPerWorker = GetEstimatedMemoryRequiredPerWorker(TotalSCMemoryUsage);
 
 				//	If there isn't space for one more worker at the current average, suspend one worker.
 				//	Resume checks for space for at least two more workers.
@@ -1308,9 +1324,8 @@ void FShaderCompileThreadRunnable::CheckMemoryLimitViolation()
 		if (MemoryAvailable > 0)
 		{
 			//	Set a memory per worker based on the config, adpating if there are running workers.
-			const int32 EffectiveActiveWorkerCount = LimitInfo.ActiveWorkerCount ? LimitInfo.ActiveWorkerCount : 1;
-			int64 MemoryUsedPerWorker = LimitInfo.MemoryUsed / EffectiveActiveWorkerCount;
-				
+			const int32 EffectiveActiveWorkerCount = LimitInfo.ActiveWorkerCount;
+			const int64 MemoryUsedPerWorker = GetEstimatedMemoryRequiredPerWorker(TotalSCMemoryUsage);
 			if(MemoryAvailable > (MemoryUsedPerWorker*(EffectiveActiveWorkerCount+2)))
 			{
 				NumWorkersToResume = 1;
@@ -1339,6 +1354,15 @@ void FShaderCompileThreadRunnable::CheckMemoryLimitViolation()
 	}
 }
 
+int64 FShaderCompileThreadRunnable::GetEstimatedMemoryRequiredPerWorker(const FShaderCompileMemoryUsage& MemoryUsage)
+{
+	const float BiasTowardsMax = FMath::Clamp(GShaderCompilerMemoryUsedPerWorkerBiasTowardsMax, 0.0f, 1.0f);
+	int64 MemoryUsedPerWorker = (MemoryUsage.VirtualMemory / (MemoryUsage.ActiveWorkerCount ? MemoryUsage.ActiveWorkerCount : 1));
+	MemoryUsedPerWorker = FMath::Lerp(MemoryUsedPerWorker, MemoryUsage.MaxWorkerMemory, BiasTowardsMax);
+
+	return MemoryUsedPerWorker;
+}
+
 bool FShaderCompileThreadRunnable::QueryMemoryLimitViolationStatus(const FShaderCompileMemoryUsage& InSCMemoryUsage, FShaderCompilerJobMemoryLimitInfo& OutInfo)
 {
 	const int64 MemoryLimitInBytes = static_cast<int64>(GShaderCompilerCurrentMemoryLimitInMiB) * 1024ull * 1024ull;
@@ -1649,12 +1673,11 @@ int32 FShaderCompileThreadRunnable::CompilingLoop()
 
 	if(NumActiveThreads)
 	{
-		FShaderCompileMemoryUsage TotalSCMemoryUsage = { 0ull, 0ull, 0 };
 		{
 			FScopeLock PerWorkerMemScopeLock(&PerWorkerMemUsageLock);
 			CachePerWorkerMemoryUsage();
 		}
-		TotalSCMemoryUsage = CalcPerWorkerMemoryUsageTotal(PerWorkerMemoryUsage);
+		FShaderCompileMemoryUsage TotalSCMemoryUsage = CalcPerWorkerMemoryUsageTotal(PerWorkerMemoryUsage);
 
 		if(bDynamicSCWBudgetEnabled)
 		{
diff --git a/Engine/Source/Runtime/Engine/Public/ShaderCompiler.h b/Engine/Source/Runtime/Engine/Public/ShaderCompiler.h
--- a/Engine/Source/Runtime/Engine/Public/ShaderCompiler.h
+++ b/Engine/Source/Runtime/Engine/Public/ShaderCompiler.h
@@ -142,6 +142,11 @@ struct FShaderCompileMemoryUsage
 	*/
 	uint64 PhysicalMemory;
 
+	/**
+	 * The max memory currently used across workers
+	 */
+	uint64 MaxWorkerMemory;
+
 	/**
 	 * The number of active shader compiler workers
 	 */
@@ -385,6 +390,9 @@ private:
 	/** Checks it the memory limit for shader compile workers has been exceeded and suspend workers as needed. */
 	void CheckMemoryLimitViolation();
 
+	/** Return an estimate for how much memory is required for a single worker. Used for both suspension and resumption */
+	int64 GetEstimatedMemoryRequiredPerWorker(const FShaderCompileMemoryUsage& LimitInfo);
+
 	/** Queries the status if the worker processes have violated the memory limitation. */
 	bool QueryMemoryLimitViolationStatus(const FShaderCompileMemoryUsage& InSCMemoryUsage, FShaderCompilerJobMemoryLimitInfo& OutInfo);
 
