[NetPhysics] When creating RewindData from PhysicsReplication, start using it directly else it could reinitialize it multiple times in the for-loop.
- Also lower the log verbosity of the log informing about received states beign outside of rewind data, as well as print out ServerFrame and TickOffset in the log.

#Synced-CL 51043916

diff --git a/Engine/Source/Runtime/Engine/Private/PhysicsEngine/PhysicsReplication.cpp b/Engine/Source/Runtime/Engine/Private/PhysicsEngine/PhysicsReplication.cpp
--- a/Engine/Source/Runtime/Engine/Private/PhysicsEngine/PhysicsReplication.cpp
+++ b/Engine/Source/Runtime/Engine/Private/PhysicsEngine/PhysicsReplication.cpp
@@ -1043,7 +1043,7 @@ void FPhysicsReplicationAsync::OnPreSimulate_Internal()
 
 	// Early out if this is a resim frame
 	Chaos::FRewindData* RewindData = RigidsSolver->GetRewindData();
-	const bool bRewindDataExist = RewindData != nullptr;
+	bool bRewindDataExist = RewindData != nullptr;
 	if (bRewindDataExist && RewindData->IsResim())
 	{
 		// TODO, Handle the transition from post-resim to interpolation better (disabled by default, resim vs replication interaction is handled via FPhysicsReplicationAsync::CacheResimInteractions)
@@ -1081,6 +1081,8 @@ void FPhysicsReplicationAsync::OnPreSimulate_Internal()
 				if (ensure(Chaos::FPBDRigidsSolver::IsNetworkPhysicsPredictionEnabled() && RigidsSolver->IsUsingFixedDt()))
 				{
 					RigidsSolver->EnableRewindCapture();
+					RewindData = RigidsSolver->GetRewindData();
+					bRewindDataExist = RewindData != nullptr;
 				}
 			}
 
@@ -1577,17 +1579,20 @@ void FPhysicsReplicationAsync::CheckTargetResimValidity(FReplicatedPhysicsTarget
 			Target.RepMode = EPhysicsReplicationMode::PredictiveInterpolation;
 		}
 		
-		if (ResimOutOfBoundsCounter == 0)
+		if (NetworkPhysicsTickOffsetAssigned)
 		{
-			UE_LOG(LogPhysics, Warning, TEXT("FPhysicsReplication DESYNCED - received target frame (%d) out of rewind data bounds (%d, %d) - %s - Target will use %s")
-				, LocalFrame, RewindData->GetEarliestFrame_Internal(), RewindData->CurrentFrame()
-					, (LocalFrame < RewindData->GetEarliestFrame_Internal()) 
-					? TEXT("Client is far ahead of the server, server might be dropping frames.") 
+			if (ResimOutOfBoundsCounter == 0)
+			{
+				UE_LOG(LogPhysics, Log, TEXT("FPhysicsReplication DESYNCED - received target ClientFrame: %d (ServerFrame: %d - FrameOffset: %d) out of rewind data bounds (%d, %d) - %s - Target will use %s")
+					, LocalFrame, Target.ServerFrame, NetworkPhysicsTickOffset, RewindData->GetEarliestFrame_Internal(), RewindData->CurrentFrame()
+					, (LocalFrame < RewindData->GetEarliestFrame_Internal())
+					? TEXT("Client is far ahead of the server, server might be dropping frames.")
 					: TEXT("Client is behind the server, client might be dropping frames."), *UEnum::GetValueAsString(Target.RepMode));
-		}
+			}
 
-		ResimOutOfBoundsCounter++;
-		ResimErrorLogTimer = 0;
+			ResimOutOfBoundsCounter++;
+			ResimErrorLogTimer = 0;
+		}
 	}
 	else
 	{
