[FTimerManager + AutoRTFM] Rework to make this safe to use from both the open and closed.

We have logic that uses FTimerManager from the open and closed of a transaction. This is extremely unsafe, and will very likely cause memory corruption on abort.
This CL makes it safe to use FTimerManager from both the open and closed of a transaction, so long as the timers as distinct.

This is a reimplementation of 50880655, which was reverted because the moving of delegates into the open would cause a double free - once by AutoRTFM's aborting of transactional memory, and a second time when the task was destructed.
This implementation does not use a queue of on-abort tasks to rollback the timer state, but instead maintains a TMap of transactional timer data. This keeps the delegate memory entirely in the closed, avoiding the double free issue, and is more efficient as we're not continually appending new on-abort tasks.

#Synced-CL 51010198

diff --git a/Engine/Source/Runtime/AutoRTFM/Public/AutoRTFM.h b/Engine/Source/Runtime/AutoRTFM/Public/AutoRTFM.h
--- a/Engine/Source/Runtime/AutoRTFM/Public/AutoRTFM.h
+++ b/Engine/Source/Runtime/AutoRTFM/Public/AutoRTFM.h
@@ -418,6 +418,13 @@ UE_AUTORTFM_CRITICAL_INLINE_ALWAYS bool IsCommitting()
 	return autortfm_is_context_status(autortfm_status_committing);
 }
 
+// Returns true if we are currently aborting a transaction.
+// This will return true when inside an on-commit handler.
+UE_AUTORTFM_CRITICAL_INLINE_ALWAYS bool IsAborting()
+{
+	return autortfm_is_context_status(autortfm_status_aborting);
+}
+
 // Returns true if we are currently retrying a transaction.
 // This will return true when inside an on-retry handler.
 UE_AUTORTFM_CRITICAL_INLINE_ALWAYS bool IsRetrying()
diff --git a/Engine/Source/Runtime/Core/Public/Misc/ScopeExit.h b/Engine/Source/Runtime/Core/Public/Misc/ScopeExit.h
--- a/Engine/Source/Runtime/Core/Public/Misc/ScopeExit.h
+++ b/Engine/Source/Runtime/Core/Public/Misc/ScopeExit.h
@@ -2,6 +2,8 @@
 
 #pragma once
 
+#include "AutoRTFM/Defines.h"
+
 namespace ScopeExitSupport
 {
 	/**
@@ -10,7 +12,7 @@ namespace ScopeExitSupport
 	 * RAII class that calls a lambda when it is destroyed.
 	 */
 	template <typename FuncType>
-	class TScopeGuard
+	class AUTORTFM_INFER TScopeGuard
 	{
 		TScopeGuard(TScopeGuard&&) = delete;
 		TScopeGuard(const TScopeGuard&) = delete;
diff --git a/Engine/Source/Runtime/Engine/Classes/Engine/TimerHandle.h b/Engine/Source/Runtime/Engine/Classes/Engine/TimerHandle.h
--- a/Engine/Source/Runtime/Engine/Classes/Engine/TimerHandle.h
+++ b/Engine/Source/Runtime/Engine/Classes/Engine/TimerHandle.h
@@ -44,7 +44,7 @@ struct FTimerHandle
 
 	FString ToString() const
 	{
-		return FString::Printf(TEXT("%llu"), Handle);
+		return FString::Printf(TEXT("%llx"), Handle);
 	}
 
 private:
@@ -80,4 +80,4 @@ private:
 	{
 		return GetTypeHash(InHandle.Handle);
 	}
-};
\ No newline at end of file
+};
diff --git a/Engine/Source/Runtime/Engine/Private/Tests/AutoRTFM/AutoRTFMTimerManagerTests.cpp b/Engine/Source/Runtime/Engine/Private/Tests/AutoRTFM/AutoRTFMTimerManagerTests.cpp
new file mode 100644
--- /dev/null
+++ b/Engine/Source/Runtime/Engine/Private/Tests/AutoRTFM/AutoRTFMTimerManagerTests.cpp
@@ -0,0 +1,725 @@
+// Copyright Epic Games, Inc. All Rights Reserved.
+
+#include "AutoRTFM.h"
+#include "AutoRTFM/Testing.h"
+#include "CoreGlobals.h"
+#include "Engine/TimerHandle.h"
+#include "Misc/AssertionMacros.h"
+#include "Misc/AutomationTest.h"
+#include "TimerManager.h"
+
+#if WITH_DEV_AUTOMATION_TESTS
+
+IMPLEMENT_SIMPLE_AUTOMATION_TEST(FAutoRTFMTimerManagerTest, "AutoRTFM + FTimerManager", EAutomationTestFlags::EngineFilter | EAutomationTestFlags::ClientContext | EAutomationTestFlags::ServerContext | EAutomationTestFlags::CommandletContext)
+
+bool FAutoRTFMTimerManagerTest::RunTest(const FString & Parameters)
+{
+	if (!AutoRTFM::ForTheRuntime::IsAutoRTFMRuntimeEnabled())
+	{
+		ExecutionInfo.AddEvent(FAutomationEvent(EAutomationEventType::Info, TEXT("SKIPPED 'FAutoRTFMTimerManagerTest' test. AutoRTFM disabled.")));
+		return true;
+	}
+
+	auto Section = [&](const TCHAR* Name, auto&& Test)
+	{
+		{
+			FTimerManager TimerManager;
+			TimerManager.Tick(0);
+			if (!Test(TimerManager))
+			{
+				AddError(FString::Printf(TEXT("In section '%s' (HasBeenTickedThisFrame: true)."), Name), 1);
+			}
+		}
+		{
+			FTimerManager TimerManager;
+			TimerManager.Tick(0);
+			GFrameCounter++;
+			if (!Test(TimerManager))
+			{
+				AddError(FString::Printf(TEXT("In section '%s' (HasBeenTickedThisFrame: false)."), Name), 1);
+			}
+		}
+	};
+	
+	auto Tick = [&](FTimerManager& TimerManager)
+	{
+		// Simulate two frames to ensure that pending timers are activated and triggered.
+		for (int I = 0; I < 2; I++)
+		{
+			GFrameCounter++;
+			TimerManager.Tick(1.5f);
+		}
+	};
+
+	// Wrapper around FTimerManager::ClearTimer() to prevent invalidating of the handle argument.
+	auto ClearTimer = [&] AUTORTFM_ENABLE (FTimerManager& TimerManager, FTimerHandle TimerHandle)
+	{
+		TimerManager.ClearTimer(TimerHandle);
+	};
+
+	struct AUTORTFM_ENABLE FTimer
+	{
+		FTimerHandle Handle;
+		bool bCalled = false;
+		TFunction<void ()> Callback()
+		{
+			return [this]
+			{
+				bCalled = true;
+			};
+		}
+	};
+
+	Section(TEXT("Transact(SetTimer), Tick"), [&] (FTimerManager& TimerManager)
+	{
+		constexpr size_t NumTimers = 16;
+		FTimer Timers[NumTimers]{};
+
+		AutoRTFM::Testing::Commit([&]
+		{
+			for (FTimer& Timer : Timers)
+			{
+				TimerManager.SetTimer(Timer.Handle, Timer.Callback(), /* InRate */ 1.0f, FTimerManagerTimerParameters{});
+				check(TimerManager.IsTimerActive(Timer.Handle));
+			}
+		});
+
+		Tick(TimerManager);
+
+		for (FTimer& Timer : Timers)
+		{
+			check(Timer.bCalled);
+		}
+
+		return true;
+	});
+
+	Section(TEXT("Transact(SetTimer, Abort), Tick"), [&] (FTimerManager& TimerManager)
+	{
+		constexpr size_t NumTimers = 16;
+		FTimer Timers[NumTimers]{};
+
+		AutoRTFM::Testing::Abort([&]
+		{
+			for (FTimer& Timer : Timers)
+			{
+				TimerManager.SetTimer(Timer.Handle, Timer.Callback(), /* InRate */ 1.0f, FTimerManagerTimerParameters{});
+				check(TimerManager.IsTimerActive(Timer.Handle));
+			}
+			AutoRTFM::AbortTransaction();
+		});
+
+		Tick(TimerManager);
+
+		for (FTimer& Timer : Timers)
+		{
+			check(!Timer.bCalled);
+			check(!TimerManager.IsTimerActive(Timer.Handle));
+		}
+
+		return true;
+	});
+
+	Section(TEXT("Transact(SetTimer, ClearTimer), Tick"), [&] (FTimerManager& TimerManager)
+	{
+		constexpr size_t NumTimers = 16;
+		FTimer Timers[NumTimers]{};
+
+		AutoRTFM::Testing::Commit([&]
+		{
+			for (FTimer& Timer : Timers)
+			{
+				TimerManager.SetTimer(Timer.Handle, Timer.Callback(), /* InRate */ 1.0f, FTimerManagerTimerParameters{});
+				check(TimerManager.IsTimerActive(Timer.Handle));
+				ClearTimer(TimerManager, Timer.Handle);
+				check(!TimerManager.IsTimerActive(Timer.Handle));
+			}
+		});
+
+		Tick(TimerManager);
+
+		for (FTimer& Timer : Timers)
+		{
+			check(!Timer.bCalled);
+			check(!TimerManager.IsTimerActive(Timer.Handle));
+		}
+
+		return true;
+	});
+
+	Section(TEXT("Transact(SetTimer, ClearTimer, Abort), Tick"), [&] (FTimerManager& TimerManager)
+	{
+		constexpr size_t NumTimers = 16;
+		FTimer Timers[NumTimers]{};
+
+		AutoRTFM::Testing::Abort([&]
+		{
+			for (FTimer& Timer : Timers)
+			{
+				TimerManager.SetTimer(Timer.Handle, Timer.Callback(), /* InRate */ 1.0f, FTimerManagerTimerParameters{});
+				check(TimerManager.IsTimerActive(Timer.Handle));
+				ClearTimer(TimerManager, Timer.Handle);
+				check(!TimerManager.IsTimerActive(Timer.Handle));
+			}
+			AutoRTFM::AbortTransaction();
+		});
+
+		Tick(TimerManager);
+
+		for (FTimer& Timer : Timers)
+		{
+			check(!Timer.bCalled);
+			check(!TimerManager.IsTimerActive(Timer.Handle));
+		}
+
+		return true;
+	});
+
+	Section(TEXT("Transact(SetTimer, Open(SetTimer)), Tick"), [&] (FTimerManager& TimerManager)
+	{
+		constexpr size_t NumOpenTimers = 16;
+		FTimer OpenTimers[NumOpenTimers];
+		constexpr size_t NumClosedTimers = 16;
+		FTimer ClosedTimers[NumClosedTimers];
+
+		AutoRTFM::Testing::Commit([&]
+		{
+			for (FTimer& ClosedTimer : ClosedTimers)
+			{
+				TimerManager.SetTimer(ClosedTimer.Handle, ClosedTimer.Callback(), /* InRate */ 1.0f, FTimerManagerTimerParameters{});
+				check(TimerManager.IsTimerActive(ClosedTimer.Handle));
+			}
+			
+			AutoRTFM::Open([&]
+			{
+				for (FTimer& OpenTimer : OpenTimers)
+				{
+					TimerManager.SetTimer(OpenTimer.Handle, OpenTimer.Callback(), /* InRate */ 1.0f, FTimerManagerTimerParameters{});
+					check(TimerManager.IsTimerActive(OpenTimer.Handle));
+				}
+			});
+		});
+
+		Tick(TimerManager);
+
+		for (FTimer& ClosedTimer : ClosedTimers)
+		{
+			check(ClosedTimer.bCalled);
+		}
+
+		for (FTimer& OpenTimer : OpenTimers)
+		{
+			check(OpenTimer.bCalled);
+		}
+
+		return true;
+	});
+
+	Section(TEXT("Transact(SetTimer, Open(SetTimer), Abort), Tick"), [&] (FTimerManager& TimerManager)
+	{
+		constexpr size_t NumOpenTimers = 16;
+		FTimer OpenTimers[NumOpenTimers];
+		constexpr size_t NumClosedTimers = 16;
+		FTimer ClosedTimers[NumClosedTimers];
+
+		AutoRTFM::Testing::Abort([&]
+		{
+			for (FTimer& ClosedTimer : ClosedTimers)
+			{
+				TimerManager.SetTimer(ClosedTimer.Handle, ClosedTimer.Callback(), /* InRate */ 1.0f, FTimerManagerTimerParameters{});
+				check(TimerManager.IsTimerActive(ClosedTimer.Handle));
+			}
+			
+			AutoRTFM::Open([&]
+			{
+				for (FTimer& OpenTimer : OpenTimers)
+				{
+					TimerManager.SetTimer(OpenTimer.Handle, OpenTimer.Callback(), /* InRate */ 1.0f, FTimerManagerTimerParameters{});
+					check(TimerManager.IsTimerActive(OpenTimer.Handle));
+				}
+			});
+			
+			AutoRTFM::AbortTransaction();
+		});
+
+		Tick(TimerManager);
+
+		for (FTimer& ClosedTimer : ClosedTimers)
+		{
+			check(!ClosedTimer.bCalled);
+		}
+
+		for (FTimer& OpenTimer : OpenTimers)
+		{
+			check(OpenTimer.bCalled);
+		}
+
+		return true;
+	});
+
+	Section(TEXT("Transact(Open(SetTimer), SetTimer, Open(ClearTimer)), Tick"), [&] (FTimerManager& TimerManager)
+	{
+		constexpr size_t NumOpenTimers = 16;
+		FTimer OpenTimers[NumOpenTimers];
+		constexpr size_t NumClosedTimers = 16;
+		FTimer ClosedTimers[NumClosedTimers];
+
+		AutoRTFM::Testing::Commit([&]
+		{
+			AutoRTFM::Open([&]
+			{
+				for (FTimer& OpenTimer : OpenTimers)
+				{
+					TimerManager.SetTimer(OpenTimer.Handle, OpenTimer.Callback(), /* InRate */ 1.0f, FTimerManagerTimerParameters{});
+					check(TimerManager.IsTimerActive(OpenTimer.Handle));
+				}
+			});
+
+			for (FTimer& ClosedTimer : ClosedTimers)
+			{
+				TimerManager.SetTimer(ClosedTimer.Handle, ClosedTimer.Callback(), /* InRate */ 1.0f, FTimerManagerTimerParameters{});
+				check(TimerManager.IsTimerActive(ClosedTimer.Handle));
+			}
+			
+			AutoRTFM::Open([&]
+			{
+				for (FTimer& OpenTimer : OpenTimers)
+				{
+					check(TimerManager.IsTimerActive(OpenTimer.Handle));
+					ClearTimer(TimerManager, OpenTimer.Handle);
+					check(!TimerManager.IsTimerActive(OpenTimer.Handle));
+				}
+			});
+		});
+
+		Tick(TimerManager);
+
+		for (FTimer& OpenTimer : OpenTimers)
+		{
+			check(!OpenTimer.bCalled);
+		}
+
+		for (FTimer& ClosedTimer : ClosedTimers)
+		{
+			check(ClosedTimer.bCalled);
+		}
+
+		return true;
+	});
+
+	Section(TEXT("Transact(Open(SetTimer), SetTimer, Open(ClearTimer), Abort), Tick"), [&] (FTimerManager& TimerManager)
+	{
+		constexpr size_t NumOpenTimers = 16;
+		FTimer OpenTimers[NumOpenTimers];
+		constexpr size_t NumClosedTimers = 16;
+		FTimer ClosedTimers[NumClosedTimers];
+
+		AutoRTFM::Testing::Abort([&]
+		{
+			AutoRTFM::Open([&]
+			{
+				for (FTimer& OpenTimer : OpenTimers)
+				{
+					TimerManager.SetTimer(OpenTimer.Handle, OpenTimer.Callback(), /* InRate */ 1.0f, FTimerManagerTimerParameters{});
+					check(TimerManager.IsTimerActive(OpenTimer.Handle));
+				}
+			});
+
+			for (FTimer& ClosedTimer : ClosedTimers)
+			{
+				TimerManager.SetTimer(ClosedTimer.Handle, ClosedTimer.Callback(), /* InRate */ 1.0f, FTimerManagerTimerParameters{});
+				check(TimerManager.IsTimerActive(ClosedTimer.Handle));
+			}
+			
+			AutoRTFM::Open([&]
+			{
+				for (FTimer& OpenTimer : OpenTimers)
+				{
+					check(TimerManager.IsTimerActive(OpenTimer.Handle));
+					ClearTimer(TimerManager, OpenTimer.Handle);
+					check(!TimerManager.IsTimerActive(OpenTimer.Handle));
+				}
+			});
+			
+			AutoRTFM::AbortTransaction();
+		});
+
+		Tick(TimerManager);
+
+		for (FTimer& OpenTimer : OpenTimers)
+		{
+			check(!OpenTimer.bCalled);
+		}
+
+		for (FTimer& ClosedTimer : ClosedTimers)
+		{
+			check(!ClosedTimer.bCalled);
+		}
+
+		return true;
+	});
+
+	Section(TEXT("SetTimer, Transact(ClearTimer), Tick"), [&] (FTimerManager& TimerManager)
+	{
+		constexpr size_t NumTimers = 16;
+		FTimer Timers[NumTimers]{};
+
+		for (FTimer& Timer : Timers)
+		{
+			TimerManager.SetTimer(Timer.Handle, Timer.Callback(), /* InRate */ 1.0f, FTimerManagerTimerParameters{});
+			check(TimerManager.IsTimerActive(Timer.Handle));
+		}
+
+		AutoRTFM::Testing::Commit([&]
+		{
+			for (size_t I = 0; I < NumTimers; I++)
+			{
+				if ((I & 1) == 0)
+				{
+					FTimer& Timer = Timers[I];
+					check(TimerManager.IsTimerActive(Timer.Handle));
+					ClearTimer(TimerManager, Timer.Handle);
+					check(!TimerManager.IsTimerActive(Timer.Handle));
+				}
+			}
+		});
+
+		Tick(TimerManager);
+
+		for (size_t I = 0; I < NumTimers; I++)
+		{
+			FTimer& Timer = Timers[I];
+			bool const bCleared = (I & 1) == 0;
+			check(Timer.bCalled == !bCleared);
+		}
+
+		return true;
+	});
+
+	Section(TEXT("SetTimer, Transact(ClearTimer, Abort), Tick"), [&] (FTimerManager& TimerManager)
+	{
+		constexpr size_t NumTimers = 16;
+		FTimer Timers[NumTimers]{};
+
+		for (FTimer& Timer : Timers)
+		{
+			TimerManager.SetTimer(Timer.Handle, Timer.Callback(), /* InRate */ 1.0f, FTimerManagerTimerParameters{});
+			check(TimerManager.IsTimerActive(Timer.Handle));
+		}
+
+		AutoRTFM::Testing::Abort([&]
+		{
+			for (size_t I = 0; I < NumTimers; I++)
+			{
+				FTimer& Timer = Timers[I];
+				ClearTimer(TimerManager, Timer.Handle);
+			}
+			AutoRTFM::AbortTransaction();
+		});
+
+		Tick(TimerManager);
+
+		for (size_t I = 0; I < NumTimers; I++)
+		{
+			FTimer& Timer = Timers[I];
+			check(Timer.bCalled);
+		}
+
+		return true;
+	});
+
+	Section(TEXT("SetTimer, Transact(PauseTimer), Tick"), [&] (FTimerManager& TimerManager)
+	{
+		constexpr size_t NumTimers = 16;
+		FTimer Timers[NumTimers]{};
+
+		for (FTimer& Timer : Timers)
+		{
+			TimerManager.SetTimer(Timer.Handle, Timer.Callback(), /* InRate */ 1.0f, FTimerManagerTimerParameters{});
+			check(TimerManager.IsTimerActive(Timer.Handle));
+		}
+
+		AutoRTFM::Testing::Commit([&]
+		{
+			for (size_t I = 0; I < NumTimers; I++)
+			{
+				if ((I & 1) == 0)
+				{
+					FTimer& Timer = Timers[I];
+					check(!TimerManager.IsTimerPaused(Timer.Handle));
+					TimerManager.PauseTimer(Timer.Handle);
+					check(TimerManager.IsTimerPaused(Timer.Handle));
+				}
+			}
+		});
+
+		Tick(TimerManager);
+
+		for (size_t I = 0; I < NumTimers; I++)
+		{
+			FTimer& Timer = Timers[I];
+			bool const bPaused = (I & 1) == 0;
+			check(Timer.bCalled == !bPaused);
+			check(TimerManager.IsTimerPaused(Timer.Handle) == bPaused);
+		}
+
+		return true;
+	});
+
+	Section(TEXT("SetTimer, Transact(PauseTimer, Abort), Tick"), [&] (FTimerManager& TimerManager)
+	{
+		constexpr size_t NumTimers = 16;
+		FTimer Timers[NumTimers]{};
+
+		for (FTimer& Timer : Timers)
+		{
+			TimerManager.SetTimer(Timer.Handle, Timer.Callback(), /* InRate */ 1.0f, FTimerManagerTimerParameters{});
+			check(TimerManager.IsTimerActive(Timer.Handle));
+		}
+
+		AutoRTFM::Testing::Abort([&]
+		{
+			for (size_t I = 0; I < NumTimers; I++)
+			{
+				if ((I & 1) == 0)
+				{
+					FTimer& Timer = Timers[I];
+					TimerManager.PauseTimer(Timer.Handle);
+					check(TimerManager.IsTimerPaused(Timer.Handle));
+				}
+			}
+			AutoRTFM::AbortTransaction();
+		});
+
+		Tick(TimerManager);
+
+		for (size_t I = 0; I < NumTimers; I++)
+		{
+			FTimer& Timer = Timers[I];
+			check(Timer.bCalled);
+			check(!TimerManager.IsTimerPaused(Timer.Handle));
+		}
+
+		return true;
+	});
+
+	Section(TEXT("SetTimer, PauseTimer, Transact(UnPauseTimer), Tick"), [&] (FTimerManager& TimerManager)
+	{
+		constexpr size_t NumTimers = 16;
+		FTimer Timers[NumTimers]{};
+
+		for (FTimer& Timer : Timers)
+		{
+			TimerManager.SetTimer(Timer.Handle, Timer.Callback(), /* InRate */ 1.0f, FTimerManagerTimerParameters{});
+			check(TimerManager.IsTimerActive(Timer.Handle));
+			TimerManager.PauseTimer(Timer.Handle);
+			check(TimerManager.IsTimerPaused(Timer.Handle));
+		}
+
+		AutoRTFM::Testing::Commit([&]
+		{
+			for (size_t I = 0; I < NumTimers; I++)
+			{
+				if ((I & 1) == 0)
+				{
+					FTimer& Timer = Timers[I];
+					check(TimerManager.IsTimerPaused(Timer.Handle));
+					TimerManager.UnPauseTimer(Timer.Handle);
+					check(!TimerManager.IsTimerPaused(Timer.Handle));
+				}
+			}
+		});
+
+		Tick(TimerManager);
+
+		for (size_t I = 0; I < NumTimers; I++)
+		{
+			FTimer& Timer = Timers[I];
+			bool const bPaused = (I & 1) != 0;
+			check(Timer.bCalled == !bPaused);
+			check(TimerManager.IsTimerPaused(Timer.Handle) == bPaused);
+		}
+
+		return true;
+	});
+
+	Section(TEXT("SetTimer, PauseTimer, Transact(UnPauseTimer, Abort), Tick"), [&] (FTimerManager& TimerManager)
+	{
+		constexpr size_t NumTimers = 16;
+		FTimer Timers[NumTimers]{};
+
+		for (FTimer& Timer : Timers)
+		{
+			TimerManager.SetTimer(Timer.Handle, Timer.Callback(), /* InRate */ 1.0f, FTimerManagerTimerParameters{});
+			check(TimerManager.IsTimerActive(Timer.Handle));
+			TimerManager.PauseTimer(Timer.Handle);
+			check(TimerManager.IsTimerPaused(Timer.Handle));
+		}
+
+		AutoRTFM::Testing::Abort([&]
+		{
+			for (size_t I = 0; I < NumTimers; I++)
+			{
+				if ((I & 1) == 0)
+				{
+					FTimer& Timer = Timers[I];
+					check(TimerManager.IsTimerPaused(Timer.Handle));
+					TimerManager.UnPauseTimer(Timer.Handle);
+					check(!TimerManager.IsTimerPaused(Timer.Handle));
+				}
+			}
+			AutoRTFM::AbortTransaction();
+		});
+
+		Tick(TimerManager);
+
+		for (size_t I = 0; I < NumTimers; I++)
+		{
+			FTimer& Timer = Timers[I];
+			check(!Timer.bCalled);
+			check(TimerManager.IsTimerPaused(Timer.Handle));
+		}
+
+		return true;
+	});
+
+	Section(TEXT("Transact(SetTimer, ClearTimer, Open(SetTimer), Abort), Tick"), [&] (FTimerManager& TimerManager)
+	{
+		constexpr size_t NumTimers = 16;
+		FTimer ClosedTimers[NumTimers]{};
+		FTimer OpenTimers[NumTimers]{};
+
+		AutoRTFM::Testing::Abort([&]
+		{
+			for (FTimer& Timer : ClosedTimers)
+			{
+				TimerManager.SetTimer(Timer.Handle, Timer.Callback(), /* InRate */ 1.0f, FTimerManagerTimerParameters{});
+				check(TimerManager.IsTimerActive(Timer.Handle));
+			}
+
+			for (FTimer& Timer : ClosedTimers)
+			{
+				check(TimerManager.IsTimerActive(Timer.Handle));
+				ClearTimer(TimerManager, Timer.Handle);
+				check(!TimerManager.IsTimerActive(Timer.Handle));
+			}
+
+			AutoRTFM::Open([&]
+			{
+				for (FTimer& Timer : OpenTimers)
+				{
+					TimerManager.SetTimer(Timer.Handle, Timer.Callback(), /* InRate */ 1.0f, FTimerManagerTimerParameters{});
+				}
+			});
+
+			AutoRTFM::AbortTransaction();
+		});
+
+		Tick(TimerManager);
+
+		for (FTimer& Timer : ClosedTimers)
+		{
+			check(!Timer.bCalled);
+		}
+		for (FTimer& Timer : OpenTimers)
+		{
+			check(Timer.bCalled);
+		}
+
+		return true;
+	});
+
+	Section(TEXT("Transact(SetTimer, ClearTimer, ClearTimer, Abort)"), [&] (FTimerManager& TimerManager)
+	{
+		FTimer Timer{};
+
+		AutoRTFM::Testing::Abort([&]
+		{
+			TimerManager.SetTimer(Timer.Handle, Timer.Callback(), /* InRate */ 1.0f, FTimerManagerTimerParameters{});
+
+			ClearTimer(TimerManager, Timer.Handle);
+
+			ClearTimer(TimerManager, Timer.Handle);
+			
+			AutoRTFM::AbortTransaction();
+		});
+
+		return true;
+	});
+
+	Section(TEXT("Transact(SetTimer, ClearTimer, PauseTimer, Abort)"), [&] (FTimerManager& TimerManager)
+	{
+		FTimer Timer{};
+
+		AutoRTFM::Testing::Abort([&]
+		{
+			TimerManager.SetTimer(Timer.Handle, Timer.Callback(), /* InRate */ 1.0f, FTimerManagerTimerParameters{});
+
+			ClearTimer(TimerManager, Timer.Handle);
+
+			TimerManager.PauseTimer(Timer.Handle);
+			
+			AutoRTFM::AbortTransaction();
+		});
+
+		return true;
+	});
+
+	Section(TEXT("Transact(SetTimerForNextTick, ClearTimer, ClearTimer, Abort)"), [&] (FTimerManager& TimerManager)
+	{
+		FTimer Timer{};
+
+		AutoRTFM::Testing::Abort([&]
+		{
+			Timer.Handle = TimerManager.SetTimerForNextTick(Timer.Callback());
+
+			ClearTimer(TimerManager, Timer.Handle);
+
+			ClearTimer(TimerManager, Timer.Handle);
+			
+			AutoRTFM::AbortTransaction();
+		});
+
+		return true;
+	});
+
+	Section(TEXT("Transact(SetTimerForNextTick, ClearTimer, PauseTimer, Abort)"), [&] (FTimerManager& TimerManager)
+	{
+		FTimer Timer{};
+
+		AutoRTFM::Testing::Abort([&]
+		{
+			Timer.Handle = TimerManager.SetTimerForNextTick(Timer.Callback());
+
+			ClearTimer(TimerManager, Timer.Handle);
+
+			TimerManager.PauseTimer(Timer.Handle);
+			
+			AutoRTFM::AbortTransaction();
+		});
+
+		return true;
+	});
+
+	Section(TEXT("SetTimerForNextTick(Transact(ClearTimer, ClearTimer)))"), [&] (FTimerManager& TimerManager)
+	{
+		FTimerHandle TimerHandle;
+		TimerHandle = TimerManager.SetTimerForNextTick([&]
+		{
+			AutoRTFM::Testing::Commit([&]
+			{
+				ClearTimer(TimerManager, TimerHandle);
+
+				ClearTimer(TimerManager, TimerHandle);
+			});
+		});
+
+		Tick(TimerManager);
+
+		return true;
+	});
+	
+	return true;
+}
+
+#endif //WITH_DEV_AUTOMATION_TESTS
diff --git a/Engine/Source/Runtime/Engine/Private/TimerManager.cpp b/Engine/Source/Runtime/Engine/Private/TimerManager.cpp
--- a/Engine/Source/Runtime/Engine/Private/TimerManager.cpp
+++ b/Engine/Source/Runtime/Engine/Private/TimerManager.cpp
@@ -5,13 +5,17 @@
 =============================================================================*/
 
 #include "TimerManager.h"
+#include "AutoRTFM.h"
 #include "Containers/StringConv.h"
 #include "Engine/GameInstance.h"
+#include "Engine/TimerHandle.h"
 #include "Logging/LogScopedCategoryAndVerbosityOverride.h"
+#include "Misc/AssertionMacros.h"
 #include "Misc/CoreDelegates.h"
 #include "Misc/StringOutputDevice.h"
 #include "ProfilingDebugging/CsvProfiler.h"
 #include "Stats/StatsTrace.h"
+#include "Templates/UnrealTemplate.h"
 #include "UnrealEngine.h"
 #include "Misc/TimeGuard.h"
 #include "HAL/PlatformStackWalk.h"
@@ -60,6 +64,20 @@ static FAutoConsoleVariableRef CVarTimerManagerGuaranteeEngineTickDelay(
 #define UE_ENABLE_DUMPALLTIMERLOGSTHRESHOLD !UE_BUILD_SHIPPING
 #endif
 
+// Verbose debug logging - especially useful for debugging transactional behavior.
+#define UE_TIMER_MANAGER_ENABLE_VERBOSE_LOGGING 0
+#if UE_TIMER_MANAGER_ENABLE_VERBOSE_LOGGING
+static_assert(!UE_BUILD_SHIPPING);
+#define UE_TIMER_MANAGER_DEBUG_LOG(MSG, ...) \
+	UE_LOG(LogEngine, Log, TEXT("%hs:%d %hs " MSG), __FILE__, __LINE__, AutoRTFM::IsClosed() ? "[closed]" : "  [open]", ##__VA_ARGS__)
+#define UE_TIMER_MANAGER_DEBUG_SCOPE(MSG) \
+	UE_TIMER_MANAGER_DEBUG_LOG(MSG " -- begin") \
+	ON_SCOPE_EXIT { UE_TIMER_MANAGER_DEBUG_LOG(MSG " -- end"); }
+#else
+#define UE_TIMER_MANAGER_DEBUG_LOG(MSG, ...)
+#define UE_TIMER_MANAGER_DEBUG_SCOPE(MSG)
+#endif
+
 #if UE_ENABLE_DUMPALLTIMERLOGSTHRESHOLD
 static int32 DumpAllTimerLogsThreshold = -1;
 static FAutoConsoleVariableRef CVarDumpAllTimerLogsThreshold(
@@ -379,38 +397,21 @@ void FTimerUnifiedDelegate::Execute() const
 	case FTimerDelegateVariant::IndexOfType<FTimerDelegate>():
 		{
 			const FTimerDelegate& FuncDelegate = VariantDelegate.Get<FTimerDelegate>();
-
-#if UE_WITH_REMOTE_OBJECT_HANDLE && UE_AUTORTFM
-			const FName WorkName{ TEXT("TimerManager::Execute FTimerDelegate") };
-			UE::RemoteExecutor::ExecuteTransactional(WorkName, [&FuncDelegate]()
+			if (FuncDelegate.IsBound())
 			{
-#endif
-				if (FuncDelegate.IsBound())
-				{
-					FScopeCycleCounterUObject Context(FuncDelegate.GetUObject());
-					FuncDelegate.Execute();
-				}
-#if UE_WITH_REMOTE_OBJECT_HANDLE && UE_AUTORTFM
-			});
-#endif
+				FScopeCycleCounterUObject Context(FuncDelegate.GetUObject());
+				FuncDelegate.Execute();
+			}
 			break;
 		}
 	case FTimerDelegateVariant::IndexOfType<FTimerDynamicDelegate>():
 		{
 			const FTimerDynamicDelegate& FuncDynDelegate = VariantDelegate.Get<FTimerDynamicDelegate>();
-#if UE_WITH_REMOTE_OBJECT_HANDLE && UE_AUTORTFM
-			const FName WorkName { TEXT("TimerManager::Execute FTimerDynamicDelegate") };
-			UE::RemoteExecutor::ExecuteTransactional(WorkName, [&FuncDynDelegate]()
+			if (FuncDynDelegate.IsBound())
 			{
-#endif
-				if (FuncDynDelegate.IsBound())
-				{
-					// stat scope is handled by UObject::ProcessEvent for the UFunction.
-					FuncDynDelegate.ProcessDelegate<UObject>(nullptr);
-				}
-#if UE_WITH_REMOTE_OBJECT_HANDLE && UE_AUTORTFM
-			});
-#endif
+				// stat scope is handled by UObject::ProcessEvent for the UFunction.
+				FuncDynDelegate.ProcessDelegate<UObject>(nullptr);
+			}
 			break;
 		}
 	case FTimerDelegateVariant::IndexOfType<FTimerFunction>():
@@ -634,6 +635,8 @@ void FTimerManager::InternalSetTimer(FTimerHandle& InOutHandle, FTimerUnifiedDel
 
 void FTimerManager::InternalSetTimer(FTimerHandle& InOutHandle, FTimerUnifiedDelegate&& InDelegate, float InRate, const FTimerManagerTimerParameters& InTimerParameters)
 {
+	UE_TIMER_MANAGER_DEBUG_LOG("InternalSetTimer(InOutHandle: %s)", *InOutHandle.ToString());
+
 	SCOPE_CYCLE_COUNTER(STAT_SetTimer);
 
 	// not currently threadsafe
@@ -643,7 +646,7 @@ void FTimerManager::InternalSetTimer(FTimerHandle& InOutHandle, FTimerUnifiedDel
 	ensureAlwaysMsgf(!InDelegate.VariantDelegate.IsType<FTimerFunction>(), TEXT("Non-UObject Delegates (TFunctions) are not safe with DSTM. Use FTimerDelegate::CreateUObject and don't capture raw pointers."));
 #endif
 
-	if (FindTimer(InOutHandle))
+	if (InOutHandle.IsValid() && (AutoRTFM::IsClosed() || FindTimer(InOutHandle)))
 	{
 		// if the timer is already set, just clear it and we'll re-add it, since 
 		// there's no data to maintain.
@@ -671,6 +674,24 @@ void FTimerManager::InternalSetTimer(FTimerHandle& InOutHandle, FTimerUnifiedDel
 		const float FirstDelay = (InTimerParameters.FirstDelay >= 0.f) ? InTimerParameters.FirstDelay : InRate;
 
 		FTimerHandle NewTimerHandle;
+#if UE_AUTORTFM
+		if (AutoRTFM::IsClosed())
+		{
+			// Reserve a slot in Timers for this transactional timer.
+			NewTimerHandle = GenerateReservedHandleForTransaction();
+
+			// Add the timer to TransactionalTimers.
+			NewTimerData.ExpireTime = FirstDelay;
+			NewTimerData.Handle = NewTimerHandle;
+			TransactionalTimers.Add(NewTimerHandle, FTransactionalTimer
+				{
+					.NewData = MoveTemp(NewTimerData),
+					.Status = FTransactionalTimer::Pending,
+				});
+			MaybeRegisterTransactionalCallbacks();
+		}
+		else
+#endif  // UE_AUTORTFM
 		if (HasBeenTickedThisFrame())
 		{
 			NewTimerData.ExpireTime = InternalTime + FirstDelay;
@@ -693,10 +714,13 @@ void FTimerManager::InternalSetTimer(FTimerHandle& InOutHandle, FTimerUnifiedDel
 	{
 		InOutHandle.Invalidate();
 	}
+	UE_TIMER_MANAGER_DEBUG_LOG("    InOutHandle: %s", *InOutHandle.ToString());
 }
 
 FTimerHandle FTimerManager::InternalSetTimerForNextTick(FTimerUnifiedDelegate&& InDelegate)
 {
+	UE_TIMER_MANAGER_DEBUG_LOG("InternalSetTimerForNextTick()");
+
 	SCOPE_CYCLE_COUNTER(STAT_SetTimerForNextTick);
 
 	// not currently threadsafe
@@ -707,7 +731,6 @@ FTimerHandle FTimerManager::InternalSetTimerForNextTick(FTimerUnifiedDelegate&&
 	NewTimerData.bLoop = false;
 	NewTimerData.bRequiresDelegate = true;
 	NewTimerData.TimerDelegate = MoveTemp(InDelegate);
-	NewTimerData.ExpireTime = InternalTime;
 
 	// Set level collection
 	const UWorld* const OwningWorld = OwningGameInstance ? OwningGameInstance->GetWorld() : nullptr;
@@ -720,30 +743,63 @@ FTimerHandle FTimerManager::InternalSetTimerForNextTick(FTimerUnifiedDelegate&&
 	const bool bQueueForCurrentFrame = (GuaranteeEngineTickDelay == 0) || HasBeenTickedThisFrame();
 	
 	FTimerHandle NewTimerHandle;
+	
+#if UE_AUTORTFM
+	if (AutoRTFM::IsClosed())
+	{
+		// Reserve a slot in Timers for this transactional timer.
+		NewTimerHandle = GenerateReservedHandleForTransaction();
+		NewTimerData.Handle = NewTimerHandle;
+
+		// Add the timer to TransactionalTimers.
+		TransactionalTimers.FindOrAdd(NewTimerHandle, FTransactionalTimer
+			{
+				.NewData = MoveTemp(NewTimerData),
+				.Status = FTransactionalTimer::Pending,
+			});
+		MaybeRegisterTransactionalCallbacks();
+	}
+	else
+#endif  // UE_AUTORTFM
 	if (bQueueForCurrentFrame)
 	{
 		NewTimerData.Status = ETimerStatus::Active;
+		NewTimerData.ExpireTime = InternalTime;
 		NewTimerHandle = AddTimer(MoveTemp(NewTimerData));
 		ActiveTimerHeap.HeapPush(NewTimerHandle, FTimerHeapOrder(Timers));
 	}
 	else
 	{
 		NewTimerData.Status = ETimerStatus::Pending;
-		NewTimerData.ExpireTime = 0.0;	// Pending timers' ExpireTime means time remaining. It will be converted upon activation.
 		NewTimerHandle = AddTimer(MoveTemp(NewTimerData));
 		PendingTimerSet.Add(NewTimerHandle);
 	}
 
+	UE_TIMER_MANAGER_DEBUG_LOG("    NewTimerHandle: %s", *NewTimerHandle.ToString());
 	return NewTimerHandle;
 }
 
 void FTimerManager::InternalClearTimer(FTimerHandle InHandle)
 {
+	UE_TIMER_MANAGER_DEBUG_LOG("InternalClearTimer(InHandle: %s)", *InHandle.ToString());
+
 	SCOPE_CYCLE_COUNTER(STAT_ClearTimer);
 
 	// not currently threadsafe
 	check(IsInGameThread());
 
+#if UE_AUTORTFM
+	if (AutoRTFM::IsClosed())
+	{
+		// Add or replace the transactional timer with the canceled status.
+		FTransactionalTimer& Timer = TransactionalTimers.FindOrAdd(InHandle);
+		Timer.NewData.Reset();
+		Timer.Status = FTransactionalTimer::Canceled;
+		MaybeRegisterTransactionalCallbacks();
+		return;
+	}
+#endif  // UE_AUTORTFM
+
 	FTimerData& Data = GetTimer(InHandle);
 	switch (Data.Status)
 	{
@@ -780,6 +836,13 @@ void FTimerManager::InternalClearTimer(FTimerHandle InHandle)
 			RemoveTimer(InHandle);
 			break;
 
+		case ETimerStatus::ReservedForTransaction:
+			// Only the on-commit or on-abort handler may remove a
+			// transactionally-reserved timer.
+			check(AutoRTFM::IsCommitting() || AutoRTFM::IsAborting());
+			RemoveTimer(InHandle);
+			break;
+			
 		default:
 			check(false);
 	}
@@ -787,6 +850,8 @@ void FTimerManager::InternalClearTimer(FTimerHandle InHandle)
 
 void FTimerManager::InternalClearAllTimers(void const* Object)
 {
+	UE_TIMER_MANAGER_DEBUG_LOG("InternalClearAllTimers(Object: %p)", Object);
+
 	SCOPE_CYCLE_COUNTER(STAT_ClearAllTimers);
 
 	if (!Object)
@@ -858,9 +923,27 @@ float FTimerManager::InternalGetTimerRate(FTimerData const* const TimerData) con
 
 void FTimerManager::PauseTimer(FTimerHandle InHandle)
 {
+	UE_TIMER_MANAGER_DEBUG_LOG("PauseTimer(InHandle: %s)", *InHandle.ToString());
+
 	// not currently threadsafe
 	check(IsInGameThread());
 
+#if UE_AUTORTFM
+	if (AutoRTFM::IsClosed())
+	{
+		// Add or replace the transactional timer with the paused status (if it
+		// didn't exist with a canceled status).
+		FTransactionalTimer& Timer = TransactionalTimers.FindOrAdd(InHandle);
+		if (Timer.Status != FTransactionalTimer::Canceled)
+		{
+			Timer.Status = FTransactionalTimer::Paused;
+		}
+		MaybeRegisterTransactionalCallbacks();
+		return;
+	}
+#endif  // UE_AUTORTFM
+	
+
 	FTimerData* TimerToPause = FindTimer(InHandle);
 	if (!TimerToPause || TimerToPause->Status == ETimerStatus::Paused)
 	{
@@ -923,9 +1006,27 @@ void FTimerManager::PauseTimer(FTimerHandle InHandle)
 
 void FTimerManager::UnPauseTimer(FTimerHandle InHandle)
 {
+	UE_TIMER_MANAGER_DEBUG_LOG("UnPauseTimer(InHandle: %s)", *InHandle.ToString());
+
 	// not currently threadsafe
 	check(IsInGameThread());
 
+#if UE_AUTORTFM
+	if (AutoRTFM::IsClosed())
+	{
+		// Add or replace the transactional timer with the pending status (if it
+		// didn't exist with a canceled status).
+		FTransactionalTimer& Timer = TransactionalTimers.FindOrAdd(InHandle);
+		if (Timer.Status != FTransactionalTimer::Canceled)
+		{
+			Timer.Status = FTransactionalTimer::Pending;
+		}
+		MaybeRegisterTransactionalCallbacks();
+		return;
+	}
+#endif  // UE_AUTORTFM
+	
+
 	FTimerData* TimerToUnPause = FindTimer(InHandle);
 	if (!TimerToUnPause || TimerToUnPause->Status != ETimerStatus::Paused)
 	{
@@ -968,6 +1069,7 @@ DECLARE_DWORD_COUNTER_STAT(TEXT("TimerManager Heap Size"),STAT_NumHeapEntries,ST
 
 void FTimerManager::Tick(float DeltaTime)
 {
+	AutoRTFM::UnreachableIfClosed();
 	SCOPED_NAMED_EVENT(FTimerManager_Tick, FColor::Orange);
 	CSV_SCOPED_TIMING_STAT_EXCLUSIVE(TimerManager);
 
@@ -1116,8 +1218,39 @@ void FTimerManager::Tick(float DeltaTime)
 				RunTimerDelegates.Add(Top->TimerDelegate);
 #endif
 
-				checkf(!WillRemoveTimerAssert(CurrentlyExecutingTimer), TEXT("RemoveTimer(CurrentlyExecutingTimer) - due to fail before Execute()"));
-				Top->TimerDelegate.Execute();
+#if UE_WITH_REMOTE_OBJECT_HANDLE && UE_AUTORTFM
+				FName WorkName;
+				switch (Top->TimerDelegate.VariantDelegate.GetIndex())
+				{
+				case FTimerDelegateVariant::IndexOfType<FTimerDelegate>():
+					WorkName = FName{ TEXT("TimerManager::Execute FTimerDelegate") };
+					break;
+				case FTimerDelegateVariant::IndexOfType<FTimerDynamicDelegate>():
+					WorkName = FName{ TEXT("TimerManager::Execute FTimerDynamicDelegate") };
+					break;
+				default:
+					break;
+				}
+				if (WorkName.IsValid())
+				{
+					UE_TIMER_MANAGER_DEBUG_LOG("ExecuteTransactional(%s)", *CurrentlyExecutingTimer.ToString());
+					UE::RemoteExecutor::ExecuteTransactional(WorkName, [this]
+					{
+						FTimerData* Top = FindTimer(CurrentlyExecutingTimer);
+						checkf(Top && !WillRemoveTimerAssert(CurrentlyExecutingTimer), TEXT("RemoveTimer(CurrentlyExecutingTimer) - due to fail before Execute()"));
+						checkf(Top->Handle == CurrentlyExecutingTimer,
+							TEXT("CurrentlyExecutingTimer changed between retries: Top->Handle=%s, CurrentlyExecutingTimer=%s"),
+							*Top->Handle.ToString(), *CurrentlyExecutingTimer.ToString());
+						Top->TimerDelegate.Execute();
+					});
+				}
+				else
+#else
+				{
+					checkf(!WillRemoveTimerAssert(CurrentlyExecutingTimer), TEXT("RemoveTimer(CurrentlyExecutingTimer) - due to fail before Execute()"));
+					Top->TimerDelegate.Execute();
+				}
+#endif
 
 				// Update Top pointer, in case it has been invalidated by the Execute call
 				Top = FindTimer(CurrentlyExecutingTimer);
@@ -1265,6 +1398,8 @@ void FTimerManager::ListTimers() const
 
 FTimerHandle FTimerManager::AddTimer(FTimerData&& TimerData)
 {
+	UE_TIMER_MANAGER_DEBUG_LOG("AddTimer()");
+	
 	const void* TimerIndicesByObjectKey = TimerData.TimerDelegate.GetBoundObject();
 	TimerData.TimerIndicesByObjectKey = TimerIndicesByObjectKey;
 
@@ -1275,23 +1410,35 @@ FTimerHandle FTimerManager::AddTimer(FTimerData&& TimerData)
 
 	if (TimerIndicesByObjectKey)
 	{
-		TSet<FTimerHandle>& HandleSet = ObjectToTimers.FindOrAdd(TimerIndicesByObjectKey);
-
-		bool bAlreadyExists = false;
-		HandleSet.Add(Result, &bAlreadyExists);
-		checkf(!bAlreadyExists, TEXT("A timer with this handle and object has already been added! (%s)"), *GetFTimerDataSafely(TimerData));
+		AddObjectToTimers(Result, TimerIndicesByObjectKey);
 	}
 
 	return Result;
 }
 
+void FTimerManager::AddObjectToTimers(FTimerHandle TimerHandle, const void *TimerIndicesByObjectKey)
+{
+	UE_TIMER_MANAGER_DEBUG_LOG("AddObjectToTimers(TimerHandle: %s, TimerIndicesByObjectKey: %p)", *TimerHandle.ToString(), TimerIndicesByObjectKey);
+
+	checkSlow(TimerIndicesByObjectKey)
+	TSet<FTimerHandle>& HandleSet = ObjectToTimers.FindOrAdd(TimerIndicesByObjectKey);
+
+	bool bAlreadyExists = false;
+	HandleSet.Add(TimerHandle, &bAlreadyExists);
+	checkf(!bAlreadyExists, TEXT("A timer with this handle and object has already been added! (%s)"), *TimerHandle.ToString());
+}
+
 void FTimerManager::RemoveTimer(FTimerHandle Handle)
 {
+	UE_TIMER_MANAGER_DEBUG_LOG("RemoveTimer(Handle: %s)", *Handle.ToString());
+
 	const FTimerData& Data = GetTimer(Handle);
 
 	// Remove TimerIndicesByObject entry if necessary
 	if (const void* TimerIndicesByObjectKey = Data.TimerIndicesByObjectKey)
 	{
+		UE_TIMER_MANAGER_DEBUG_LOG("    TimerIndicesByObjectKey: %p", TimerIndicesByObjectKey);
+
 		TSet<FTimerHandle>* TimersForObject = ObjectToTimers.Find(TimerIndicesByObjectKey);
 		checkf(TimersForObject, TEXT("Removed timer was bound to an object which is not tracked by ObjectToTimers! (%s)"), *GetFTimerDataSafely(Data));
 
@@ -1304,7 +1451,22 @@ void FTimerManager::RemoveTimer(FTimerHandle Handle)
 		}
 	}
 
-	Timers.RemoveAt(Handle.GetIndex());
+#if UE_AUTORTFM
+	if (AutoRTFM::IsTransactional() && TransactionalTimers.Find(Handle))
+	{
+		// Removing a timer in the transactional open that has been modified in
+		// the transaction. The timer slot cannot be released until the
+		// transaction is complete, so change the status to ReservedForTransaction.
+		FTimerData Reserved{};
+		Reserved.Handle = Handle;
+		Reserved.Status = ETimerStatus::ReservedForTransaction;
+		Timers[Handle.GetIndex()] = MoveTemp(Reserved);
+	}
+	else
+#endif  // UE_AUTORTFM
+	{
+		Timers.RemoveAt(Handle.GetIndex());
+	}
 }
 
 bool FTimerManager::WillRemoveTimerAssert(FTimerHandle Handle) const
@@ -1343,6 +1505,126 @@ FTimerHandle FTimerManager::GenerateHandle(int32 Index)
 	return Result;
 }
 
+#if UE_AUTORTFM
+void FTimerManager::MaybeRegisterTransactionalCallbacks()
+{
+	if (bRegisteredTransactionalCallbacks)
+	{
+		return;
+	}
+
+	bRegisteredTransactionalCallbacks = true;
+
+	AutoRTFM::OnCommit([this]
+	{
+		UE_TIMER_MANAGER_DEBUG_SCOPE("OnCommit()");
+		
+		for (TTuple<FTimerHandle, FTimerManager::FTransactionalTimer>& It : TransactionalTimers)
+		{
+			const FTimerHandle Handle = It.Key;
+			FTransactionalTimer& TransactionalTimer = It.Value;
+			UE_TIMER_MANAGER_DEBUG_LOG("    TransactionalTimer: {%s, Handle: %s, Status: %d}",
+				TransactionalTimer.NewData ? TEXT("with-data") : TEXT("no-data"),
+				*Handle.ToString(), (int)TransactionalTimer.Status);
+
+			if (TransactionalTimer.NewData)
+			{
+				// Timer was created during the transaction
+				FTimerData& TimerData = *TransactionalTimer.NewData;
+				check(Handle == TimerData.Handle);
+
+				switch (TransactionalTimer.Status)
+				{
+				case FTransactionalTimer::Pending:
+					if (HasBeenTickedThisFrame())
+					{
+						TimerData.ExpireTime += InternalTime;
+						TimerData.Status = ETimerStatus::Active;
+						ActiveTimerHeap.HeapPush(TimerData.Handle, FTimerHeapOrder(Timers));
+					}
+					else
+					{
+						TimerData.Status = ETimerStatus::Pending;
+						PendingTimerSet.Add(TimerData.Handle);
+					}
+					break;
+
+				case FTransactionalTimer::Paused:
+					TimerData.Status = ETimerStatus::Paused;
+					PausedTimerSet.Add(Handle);
+					break;
+
+				case FTransactionalTimer::Canceled:
+					checkf(false, TEXT("FTransactionalTimer::NewData should be empty when Status is Canceled"));
+					break;
+				}
+
+				if (const void* TimerIndicesByObjectKey = TimerData.TimerDelegate.GetBoundObject())
+				{
+					TimerData.TimerIndicesByObjectKey = TimerIndicesByObjectKey;
+					AddObjectToTimers(Handle, TimerIndicesByObjectKey);
+				}
+
+				// Move the transactional timer to the non-transactional Timers array.
+				FTimerData& NonTransactionalTimer = Timers[Handle.GetIndex()];
+				check(Handle == NonTransactionalTimer.Handle);
+				NonTransactionalTimer = MoveTemp(TimerData);
+			}
+			else
+			{
+				switch (TransactionalTimer.Status)
+				{
+				case FTransactionalTimer::Pending:
+					UnPauseTimer(Handle);
+					break;
+
+				case FTransactionalTimer::Paused:
+					PauseTimer(Handle);
+					break;
+
+				case FTransactionalTimer::Canceled:
+					if (FindTimer(Handle))
+					{
+						InternalClearTimer(Handle);
+					}
+					break;
+				}
+			}
+		}
+		TransactionalTimers.Empty();
+		bRegisteredTransactionalCallbacks = false;
+	});
+
+	AutoRTFM::OnPreAbort([this]
+	{
+		UE_TIMER_MANAGER_DEBUG_SCOPE("OnPreAbort()");
+		
+		// Remove transactionally reserved timers.
+		for (auto& It : TransactionalTimers)
+		{
+			const FTimerHandle& Handle = It.Key;
+			const int32 Index = Handle.GetIndex();
+			check(Timers[Index].Handle == Handle);
+			if (Timers[Index].Status == ETimerStatus::ReservedForTransaction)
+			{
+				RemoveTimer(Handle);
+			}
+		}
+		// Note: TransactionalTimers & bRegisteredTransactionalCallbacks will be
+		// cleared automatically by AutoRTFM memory rollback.
+	});
+}
+
+FTimerHandle FTimerManager::GenerateReservedHandleForTransaction()
+{
+	FTimerData Reservation;
+	Reservation.Status = ETimerStatus::ReservedForTransaction;
+	int32 NewIndex = Timers.Add(MoveTemp(Reservation));
+	FTimerHandle Handle = GenerateHandle(NewIndex);
+	Timers[NewIndex].Handle = Handle;
+	return Handle;
+}
+#endif  // UE_AUTORTFM
 
 // Handler for ListTimers console command
 static void OnListTimers(UWorld* World)
@@ -1359,3 +1641,9 @@ FAutoConsoleCommandWithWorld ListTimersConsoleCommand(
 	TEXT(""),
 	FConsoleCommandWithWorldDelegate::CreateStatic(OnListTimers)
 	);
+
+#if UE_TIMER_MANAGER_ENABLE_VERBOSE_LOGGING
+#undef UE_TIMER_MANAGER_DEBUG_LOG
+#undef UE_TIMER_MANAGER_DEBUG_SCOPE
+#endif
+#undef UE_TIMER_MANAGER_ENABLE_VERBOSE_LOGGING
diff --git a/Engine/Source/Runtime/Engine/Private/TimerManagerTests.cpp b/Engine/Source/Runtime/Engine/Private/TimerManagerTests.cpp
--- a/Engine/Source/Runtime/Engine/Private/TimerManagerTests.cpp
+++ b/Engine/Source/Runtime/Engine/Private/TimerManagerTests.cpp
@@ -35,7 +35,7 @@ struct FDummy
 };
 
 // Make sure that the timer manager works as expected when given invalid delegates and handles
-bool TimerManagerTest_InvalidTimers(UWorld* World, FAutomationTestBase* Test)
+AUTORTFM_DISABLE bool TimerManagerTest_InvalidTimers(UWorld* World, FAutomationTestBase* Test)
 {
 	FTimerManager& TimerManager = World->GetTimerManager();
 	FTimerHandle Handle;
@@ -56,7 +56,7 @@ bool TimerManagerTest_InvalidTimers(UWorld* World, FAutomationTestBase* Test)
 }
 
 // Make sure that the timer manager works as expected when given delegates and handles that aren't in the timer manager
-bool TimerManagerTest_MissingTimers(UWorld* World, FAutomationTestBase* Test)
+AUTORTFM_DISABLE bool TimerManagerTest_MissingTimers(UWorld* World, FAutomationTestBase* Test)
 {
 	FTimerManager& TimerManager = World->GetTimerManager();
 	FTimerHandle Handle = TimerManager.GenerateHandle(123);
@@ -76,7 +76,7 @@ bool TimerManagerTest_MissingTimers(UWorld* World, FAutomationTestBase* Test)
 	return true;
 }
 
-bool TimerManagerTest_ValidTimer_HandleWithDelegate(UWorld* World, FAutomationTestBase* Test)
+AUTORTFM_DISABLE bool TimerManagerTest_ValidTimer_HandleWithDelegate(UWorld* World, FAutomationTestBase* Test)
 {
 	FTimerManager& TimerManager = World->GetTimerManager();
 	FTimerDelegate Delegate;
@@ -193,7 +193,7 @@ FTimerHandle* FLoopingTestFunc::Handle = nullptr;
 int FLoopingTestFunc::TimerCalled = 0;
 float FLoopingTestFunc::NewTime = 1.0f;
 
-bool TimerManagerTest_ValidTimer_HandleLoopingSetDuringExecute(UWorld* World, FAutomationTestBase* Test)
+AUTORTFM_DISABLE bool TimerManagerTest_ValidTimer_HandleLoopingSetDuringExecute(UWorld* World, FAutomationTestBase* Test)
 {
 	FTimerManager& TimerManager = World->GetTimerManager();
 	FTimerHandle Handle;
@@ -222,7 +222,7 @@ bool TimerManagerTest_ValidTimer_HandleLoopingSetDuringExecute(UWorld* World, FA
 	return true;
 }
 
-bool TimerManagerTest_LoopingTimers_DifferentHandles(UWorld* World, FAutomationTestBase* Test)
+AUTORTFM_DISABLE bool TimerManagerTest_LoopingTimers_DifferentHandles(UWorld* World, FAutomationTestBase* Test)
 {
 	FTimerManager& TimerManager = World->GetTimerManager();
 	FTimerHandle HandleOne, HandleTwo;
diff --git a/Engine/Source/Runtime/Engine/Public/TimerManager.h b/Engine/Source/Runtime/Engine/Public/TimerManager.h
--- a/Engine/Source/Runtime/Engine/Public/TimerManager.h
+++ b/Engine/Source/Runtime/Engine/Public/TimerManager.h
@@ -6,6 +6,7 @@
 
 #pragma once
 
+#include "AutoRTFM.h"
 #include "CoreMinimal.h"
 #include "Stats/Stats.h"
 #include "UObject/Object.h"
@@ -66,7 +67,8 @@ enum class ETimerStatus : uint8
 	Active,
 	Paused,
 	Executing,
-	ActivePendingRemoval
+	ActivePendingRemoval,
+	ReservedForTransaction,
 };
 
 struct FTimerData
@@ -128,6 +130,9 @@ struct FTimerManagerTimerParameters
 
 /** 
  * Class to globally manage timers.
+ *
+ * FTimerManager supports transactional usage from both the open and closed.
+ * Open modifications may be clobbered with transactional changes on abort.
  */
 class FTimerManager : public FNoncopyable
 {
@@ -285,6 +290,7 @@ public:
 	/** Clears all timers that are bound to functions on the given object. */
 	inline void ClearAllTimersForObject(void const* Object)
 	{
+		AutoRTFM::UnreachableIfClosed();
 		if (Object)
 		{
 			InternalClearAllTimers( Object );
@@ -325,6 +331,16 @@ public:
 	 */
 	inline bool IsTimerActive(FTimerHandle InHandle) const
 	{
+#if UE_AUTORTFM
+		if (AutoRTFM::IsClosed())
+		{
+			if (const FTransactionalTimer* Timer = TransactionalTimers.Find(InHandle))
+			{
+				return Timer->Status == FTransactionalTimer::Pending;
+			}
+		}
+#endif  // UE_AUTORTFM
+		
 		FTimerData const* const TimerData = FindTimer( InHandle );
 		return TimerData && TimerData->Status != ETimerStatus::Paused;
 	}
@@ -337,6 +353,16 @@ public:
 	*/
 	inline bool IsTimerPaused(FTimerHandle InHandle) const
 	{
+#if UE_AUTORTFM
+		if (AutoRTFM::IsClosed())
+		{
+			if (const FTransactionalTimer* Timer = TransactionalTimers.Find(InHandle))
+			{
+				return Timer->Status == FTransactionalTimer::Paused;
+			}
+		}
+#endif  // UE_AUTORTFM
+		
 		FTimerData const* const TimerData = FindTimer(InHandle);
 		return TimerData && TimerData->Status == ETimerStatus::Paused;
 	}
@@ -349,6 +375,16 @@ public:
 	*/
 	inline bool IsTimerPending(FTimerHandle InHandle) const
 	{
+#if UE_AUTORTFM
+		if (AutoRTFM::IsClosed())
+		{
+			if (const FTransactionalTimer* Timer = TransactionalTimers.Find(InHandle))
+			{
+				return Timer->Status == FTransactionalTimer::Pending;
+			}
+		}
+#endif  // UE_AUTORTFM
+
 		FTimerData const* const TimerData = FindTimer(InHandle);
 		return TimerData && TimerData->Status == ETimerStatus::Pending;
 	}
@@ -361,6 +397,16 @@ public:
 	*/
 	inline bool TimerExists(FTimerHandle InHandle) const
 	{
+#if UE_AUTORTFM
+		if (AutoRTFM::IsClosed())
+		{
+			if (const FTransactionalTimer* Timer = TransactionalTimers.Find(InHandle))
+			{
+				return Timer->Status != FTransactionalTimer::Canceled;
+			}
+		}
+#endif  // UE_AUTORTFM
+
 		return FindTimer(InHandle) != nullptr;
 	}
 
@@ -412,7 +458,7 @@ private:
 // This should be private, but needs to be public for testing.
 public:
 	/** Generates a handle for a timer at a given index */
-	ENGINE_API FTimerHandle GenerateHandle(int32 Index);
+	AUTORTFM_DISABLE ENGINE_API FTimerHandle GenerateHandle(int32 Index);
 
 // These should be private, but need to be protected so IMPLEMENT_GET_PROTECTED_FUNC works for testing.
 protected:
@@ -428,7 +474,7 @@ private:
 	ENGINE_API void InternalSetTimer( FTimerHandle& InOutHandle, FTimerUnifiedDelegate&& InDelegate, float InRate, const FTimerManagerTimerParameters& InTimerParameters );
 	ENGINE_API FTimerHandle InternalSetTimerForNextTick( FTimerUnifiedDelegate&& InDelegate );
 	ENGINE_API void InternalClearTimer( FTimerHandle InDelegate );
-	ENGINE_API void InternalClearAllTimers( void const* Object );
+	AUTORTFM_DISABLE ENGINE_API void InternalClearAllTimers( void const* Object );
 	ENGINE_API float InternalGetTimerRate( FTimerData const* const TimerData ) const;
 	ENGINE_API float InternalGetTimerElapsed( FTimerData const* const TimerData ) const;
 	ENGINE_API float InternalGetTimerRemaining( FTimerData const* const TimerData ) const;
@@ -441,11 +487,25 @@ private:
 	ENGINE_API FTimerData& GetTimer(FTimerHandle const& InHandle);
 
 	/** Adds a timer from the Timers list, also updating the TimerIndicesByObject map.  Returns the insertion index. */
-	ENGINE_API FTimerHandle AddTimer(FTimerData&& TimerData);
+	AUTORTFM_DISABLE ENGINE_API FTimerHandle AddTimer(FTimerData&& TimerData);
+	/** Adds TimerHandle to the handle set of ObjectToTimers with the key TimerIndicesByObjectKey */
+	AUTORTFM_DISABLE ENGINE_API void AddObjectToTimers(FTimerHandle TimerHandle, const void* TimerIndicesByObjectKey);
+	
 	/** Removes a timer from the Timers list at the given index, also cleaning up the TimerIndicesByObject map */
-	ENGINE_API void RemoveTimer(FTimerHandle Handle);
+	AUTORTFM_DISABLE ENGINE_API void RemoveTimer(FTimerHandle Handle);
 	ENGINE_API bool WillRemoveTimerAssert(FTimerHandle Handle) const;
 
+#if UE_AUTORTFM
+	/**
+	 * Registers an on-commit and on-abort handler to transactionally commit or
+	 * revert the transactional timers in TransactionalTimers, if these handlers
+	 * have not already been registered.
+	 */
+	ENGINE_API void MaybeRegisterTransactionalCallbacks();
+	/** Adds a new FTimerData with the ReservedForTransaction status, returning its unique handle */
+	AUTORTFM_OPEN ENGINE_API FTimerHandle GenerateReservedHandleForTransaction();
+#endif  // UE_AUTORTFM
+
 	/** The array of timers - all other arrays will index into this */
 	TSparseArray<FTimerData> Timers;
 	/** Heap of actively running timers. */
@@ -457,6 +517,43 @@ private:
 	/** A map of object pointers to timers with delegates bound to those objects, for quick lookup */
 	TMap<const void*, TSet<FTimerHandle>> ObjectToTimers;
 
+#if UE_AUTORTFM
+	/** FTransactionalTimer holds the state of a timer modified in the closed of a transaction */
+	struct FTransactionalTimer
+	{
+		// The transactional timer state.
+		enum EStatus
+		{
+			// On commit of the transaction, the timer will be put into either
+			// the Pending or Active state (depending on the value of 
+			// HasBeenTickedThisFrame() when the transaction commits).
+			Pending,
+			// On commit of the transaction, the timer will be put into the
+			// Paused state.
+			Paused,
+			// On commit of the transaction, the timer will be removed.
+			Canceled,
+		};
+		// Timer data, if the timer was transactionally created.
+		// Note that FTransactionalTimer::Status is used instead of
+		// FTimerData::Status, and NewData->Status is ignored.
+		TOptional<FTimerData> NewData;
+		// The current status of the transactional timer.
+		EStatus Status = Pending;
+	};
+
+	/**
+	 * A map of timer handle to transactional timers, which will be merged into
+	 * Timers on transaction commit, or discarded on transaction abort.
+	 * The timer handle key indexes to a slot in Timers which must be preserved
+	 * for the duration of the transaction.
+	 */
+	TMap<FTimerHandle, FTransactionalTimer> TransactionalTimers;
+
+	/** True if on-commit and on-abort handlers have been registered for this transaction nest. */
+	bool bRegisteredTransactionalCallbacks = false; 
+#endif  // UE_AUTORTFM
+
 	/** An internally consistent clock, independent of World.  Advances during ticking. */
 	double InternalTime;
 
