UPROPERTY Markup: Attenuation, AudioBus Send, Submix Send

- Consistency pass wherever Min/Max ranges are defined
- Reordering: Non Spatialized Radius End/Start -> Non Spatialized Distance Min/Max
- Renaming: Manual -> Constant
- Reordering Air Absorption Method at the top under the enablement checkbox
- Various tooltip fixups
- Hz unit assignment to LPF/HPF

[RN] minor, audio

#Synced-CL 51085281

diff --git a/Engine/Source/Editor/DetailCustomizations/Private/AttenuationSettingsCustomizations.cpp b/Engine/Source/Editor/DetailCustomizations/Private/AttenuationSettingsCustomizations.cpp
--- a/Engine/Source/Editor/DetailCustomizations/Private/AttenuationSettingsCustomizations.cpp
+++ b/Engine/Source/Editor/DetailCustomizations/Private/AttenuationSettingsCustomizations.cpp
@@ -529,15 +529,15 @@ void FSoundAttenuationSettingsCustomization::CustomizeChildren(TSharedRef<IPrope
 		.Visibility(IsAttenuationOverriddenVisibleAttribute())
 		.EditCondition(GetIsSpatializationEnabledAttribute(), nullptr);
 
-	LayoutBuilder.AddPropertyToCategory(PropertyHandles.FindChecked(GET_MEMBER_NAME_CHECKED(FSoundAttenuationSettings, NonSpatializedRadiusStart)))
+	LayoutBuilder.AddPropertyToCategory(PropertyHandles.FindChecked(GET_MEMBER_NAME_CHECKED(FSoundAttenuationSettings, NonSpatializedRadiusMode)))
 		.Visibility(IsAttenuationOverriddenVisibleAttribute())
 		.EditCondition(GetIsSpatializationEnabledAttribute(), nullptr);
-
+	
 	LayoutBuilder.AddPropertyToCategory(PropertyHandles.FindChecked(GET_MEMBER_NAME_CHECKED(FSoundAttenuationSettings, NonSpatializedRadiusEnd)))
 		.Visibility(IsAttenuationOverriddenVisibleAttribute())
 		.EditCondition(GetIsSpatializationEnabledAttribute(), nullptr);
-
-	LayoutBuilder.AddPropertyToCategory(PropertyHandles.FindChecked(GET_MEMBER_NAME_CHECKED(FSoundAttenuationSettings, NonSpatializedRadiusMode)))
+	
+	LayoutBuilder.AddPropertyToCategory(PropertyHandles.FindChecked(GET_MEMBER_NAME_CHECKED(FSoundAttenuationSettings, NonSpatializedRadiusStart)))
 		.Visibility(IsAttenuationOverriddenVisibleAttribute())
 		.EditCondition(GetIsSpatializationEnabledAttribute(), nullptr);
 
@@ -557,6 +557,10 @@ void FSoundAttenuationSettingsCustomization::CustomizeChildren(TSharedRef<IPrope
 		.Visibility(IsAttenuationOverriddenVisibleAttribute())
 		.EditCondition(IsAttenuationOverriddenAttribute(), nullptr);
 
+	LayoutBuilder.AddPropertyToCategory(AbsorptionMethodHandle)
+		.Visibility(IsAttenuationOverriddenVisibleAttribute())
+		.EditCondition(GetIsAirAbsorptionEnabledAttribute(), nullptr);
+	
 	LayoutBuilder.AddPropertyToCategory(PropertyHandles.FindChecked(GET_MEMBER_NAME_CHECKED(FSoundAttenuationSettings, LPFRadiusMin)))
 		.Visibility(IsAttenuationOverriddenVisibleAttribute())
 		.EditCondition(GetIsAirAbsorptionEnabledAttribute(), nullptr);
@@ -585,10 +589,6 @@ void FSoundAttenuationSettingsCustomization::CustomizeChildren(TSharedRef<IPrope
 		.Visibility(IsAttenuationOverriddenVisibleAttribute())
 		.EditCondition(GetIsAirAbsorptionEnabledAttribute(), nullptr);
 
-	LayoutBuilder.AddPropertyToCategory(AbsorptionMethodHandle)
-		.Visibility(IsAttenuationOverriddenVisibleAttribute())
-		.EditCondition(GetIsAirAbsorptionEnabledAttribute(), nullptr);
-
 	LayoutBuilder.AddPropertyToCategory(PropertyHandles.FindChecked(GET_MEMBER_NAME_CHECKED(FSoundAttenuationSettings, CustomLowpassAirAbsorptionCurve)))
 		.Visibility(TAttribute<EVisibility>(this, &FSoundAttenuationSettingsCustomization::IsCustomAirAbsorptionCurveSelected))
 		.EditCondition(GetIsAirAbsorptionEnabledAttribute(), nullptr);
diff --git a/Engine/Source/Runtime/Engine/Classes/Sound/SoundAttenuation.h b/Engine/Source/Runtime/Engine/Classes/Sound/SoundAttenuation.h
--- a/Engine/Source/Runtime/Engine/Classes/Sound/SoundAttenuation.h
+++ b/Engine/Source/Runtime/Engine/Classes/Sound/SoundAttenuation.h
@@ -40,38 +40,38 @@ enum ESoundSpatializationAlgorithm : int
 UENUM(BlueprintType)
 enum class EAirAbsorptionMethod : uint8
 {
-	// The air absorption conform to a linear distance function
-	Linear,
+	// Air absorption based on linear interpolation between a distance range and a remapping range
+	Linear UMETA(DisplayName = "Distance Linear"),
 
-	// The air absorption conforms to a custom distance curve.
-	CustomCurve,
+	// Air absorption based on a supplied curve mapping (distance = X, air absorption = Y)
+	CustomCurve UMETA(DisplayName = "Distance Custom"),
 };
 
 
 UENUM(BlueprintType)
 enum class EReverbSendMethod : uint8
 {
-	// A reverb send based on linear interpolation between a distance range and send-level range
-	Linear,
+	// Reverb send based on linear interpolation between a distance range and a remapping range
+	Linear UMETA(DisplayName = "Distance Linear"),
 
-	// A reverb send based on a supplied curve
-	CustomCurve,
+	// Reverb send based on a supplied curve mapping (distance = X, send level = Y)
+	CustomCurve UMETA(DisplayName = "Distance Custom"),
 
-	// A manual reverb send level (Uses the specified constant send level value. Useful for 2D sounds.)
-	Manual,
+	// Constant reverb send level (Uses the specified constant send level value. Useful for 2D sounds.)
+	Manual UMETA(DisplayName = "Constant"),
 };
 
 UENUM(BlueprintType)
 enum class EPriorityAttenuationMethod : uint8
 {
-	// A priority attenuation based on linear interpolation between a distance range and priority attenuation range
-	Linear,
+	// Priority attenuation based on linear interpolation between a distance range and a remapping range
+	Linear UMETA(DisplayName = "Distance Linear"),
 
-	// A priority attenuation based on a supplied curve
-	CustomCurve,
+	// Priority attenuation based on a supplied curve mapping (distance = X, priority = Y)
+	CustomCurve UMETA(DisplayName = "Distance Custom"),
 
-	// A manual priority attenuation (Uses the specified constant value. Useful for 2D sounds.)
-	Manual,
+	// Constant priority attenuation (Uses the specified constant value. Useful for 2D sounds.)
+	Manual UMETA(DisplayName = "Constant"),
 };
 
 
@@ -109,7 +109,7 @@ struct TStructOpsTypeTraits<FSoundAttenuationPluginSettings> : public TStructOps
 };
 
 
-// Defines how to speaker map the sound when using the non-spatialized radius feature
+// Defines how to speaker map the sound when using the Non Spatialized Distance feature
 UENUM(BlueprintType)
 enum class ENonSpatializedRadiusSpeakerMapMode : uint8
 {
@@ -237,18 +237,18 @@ struct FSoundAttenuationSettings : public FBaseAttenuationSettings
 
  	UPROPERTY()
  	float OmniRadius_DEPRECATED;
-#endif
+#endif 
 
-	/** The distance below which a sound begins to linearly interpolate towards being non-spatialized (2D). See "Non Spatialized Radius End" to define the end of the interpolation and the "Non Spatialized Radius Mode" for the mode of the interpolation. Note: this does not apply when using a 3rd party binaural plugin (audio will remain spatialized). */
-	UPROPERTY(EditAnywhere, BlueprintReadWrite, Category = AttenuationSpatialization, meta = (ClampMin = "0", EditCondition = "bSpatialize"))
+	/** The distance below which a sound begins to linearly interpolate towards being non-spatialized (2D). See "Non Spatialized Distance Min" to define the end of the interpolation and the "Non Spatialized Distance Mode" for the mode of the interpolation. */
+	UPROPERTY(EditAnywhere, BlueprintReadWrite, Category = AttenuationSpatialization, meta = (DisplayName = "Non Spatialized Distance Max", ClampMin = "0", EditCondition = "bSpatialize"))
 	float NonSpatializedRadiusStart;
 
-	/** The distance below which a sound is fully non-spatialized (2D). See "Non Spatialized Radius Start" to define the start of the interpolation and the "Non Spatialized Radius Mode" for the mode of the interpolation. */
-	UPROPERTY(EditAnywhere, BlueprintReadWrite, Category = AttenuationSpatialization, meta = (ClampMin = "0", EditCondition = "bSpatialize"))
+	/** The distance below which a sound is fully non-spatialized (2D). See "Non Spatialized Distance Max" to define the start of the interpolation and the "Non Spatialized Distance Mode" for the mode of the interpolation. */
+	UPROPERTY(EditAnywhere, BlueprintReadWrite, Category = AttenuationSpatialization, meta = (DisplayName = "Non Spatialized Distance Min", ClampMin = "0", EditCondition = "bSpatialize"))
 	float NonSpatializedRadiusEnd;
 
-	/** Defines how to interpolate a 3D sound towards a 2D sound when using the non-spatialized radius start and end properties. */
-	UPROPERTY(EditAnywhere, BlueprintReadWrite, Category = AttenuationSpatialization, meta = (ClampMin = "0", EditCondition = "bSpatialize"))
+	/** Defines how to interpolate a 3D sound towards a 2D sound when using the Non Spatialized Distance Min and Max properties. Note: this does not apply when using a 3rd party binaural plugin (audio will remain spatialized). */
+	UPROPERTY(EditAnywhere, BlueprintReadWrite, Category = AttenuationSpatialization, meta = (DisplayName = "Non Spatialized Distance Mode", ClampMin = "0", EditCondition = "bSpatialize"))
 	ENonSpatializedRadiusSpeakerMapMode NonSpatializedRadiusMode;
 
 	/** The world-space distance between left and right stereo channels when stereo assets are 3D spatialized. */
@@ -266,28 +266,28 @@ struct FSoundAttenuationSettings : public FBaseAttenuationSettings
 	float RadiusMax_DEPRECATED;
 #endif
 
-	/* The distance min range at which to apply an absorption LPF filter. */
-	UPROPERTY(EditAnywhere, BlueprintReadWrite, Category = AttenuationAirAbsorption, meta = (DisplayName = "Min Distance Range"))
+	/* The distance min range at which to apply our absorption filter(s). */
+	UPROPERTY(EditAnywhere, BlueprintReadWrite, Category = AttenuationAirAbsorption, meta = (DisplayName = "Distance Min"))
 	float LPFRadiusMin;
 
-	/* The max distance range at which to apply an absorption LPF filter. Absorption freq cutoff interpolates between filter frequency ranges between these distance values. */
-	UPROPERTY(EditAnywhere, BlueprintReadWrite, Category = AttenuationAirAbsorption, meta = (DisplayName = "Max Distance Range"))
+	/* The max distance range at which to apply our absorption filter(s). Absorption freq cutoff interpolates between filter frequency ranges within these distance values. */
+	UPROPERTY(EditAnywhere, BlueprintReadWrite, Category = AttenuationAirAbsorption, meta = (DisplayName = "Distance Max"))
 	float LPFRadiusMax;
 
-	/* The range of the cutoff frequency (in Hz) of the lowpass absorption filter. */
-	UPROPERTY(EditAnywhere, BlueprintReadWrite, Category = AttenuationAirAbsorption, meta = (DisplayName = "Low Pass Cutoff Frequency Min"))
+	/* The cutoff frequency (in Hz) of the lowpass absorption filter when at Distance Min. */
+	UPROPERTY(EditAnywhere, BlueprintReadWrite, Category = AttenuationAirAbsorption, meta = (Units = "Hz", DisplayName = "Low Pass Frequency Min"))
 	float LPFFrequencyAtMin;
 
-	/* The range of the cutoff frequency (in Hz) of the lowpass absorption filter. */
-	UPROPERTY(EditAnywhere, BlueprintReadWrite, Category = AttenuationAirAbsorption, meta = (DisplayName = "Low Pass Cutoff Frequency Max"))
+	/* The cutoff frequency (in Hz) of the lowpass absorption filter when at Distance Max. */
+	UPROPERTY(EditAnywhere, BlueprintReadWrite, Category = AttenuationAirAbsorption, meta = (Units = "Hz", DisplayName = "Low Pass Frequency Max"))
 	float LPFFrequencyAtMax;
-
-	/* The range of the cutoff frequency (in Hz) of the highpass absorption filter. */
-	UPROPERTY(EditAnywhere, BlueprintReadWrite, Category = AttenuationAirAbsorption, meta = (DisplayName = "High Pass Cutoff Frequency Min"))
+	
+	/* The cutoff frequency (in Hz) of the highpass absorption filter when at Distance Min. */
+	UPROPERTY(EditAnywhere, BlueprintReadWrite, Category = AttenuationAirAbsorption, meta = (Units = "Hz", DisplayName = "High Pass Frequency Min"))
 	float HPFFrequencyAtMin;
 
-	/* The range of the cutoff frequency (in Hz) of the highpass absorption filter. */
-	UPROPERTY(EditAnywhere, BlueprintReadWrite, Category = AttenuationAirAbsorption, meta = (DisplayName = "High Pass Cutoff Frequency Max"))
+	/* The cutoff frequency (in Hz) of the highpass absorption filter when at Distance Max. */
+	UPROPERTY(EditAnywhere, BlueprintReadWrite, Category = AttenuationAirAbsorption, meta = (Units = "Hz", DisplayName = "High Pass Frequency Max"))
 	float HPFFrequencyAtMax;
 
 	/** Azimuth angle (in degrees) relative to the listener forward vector which defines the focus region of sounds. Sounds playing at an angle less than this will be in focus. */
@@ -351,23 +351,23 @@ struct FSoundAttenuationSettings : public FBaseAttenuationSettings
 #endif
 
 	/** The amount to send to master reverb when sound is located at a distance equal to value specified in the reverb min send distance. */
-	UPROPERTY(EditAnywhere, BlueprintReadWrite, Category = AttenuationReverbSend, meta = (DisplayName = "Reverb Min Send Level"))
+	UPROPERTY(EditAnywhere, BlueprintReadWrite, Category = AttenuationReverbSend, meta = (DisplayName = "Reverb Send Level Min"))
 	float ReverbWetLevelMin;
 
 	/** The amount to send to master reverb when sound is located at a distance equal to value specified in the reverb max send distance. */
-	UPROPERTY(EditAnywhere, BlueprintReadWrite, Category = AttenuationReverbSend, meta = (DisplayName = "Reverb Max Send Level"))
+	UPROPERTY(EditAnywhere, BlueprintReadWrite, Category = AttenuationReverbSend, meta = (DisplayName = "Reverb Send Level Max"))
 	float ReverbWetLevelMax;
 
-	/** The min distance to send to the master reverb. */
-	UPROPERTY(EditAnywhere, BlueprintReadWrite, Category = AttenuationReverbSend, meta = (DisplayName = "Reverb Min Send Distance"))
+	/** The distance from listener at which we should be sending the level defined via Reverb Send Level Min */
+	UPROPERTY(EditAnywhere, BlueprintReadWrite, Category = AttenuationReverbSend, meta = (DisplayName = "Reverb Send Distance Min"))
 	float ReverbDistanceMin;
 
-	/** The max distance to send to the master reverb. */
-	UPROPERTY(EditAnywhere, BlueprintReadWrite, Category = AttenuationReverbSend, meta = (DisplayName = "Reverb Max Send Distance"))
+	/** The distance from listener at which we should be sending the level defined via Reverb Send Level Max */
+	UPROPERTY(EditAnywhere, BlueprintReadWrite, Category = AttenuationReverbSend, meta = (DisplayName = "Reverb Send Distance Max"))
 	float ReverbDistanceMax;
 
-	/* The manual master reverb send level to use. Doesn't change as a function of distance. */
-	UPROPERTY(EditAnywhere, BlueprintReadWrite, Category = AttenuationReverbSend)
+	/* Constant reverb send level to use. Doesn't change as a function of distance. */
+	UPROPERTY(EditAnywhere, BlueprintReadWrite, Category = AttenuationReverbSend, meta = (DisplayName = "Reverb Send Level"))
 	float ManualReverbSendLevel;
 
 	/** Interpolated value to scale priority against when the sound is at the minimum priority attenuation distance from the closest listener. */
@@ -386,7 +386,7 @@ struct FSoundAttenuationSettings : public FBaseAttenuationSettings
 	UPROPERTY(EditAnywhere, BlueprintReadWrite, Category = AttenuationPriority, meta = (ClampMin = "0.0", DisplayName = "Priority Attenuation Max Distance"))
 	float PriorityAttenuationDistanceMax;
 
-	/* Static priority scalar to use (doesn't change as a function of distance). */
+	/* Constant priority scalar to use (doesn't change as a function of distance). */
 	UPROPERTY(EditAnywhere, BlueprintReadWrite, Category = AttenuationPriority, meta = (ClampMin = "0.0", UIMin = "0.0", UIMax = "1.0", DisplayName = "Attenuation Priority"))
 	float ManualPriorityAttenuation;
 
diff --git a/Engine/Source/Runtime/Engine/Classes/Sound/SoundSourceBusSend.h b/Engine/Source/Runtime/Engine/Classes/Sound/SoundSourceBusSend.h
--- a/Engine/Source/Runtime/Engine/Classes/Sound/SoundSourceBusSend.h
+++ b/Engine/Source/Runtime/Engine/Classes/Sound/SoundSourceBusSend.h
@@ -12,14 +12,14 @@ class UAudioBus;
 UENUM(BlueprintType)
 enum class ESourceBusSendLevelControlMethod : uint8
 {
-	// A send based on linear interpolation between a distance range and send-level range
-	Linear,
+	// A send based on linear interpolation between a distance range and a remapped range
+	Linear  UMETA(DisplayName = "Distance Linear"),
 
-	// A send based on a supplied curve
-	CustomCurve,
+	// A send based on a supplied curve mapping (distance = X, send level = Y)
+	CustomCurve UMETA(DisplayName = "Distance Custom"),
 
-	// A manual send level (Uses the specified constant send level value. Useful for 2D sounds.)
-	Manual,
+	// A constant send level (Uses the specified constant send level value. Useful for 2D sounds.)
+	Manual UMETA(DisplayName = "Constant"),
 };
 
 USTRUCT(BlueprintType)
@@ -27,11 +27,7 @@ struct FSoundSourceBusSendInfo
 {
 	GENERATED_USTRUCT_BODY()
 
-	/*
-		Manual: Use Send Level only
-		Linear: Interpolate between Min and Max Send Levels based on listener distance (between Min/Max Send Distance)
-		Custom Curve: Use the float curve to map Send Level to distance (0.0-1.0 on curve maps to Min/Max Send Distance)
-	*/
+	// Changes the method we use to determine the send level - see enum value tooltips for more information
 	UPROPERTY(EditAnywhere, Category = BusSend, meta = (DisplayName = "Send Level Control Method"))
 	ESourceBusSendLevelControlMethod SourceBusSendLevelControlMethod;
 
@@ -43,28 +39,28 @@ struct FSoundSourceBusSendInfo
 	UPROPERTY(EditAnywhere, Category = BusSend)
 	TObjectPtr<UAudioBus> AudioBus;
 
-	// Manually set the amount of audio to send to the bus
-	UPROPERTY(EditAnywhere, Category = BusSend, meta = (DisplayName = "Manual Send Level", EditCondition = "SourceBusSendLevelControlMethod == ESourceBusSendLevelControlMethod::Manual", EditConditionHides))
+	// Constant send level to use. Doesn't change as a function of distance.
+	UPROPERTY(EditAnywhere, Category = BusSend, meta = (DisplayName = "Send Level", EditCondition = "SourceBusSendLevelControlMethod == ESourceBusSendLevelControlMethod::Manual", EditConditionHides))
 	float SendLevel;
 
-	// The amount to send to the bus when sound is located at a distance less than or equal to value specified in the Min Send Distance
-	UPROPERTY(EditAnywhere, BlueprintReadWrite, Category = BusSend, meta = (EditCondition = "SourceBusSendLevelControlMethod == ESourceBusSendLevelControlMethod::Linear", EditConditionHides))
+	// The amount to send to the bus when sound is located at a distance less than or equal to value specified in the Send Distance Min
+	UPROPERTY(EditAnywhere, BlueprintReadWrite, Category = BusSend, meta = (DisplayName = "Send Level Min", EditCondition = "SourceBusSendLevelControlMethod == ESourceBusSendLevelControlMethod::Linear", EditConditionHides))
 	float MinSendLevel;
 
-	// The amount to send to the bus when sound is located at a distance greater than or equal to value specified in the Max Send Distance
-	UPROPERTY(EditAnywhere, BlueprintReadWrite, Category = BusSend, meta = (EditCondition = "SourceBusSendLevelControlMethod == ESourceBusSendLevelControlMethod::Linear", EditConditionHides))
+	// The amount to send to the bus when sound is located at a distance greater than or equal to value specified in the Send Distance Max
+	UPROPERTY(EditAnywhere, BlueprintReadWrite, Category = BusSend, meta = (DisplayName = "Send Level Max", EditCondition = "SourceBusSendLevelControlMethod == ESourceBusSendLevelControlMethod::Linear", EditConditionHides))
 	float MaxSendLevel;
 
-	/** The distance at which to start mapping between to Min/Max Send Level
-	 *  Distances LESS than this will result in a clamped Min Send Level
+	/** The distance at which to start mapping between to Send Level Min/Max 
+	 *  Distances LESS than this will result in a clamped Send Level Min
 	 */ 
-	UPROPERTY(EditAnywhere, BlueprintReadWrite, Category = BusSend, meta = (EditCondition = "SourceBusSendLevelControlMethod != ESourceBusSendLevelControlMethod::Manual", EditConditionHides))
+	UPROPERTY(EditAnywhere, BlueprintReadWrite, Category = BusSend, meta = (DisplayName = "Send Distance Min", EditCondition = "SourceBusSendLevelControlMethod != ESourceBusSendLevelControlMethod::Manual", EditConditionHides))
 	float MinSendDistance;
 
-	/** The distance at which to stop mapping between Min/Max Send Level
-	 *  Distances GREATER than this will result in a clamped Max Send Level
+	/** The distance at which to stop mapping between Send Level Min/Max 
+	 *  Distances GREATER than this will result in a clamped Send Level Max
 	 */
-	UPROPERTY(EditAnywhere, BlueprintReadWrite, Category = BusSend, meta = (EditCondition = "SourceBusSendLevelControlMethod != ESourceBusSendLevelControlMethod::Manual", EditConditionHides))
+	UPROPERTY(EditAnywhere, BlueprintReadWrite, Category = BusSend, meta = (DisplayName = "Send Distance Max", EditCondition = "SourceBusSendLevelControlMethod != ESourceBusSendLevelControlMethod::Manual", EditConditionHides))
 	float MaxSendDistance;
 
 	// The custom send curve to use for distance-based send level. (0.0-1.0 on the curve's X-axis maps to Min/Max Send Distance)
diff --git a/Engine/Source/Runtime/Engine/Classes/Sound/SoundSubmixSend.h b/Engine/Source/Runtime/Engine/Classes/Sound/SoundSubmixSend.h
--- a/Engine/Source/Runtime/Engine/Classes/Sound/SoundSubmixSend.h
+++ b/Engine/Source/Runtime/Engine/Classes/Sound/SoundSubmixSend.h
@@ -68,14 +68,14 @@ enum class EAudioRecordingExportType : uint8
 UENUM(BlueprintType)
 enum class ESendLevelControlMethod : uint8
 {
-	// A send based on linear interpolation between a distance range and send-level range
-	Linear,
+	// A send based on linear interpolation between the distance range and the remapped range
+	Linear UMETA(DisplayName = "Distance Linear"),
 
-	// A send based on a supplied curve
-	CustomCurve,
+	// A send based on distance using the supplied curve mapping (distance = X, send level = Y)
+	CustomCurve UMETA(DisplayName = "Distance Custom"),
 
-	// A manual send level (Uses the specified constant send level value. Useful for 2D sounds.)
-	Manual,
+	// A constant send level (Uses the specified constant send level value. Useful for 2D sounds.)
+	Manual UMETA(DisplayName = "Constant"),
 };
 
 // Common set of settings that are uses as submix sends.
@@ -86,11 +86,7 @@ struct FSoundSubmixSendInfoBase
 
 	ENGINE_API FSoundSubmixSendInfoBase();
 
-	/*
-		Manual: Use Send Level only
-		Linear: Interpolate between Min and Max Send Levels based on listener distance (between Min/Max Send Distance)
-		Custom Curve: Use the float curve to map Send Level to distance (0.0-1.0 on curve maps to Min/Max Send Distance)
-	*/
+	// // Changes the method we use to determine the send level - see enum value tooltips for more information
 	UPROPERTY(EditAnywhere, BlueprintReadWrite, Category = SubmixSend)
 	ESendLevelControlMethod SendLevelControlMethod;
 	
@@ -99,31 +95,31 @@ struct FSoundSubmixSendInfoBase
 	TObjectPtr<USoundSubmixBase> SoundSubmix;
 
 	// Manually set the amount of audio to send
-	UPROPERTY(EditAnywhere, BlueprintReadWrite, Category = SubmixSend, meta = (DisplayName = "Manual Send Level", EditCondition = "SendLevelControlMethod == ESendLevelControlMethod::Manual", EditConditionHides))
+	UPROPERTY(EditAnywhere, BlueprintReadWrite, Category = SubmixSend, meta = (DisplayName = "Send Level", EditCondition = "SendLevelControlMethod == ESendLevelControlMethod::Manual", EditConditionHides))
 	float SendLevel;
 
-	// Whether to disable the internal 0-1 clamp for Manual Send Level control
-	UPROPERTY(EditAnywhere, BlueprintReadWrite, Category = SubmixSend, meta = (EditCondition = "SendLevelControlMethod == ESendLevelControlMethod::Manual", EditConditionHides))
+	// Whether to disable the internal 0-1 clamp for Constant Send Level control
+	UPROPERTY(EditAnywhere, BlueprintReadWrite, Category = SubmixSend, meta = (DisplayName = "Disable Send Level Clamp", EditCondition = "SendLevelControlMethod == ESendLevelControlMethod::Manual", EditConditionHides))
 	bool DisableManualSendClamp;
 
-	// The amount to send to the Submix when sound is located at a distance less than or equal to value specified in the Min Send Distance
-	UPROPERTY(EditAnywhere, BlueprintReadWrite, Category = SubmixSend, meta = (EditCondition = "SendLevelControlMethod == ESendLevelControlMethod::Linear", EditConditionHides))
+	// The amount to send to the Submix when sound is located at a distance less than or equal to value specified in the Send Distance Min
+	UPROPERTY(EditAnywhere, BlueprintReadWrite, Category = SubmixSend, meta = (DisplayName = "Send Level Min", EditCondition = "SendLevelControlMethod == ESendLevelControlMethod::Linear", EditConditionHides))
 	float MinSendLevel;
 
-	// The amount to send to the Submix when sound is located at a distance greater than or equal to value specified in the Max Send Distance
-	UPROPERTY(EditAnywhere, BlueprintReadWrite, Category = SubmixSend, meta = (EditCondition = "SendLevelControlMethod == ESendLevelControlMethod::Linear", EditConditionHides))
+	// The amount to send to the Submix when sound is located at a distance greater than or equal to value specified in the Send Distance Max
+	UPROPERTY(EditAnywhere, BlueprintReadWrite, Category = SubmixSend, meta = (DisplayName = "Send Level Max", EditCondition = "SendLevelControlMethod == ESendLevelControlMethod::Linear", EditConditionHides))
 	float MaxSendLevel;
 
-	/** The distance at which to start mapping between to Min/Max Send Level
-	 *  Distances LESS than this will result in a clamped Min Send Level
+	/** The distance at which to start mapping between to Send Level Min/Max 
+	 *  Distances LESS than this will result in a clamped Send Level Min
 	 */ 
-	UPROPERTY(EditAnywhere, BlueprintReadWrite, Category = SubmixSend, meta = (EditCondition = "SendLevelControlMethod != ESendLevelControlMethod::Manual", EditConditionHides))
+	UPROPERTY(EditAnywhere, BlueprintReadWrite, Category = SubmixSend, meta = (DisplayName = "Send Distance Min", EditCondition = "SendLevelControlMethod != ESendLevelControlMethod::Manual", EditConditionHides))
 	float MinSendDistance;
 
-	/** The distance at which to stop mapping between Min/Max Send Level
-	 *  Distances GREATER than this will result in a clamped Max Send Level
+	/** The distance at which to stop mapping between Send Level Min/Max
+	 *  Distances GREATER than this will result in a clamped Send Level Max
 	 */
-	UPROPERTY(EditAnywhere, BlueprintReadWrite, Category = SubmixSend, meta = (EditCondition = "SendLevelControlMethod != ESendLevelControlMethod::Manual", EditConditionHides))
+	UPROPERTY(EditAnywhere, BlueprintReadWrite, Category = SubmixSend, meta = (DisplayName = "Send Distance Max", EditCondition = "SendLevelControlMethod != ESendLevelControlMethod::Manual", EditConditionHides))
 	float MaxSendDistance;
 
 	// The custom send curve to use for distance-based send level. (0.0-1.0 on the curve's X-axis maps to Min/Max Send Distance)
