Clarify comment on how parallel animation evaluation is performed

#Synced-CL 51156570

diff --git a/Engine/Source/Runtime/Engine/Private/Components/SkeletalMeshComponent.cpp b/Engine/Source/Runtime/Engine/Private/Components/SkeletalMeshComponent.cpp
--- a/Engine/Source/Runtime/Engine/Private/Components/SkeletalMeshComponent.cpp
+++ b/Engine/Source/Runtime/Engine/Private/Components/SkeletalMeshComponent.cpp
@@ -4959,8 +4959,11 @@ bool USkeletalMeshComponent::HandleExistingParallelEvaluationTask(bool bBlockOnT
 
 			// We had a valid task handle and now we are done waiting for it
 			// Perform completion now
-			// Note that because completion happens in a separate task from the one we waited on
-			// it could execute concurrently with this call to finalize evaluation
+			// Note that because completion happens in a separate task (on the main thread) from the one we waited on
+			// it could execute after this call to finalize evaluation. If it does, it should be a no-op operation as
+			// it means the skeletal mesh or some other core piece has changed and the completion validates that the
+			// data we started with is consistent with the current data.
+			// We don't wait for the main thread completion task because it could execute much later, instead we do it inline
 			CompleteParallelAnimationEvaluation(bPerformPostAnimEvaluation);
 		}
 
