New Material Translator:
- Some code reorganization and improvements.

#rnx

#Synced-CL 51020780

diff --git a/Engine/Source/Runtime/Engine/Private/Materials/MaterialExpressionsIR.cpp b/Engine/Source/Runtime/Engine/Private/Materials/MaterialExpressionsIR.cpp
--- a/Engine/Source/Runtime/Engine/Private/Materials/MaterialExpressionsIR.cpp
+++ b/Engine/Source/Runtime/Engine/Private/Materials/MaterialExpressionsIR.cpp
@@ -1185,7 +1185,6 @@ static FValueRef EmitAlmostEquals(MIR::FEmitter& Em, FValueRef A, FValueRef B, f
 
 void UMaterialExpressionIf::Build(MIR::FEmitter& Em)
 {
-	// Get input values
 	FValueRef AValue = Em.InputDefaultFloat(&A, 0.f);
 	FValueRef BValue = Em.InputDefaultFloat(&B, ConstB);
 	FValueRef AGreaterThanBValue = Em.InputDefaultFloat(&AGreaterThanB, 0.f);
@@ -4582,7 +4581,10 @@ void UMaterialExpressionDynamicParameter::Build(MIR::FEmitter& Em)
 
 		void Analyze(MIR::FExternAnalysisContext& Context)
 		{
-			Context.Module->GetStatistics().DynamicParticleParameterMask |= 1 << ParameterIndex;
+			const uint32 DynamicParametersMask = 1 << ParameterIndex;
+			Context.Module->GetStatistics().DynamicParticleParameterMask |= DynamicParametersMask;
+			Context.Module->AddIntegerEnvironmentDefine(TEXT("USE_DYNAMIC_PARAMETERS"), 1);
+			Context.Module->AddIntegerEnvironmentDefine(TEXT("DYNAMIC_PARAMETERS_MASK"), DynamicParametersMask);
 		}
 
 		void ToHLSL(MIR::FExternPrinterHLSL& Printer) const
@@ -4925,7 +4927,10 @@ void UMaterialExpressionClearCoatNormalCustomOutput::Build(MIR::FEmitter& Em)
 void UMaterialExpressionFirstPersonOutput::Build(MIR::FEmitter& Em)
 {
 	FValueRef InputValue = Em.CastToFloat(Em.InputDefaultFloat(&FirstPersonInterpolationAlpha, ConstFirstPersonInterpolationAlpha), 1);
-	Em.SetCustomOutputs(TEXTVIEW("FirstPersonOutput"), { &InputValue, 1 }, MIR::EMaterialOutputFrequency::PerVertex);
+	Em.SetCustomOutputs(TEXTVIEW("FirstPersonOutput"), { &InputValue, 1 }, MIR::EMaterialOutputFrequency::PerVertex, [] (FMaterialIRModule& Module)
+	{
+		Module.GetCompilationOutput().bModifiesMeshPosition |= true;
+	});
 }
 
 void UMaterialExpressionMotionVectorWorldOffsetOutput::Build(MIR::FEmitter& Em)
diff --git a/Engine/Source/Runtime/Engine/Private/Materials/MaterialIREmitter.cpp b/Engine/Source/Runtime/Engine/Private/Materials/MaterialIREmitter.cpp
--- a/Engine/Source/Runtime/Engine/Private/Materials/MaterialIREmitter.cpp
+++ b/Engine/Source/Runtime/Engine/Private/Materials/MaterialIREmitter.cpp
@@ -1402,7 +1402,7 @@ static FSetMaterialOutput* CreateSetMaterialOutput(FEmitter& Em, FMaterialIRModu
 static bool IsVertexStageMaterialProperty(EMaterialProperty Property)
 {
 	return Property == MP_WorldPositionOffset
-		|| (Property >= MP_CustomizedUVs0 && Property <= MP_CustomizedUVs7);
+		|| (Property >= MP_CustomizedUVs0 && Property <= MP_LastCustomizedUVs);
 }
 
 void FEmitter::Private_SetMaterialOutput(EMaterialProperty Property, FValueRef Arg)
@@ -1456,6 +1456,12 @@ void FEmitter::SetCustomOutputs(FStringView Name, TConstArrayView<FValueRef> Arg
 	}
 }
 
+void FEmitter::SetCustomOutputs(FStringView Name, TConstArrayView<FValueRef> Args, EMaterialOutputFrequency Frequency, TFunctionRef<void(FMaterialIRModule&)> Func)
+{
+	SetCustomOutputs(Name, Args, Frequency);
+	Func(*Module);
+}
+
 static void InvalidVectorConstructionError(FEmitter& Em, int32 DesiredVectorDimensions, FValueRef A, FValueRef B, FValueRef C = {}, FValueRef D = {})
 {
 	if (C.IsValid() && D.IsValid())
@@ -4586,9 +4592,37 @@ bool FEmitter::FValueKeyFuncs::Matches(KeyInitType A, KeyInitType B)
 	return A->Equals(B);
 }
 
+static uint32 HashBytes(const void* InPtr, uint32 Size)
+{
+	check((uintptr_t)InPtr % 4 == 0); // Require the pointer to be 4-byte aligned
+
+	uint32 Hash = 0;
+
+	const uint32* WordPtr = (const uint32*)InPtr;
+	uint32 NumWords = Size / 4;
+	for (uint32 i = 0; i < NumWords; ++i)
+	{
+		Hash = HashCombineFast(Hash, WordPtr[i]);
+	}
+
+	const uint8* EndPtr = (const uint8*)(WordPtr + NumWords);
+
+	uint32 End = 0;
+	switch (Size % 4)
+	{
+		case 3: End |= uint32(EndPtr[2]) << 16; // fallthrough
+		case 2: End |= uint32(EndPtr[1]) << 8;  // fallthrough
+		case 1: End |= uint32(EndPtr[0]);       // fallthrough
+			Hash = HashCombineFast(Hash, End);
+		default: break;
+	}
+
+	return Hash;
+}
+
 uint32 FEmitter::FValueKeyFuncs::GetKeyHash(KeyInitType Key)
 {
-	return Internal::HashBytes(Key, Key->GetSizeInBytes());
+	return MIR::HashBytes(Key, Key->GetSizeInBytes());
 }
 
 } // namespace MIR
diff --git a/Engine/Source/Runtime/Engine/Private/Materials/MaterialIRInternal.cpp b/Engine/Source/Runtime/Engine/Private/Materials/MaterialIRInternal.cpp
--- a/Engine/Source/Runtime/Engine/Private/Materials/MaterialIRInternal.cpp
+++ b/Engine/Source/Runtime/Engine/Private/Materials/MaterialIRInternal.cpp
@@ -42,34 +42,6 @@ EMaterialTextureParameterType TextureMaterialValueTypeToParameterType(EMaterialV
 	}
 }
 
-uint32 HashBytes(const void* InPtr, uint32 Size)
-{
-	check((uintptr_t)InPtr % 4 == 0); // Require the pointer to be 4-byte aligned
-
-	uint32 Hash = 0;
-
-	const uint32* WordPtr = (const uint32*)InPtr;
-	uint32 NumWords = Size / 4;
-	for (uint32 i = 0; i < NumWords; ++i)
-	{
-		Hash = HashCombineFast(Hash, WordPtr[i]);
-	}
-
-	const uint8* EndPtr = (const uint8*)(WordPtr + NumWords);
-
-	uint32 End = 0;
-	switch (Size % 4)
-	{
-		case 3: End |= uint32(EndPtr[2]) << 16; // fallthrough
-		case 2: End |= uint32(EndPtr[1]) << 8;  // fallthrough
-		case 1: End |= uint32(EndPtr[0]);       // fallthrough
-			Hash = HashCombineFast(Hash, End);
-		default: break;
-	}
-
-	return Hash;
-}
-
 } // namespace MIR::Internal
 
 #endif
diff --git a/Engine/Source/Runtime/Engine/Private/Materials/MaterialIRInternal.h b/Engine/Source/Runtime/Engine/Private/Materials/MaterialIRInternal.h
--- a/Engine/Source/Runtime/Engine/Private/Materials/MaterialIRInternal.h
+++ b/Engine/Source/Runtime/Engine/Private/Materials/MaterialIRInternal.h
@@ -11,9 +11,11 @@ namespace MIR::Internal {
 // Returns the material value tpye of the specified UTexture or URuntimeVirtualTexture.
 EMaterialValueType GetTextureMaterialValueType(const UObject* TextureObject);
 
-//
+// Maps texture EMaterialValueType values to EMaterialTextureParameterType.
 EMaterialTextureParameterType TextureMaterialValueTypeToParameterType(EMaterialValueType Type);
 
+// @massimo.tristano expose an internal "Context" struct between the builder and the emitter, and move these methods
+// to the context itself
 // Returns the value flowing into given expression input (previously set through `BindValueToExpressionInput`).
 FValue* FetchValueFromExpressionInput(FMaterialIRModuleBuilderImpl* Builder, const FExpressionInput* Input);
 
@@ -23,9 +25,6 @@ void BindValueToExpressionInput(FMaterialIRModuleBuilderImpl* Builder, const FEx
 // Flows a value int given expression output.
 void BindValueToExpressionOutput(FMaterialIRModuleBuilderImpl* Builder, const FExpressionOutput* Output, FValue* Value);
 
-// Computes the hash of a blob of memory. Note: Ptr must be 4 bytes aligned.
-uint32 HashBytes(const void* Ptr, uint32 Size);
-
 } // namespace MIR::Internal
 
 #endif // #if WITH_EDITOR
diff --git a/Engine/Source/Runtime/Engine/Private/Materials/MaterialIRModuleBuilder.cpp b/Engine/Source/Runtime/Engine/Private/Materials/MaterialIRModuleBuilder.cpp
--- a/Engine/Source/Runtime/Engine/Private/Materials/MaterialIRModuleBuilder.cpp
+++ b/Engine/Source/Runtime/Engine/Private/Materials/MaterialIRModuleBuilder.cpp
@@ -46,6 +46,12 @@ static TAutoConsoleVariable<int32> CVarMaterialIRDebugDumpLevel(
 	TEXT("2: Everything above, plus dump the 'Uses' graph in Graphviz Dot syntax (to '{SavedDir}/Materials/IRDumpUseGraph.dot').\n"),
 	ECVF_RenderThreadSafe);
 
+// Returns whether given property is a customized UV.
+static inline bool IsMaterialPropertyCustomizedUV(EMaterialProperty Property)
+{
+	return Property >= MP_CustomizedUVs0 && Property <= MP_LastCustomizedUVs;
+}
+
 struct FAnalysisContext
 {
 	UMaterialExpressionMaterialFunctionCall* Call{};
@@ -167,7 +173,7 @@ struct FMaterialIRModuleBuilderImpl
 	// property is active in the translating material.
 	bool PropertyIsActiveInMaterial(const EMaterialProperty Property)
 	{
-		if (MIR::IsMaterialPropertyCustomizedUV(Property) && Property >= MP_CustomizedUVs0 + Emitter.BaseMaterial->NumCustomizedUVs)
+		if (IsMaterialPropertyCustomizedUV(Property) && Property >= MP_CustomizedUVs0 + Emitter.BaseMaterial->NumCustomizedUVs)
 		{
 			return false;
 		}
@@ -756,7 +762,7 @@ struct FMaterialIRModuleBuilderImpl
 					PropertyValue = Emitter.Cast(PropertyValue, OutputArgType);
 				}
 			}
-			else if (MIR::IsMaterialPropertyCustomizedUV(Property))
+			else if (IsMaterialPropertyCustomizedUV(Property))
 			{
 				// To keep backwards compatibility with the old translator, a CustomUV attribute pin with no given value
 				// does not default to (0,0) (as it should) but to the uncustomized texture coordinate. If we find a customized UV
@@ -1116,7 +1122,7 @@ struct FMaterialIRModuleBuilderImpl
 			{
 				const MIR::FSetMaterialOutput* SetMaterialOutput = MIR::As<MIR::FSetMaterialOutput>(Instr);
 				return SetMaterialOutput
-					&& MIR::IsMaterialPropertyCustomizedUV(SetMaterialOutput->Property)
+					&& IsMaterialPropertyCustomizedUV(SetMaterialOutput->Property)
 					&& (Module->GetStatistics().InterpolatedTexCoordsMask & (1 << (SetMaterialOutput->Property - MP_CustomizedUVs0))) == 0;
 			});
 		}
@@ -1223,21 +1229,8 @@ struct FMaterialIRModuleBuilderImpl
 		CompilationOutput.bUsesAnisotropy 			= MaterialPropertyHasNonZeroValue(MP_Anisotropy);
 		
 		// Check if displacement is actually used (disable if it was left at the default, invalid constant value of -1)
-		CompilationOutput.bUsesDisplacement 		= DoesPlatformSupportNanite(Module->GetShaderPlatform()) && Builder->MaterialInterface->IsTessellationEnabled() && MaterialPropertyHasNonZeroValue(MP_Displacement);
-
-		// bModifiesMeshPosition
-		// @massimo.tristano this kind of logic should be moved to the first person custom output expression Build() function
-		bool bHasFirstPersonOutput = false;
-		for (const MIR::FSetMaterialOutput* Output : Module->GetCustomOutputs())
-		{
-			if (Output->Name == TEXTVIEW("FirstPersonOutput"))
-			{
-				bHasFirstPersonOutput = true;
-				break;
-			}
-		}
-		
-		CompilationOutput.bModifiesMeshPosition = CompilationOutput.bUsesPixelDepthOffset || CompilationOutput.bUsesWorldPositionOffset || CompilationOutput.bUsesDisplacement || bHasFirstPersonOutput;
+		CompilationOutput.bUsesDisplacement 	 |= DoesPlatformSupportNanite(Module->GetShaderPlatform()) && Builder->MaterialInterface->IsTessellationEnabled() && MaterialPropertyHasNonZeroValue(MP_Displacement);
+		CompilationOutput.bModifiesMeshPosition  |= CompilationOutput.bUsesPixelDepthOffset || CompilationOutput.bUsesWorldPositionOffset || CompilationOutput.bUsesDisplacement;
 		
 		// Run a final material compilation output validation and report any error reported.
 		TArray<FString> ValidationErrors;
diff --git a/Engine/Source/Runtime/Engine/Private/Materials/MaterialIRToHLSLTranslator.cpp b/Engine/Source/Runtime/Engine/Private/Materials/MaterialIRToHLSLTranslator.cpp
--- a/Engine/Source/Runtime/Engine/Private/Materials/MaterialIRToHLSLTranslator.cpp
+++ b/Engine/Source/Runtime/Engine/Private/Materials/MaterialIRToHLSLTranslator.cpp
@@ -2293,7 +2293,7 @@ static void GenerateOtherTemplateStringParameters(const FMaterial* Material, con
 		FPrinter Printer { MaterialAggregatesRequiringDeclaration, 1 };
 		for (int32 i = 0; i < ShaderTexCoordsArrayLength; i++)
 		{
-			const bool bIsCustomizedUV = i < MIR::MaxNumCustomizedUVs && Module->IsMaterialPropertyUsed((EMaterialProperty)(MP_CustomizedUVs0 + i));
+			const bool bIsCustomizedUV = i < MP_NumCustomizedUVs && Module->IsMaterialPropertyUsed((EMaterialProperty)(MP_CustomizedUVs0 + i));
 			Printer << TEXTVIEW("OutTexCoords[") << i << TEXTVIEW("] = Parameters.");
 			Printer << (bIsCustomizedUV ? TEXTVIEW("CustomizedUVs[") : TEXTVIEW("TexCoords["));
 			Printer << i << TEXT(']') << EndOfStatement;
@@ -2460,14 +2460,6 @@ static void GetShaderCompilerEnvironment(const FMaterial* Material, const FMater
 	OutEnvironment.SetDefine(TEXT("IS_MATERIAL_SHADER"), true);
 	OutEnvironment.SetDefine(TEXT("VIRTUAL_TEXTURE_OUTPUT"), CompilationOutput.bHasRuntimeVirtualTextureOutputNode != 0);
 
-	// @massimo.tristano refactor this
-	uint32 DynamicParametersMask = Module->GetStatistics().DynamicParticleParameterMask;
-	if (DynamicParametersMask)
-	{
-		OutEnvironment.SetDefine(TEXT("USE_DYNAMIC_PARAMETERS"), 1);
-		OutEnvironment.SetDefine(TEXT("DYNAMIC_PARAMETERS_MASK"), DynamicParametersMask);
-	}
-
 	// Set all defines that are defined by the module.
 	// Any conditional exemption via material properties, such as 'Material->IsUsedWithInstancedStaticMeshes()', are handled during the material IR analysis.
 	// @massimo.tristano refactor this
diff --git a/Engine/Source/Runtime/Engine/Public/Materials/MaterialIRCommon.h b/Engine/Source/Runtime/Engine/Public/Materials/MaterialIRCommon.h
--- a/Engine/Source/Runtime/Engine/Public/Materials/MaterialIRCommon.h
+++ b/Engine/Source/Runtime/Engine/Public/Materials/MaterialIRCommon.h
@@ -9,7 +9,9 @@
 
 #if WITH_EDITOR
 
-/* Forward declarations */
+//------------------------------------------------------------------------
+//                          Forward declarations
+//------------------------------------------------------------------------
 
 enum EMaterialProperty : int;
 enum EMaterialDomain : int;
@@ -49,30 +51,24 @@ namespace ERHIFeatureLevel { enum Type : int; }
 
 namespace MIR
 {
+	class FEmitter;
+	struct FBlock;
+	struct FExtern;
+	struct FExternImpl;
+	struct FExternPrinterHLSL;
+	struct FInstruction;
+	struct FPrimitive;
+	struct FSetMaterialOutput;
+	struct FType;
+	struct FValue;
+}
 
-class FEmitter;
-struct FBlock;
-struct FExtern;
-struct FExternImpl;
-struct FExternPrinterHLSL;
-struct FInstruction;
-struct FPrimitive;
-struct FSetMaterialOutput;
-struct FType;
-struct FValue;
-
-// A simple string view without constructors, to be used inside FValues.
-struct FSimpleStringView
+namespace MIR
 {
-	const TCHAR* Ptr; // The string characters
-	int32        Len; // The string length (excluding null byte)
 
-	FStringView ToView() const { return { Ptr, Len }; }
-	FString ToString() const { return FString{ ToView() }; }
-	operator FStringView() const { return ToView(); }
-	FSimpleStringView& operator=(FStringView View) { Ptr = View.GetData(); Len = View.Len(); return *this; }
-	const TCHAR* operator*() const { return Ptr; }
-};
+//------------------------------------------------------------------------
+//                          Utilities
+//------------------------------------------------------------------------
 
 // Identifies the block of execution in which an instruction runs.
 // Note: If you introduce a new stage, make sure you update EValueFlags accordingly.
@@ -84,7 +80,7 @@ enum EStage
 	NumStages
 };
 
-//
+// Bitmask of execution stages, aligned with EStage indices.
 enum class EStageMask : uint8
 {
 	None    = 0,
@@ -97,7 +93,7 @@ enum class EStageMask : uint8
 
 ENUM_CLASS_FLAGS(EStageMask);
 
-//
+// Returns whether given Stage is in given Mask.
 static inline bool IsStageInMask(EStage Stage, EStageMask Mask)
 {
 	return (1 << (uint32)Stage & (uint32)Mask) != 0;
@@ -147,7 +143,22 @@ enum class EDerivativeAxis
 	X, Y,
 };
 
-/* Utilities */
+//------------------------------------------------------------------------
+//                          Utilities
+//------------------------------------------------------------------------
+
+// A simple string view without constructors, to be used inside FValues.
+struct FSimpleStringView
+{
+	const TCHAR* Ptr; // The string characters
+	int32        Len; // The string length (excluding null byte)
+
+	FStringView ToView() const { return { Ptr, Len }; }
+	FString ToString() const { return FString{ ToView() }; }
+	operator FStringView() const { return ToView(); }
+	FSimpleStringView& operator=(FStringView View) { Ptr = View.GetData(); Len = View.Len(); return *this; }
+	const TCHAR* operator*() const { return Ptr; }
+};
 
 // Use this to efficiently allocate a temporary array using MemStack, instead
 // of using TArray and going through the global allocator. Only declare local variable
@@ -183,6 +194,7 @@ struct TTemporaryArray : public TArrayView<T>
 	}
 };
 
+// Zeroes an array of elements.
 template <typename T>
 void ZeroArray(TArrayView<T> Array)
 {
@@ -190,6 +202,8 @@ void ZeroArray(TArrayView<T> Array)
 	FMemory::Memzero(Array.GetData(), Array.Num() * sizeof(T));
 }
 
+// Tries to find a value with given Key into given TMap intance. If it's found, its value is assigned to OutValue.
+// Returns whether the value was found and OutValue was set.
 template <typename TKey, typename TValue>
 bool Find(const TMap<TKey, TValue>& Map, const TKey& Key, TValue& OutValue)
 {
@@ -200,20 +214,6 @@ bool Find(const TMap<TKey, TValue>& Map, const TKey& Key, TValue& OutValue)
 	return false;
 }
 
-// Returns whether given property is a customized UV.
-static inline bool IsMaterialPropertyCustomizedUV(EMaterialProperty Property)
-{
-	return Property >= MP_CustomizedUVs0 && Property <= MP_CustomizedUVs7;
-}
-
-static constexpr int MaxNumCustomizedUVs = 8;
-
-// @massimo.tristano this needs to be moved off to a better place
-namespace Internal
-{
-	uint32 HashBytes(const char* Ptr, uint32 Size);
-}
-
 } // namespace MIR
 
 #define UE_MIR_UNREACHABLE() { check(!"Unreachable"); UE_ASSUME(false); }
diff --git a/Engine/Source/Runtime/Engine/Public/Materials/MaterialIREmitter.h b/Engine/Source/Runtime/Engine/Public/Materials/MaterialIREmitter.h
--- a/Engine/Source/Runtime/Engine/Public/Materials/MaterialIREmitter.h
+++ b/Engine/Source/Runtime/Engine/Public/Materials/MaterialIREmitter.h
@@ -542,6 +542,7 @@ public:
 	// The type of each custom output is inferred by the type of the argument value it is set to, therefore make sure
 	// you cast each output argument to the desired type before calling this function.
 	void SetCustomOutputs(FStringView Name, TConstArrayView<FValueRef> Args, EMaterialOutputFrequency Frequency);
+	void SetCustomOutputs(FStringView Name, TConstArrayView<FValueRef> Args, EMaterialOutputFrequency Frequency, TFunctionRef<void(FMaterialIRModule&)> Func);
 	
 	// Emits a 2D column vector from X and Y scalars.
 	FValueRef Vector2(FValueRef X, FValueRef Y);
diff --git a/Engine/Source/Runtime/Engine/Public/SceneTypes.h b/Engine/Source/Runtime/Engine/Public/SceneTypes.h
--- a/Engine/Source/Runtime/Engine/Public/SceneTypes.h
+++ b/Engine/Source/Runtime/Engine/Public/SceneTypes.h
@@ -184,6 +184,10 @@ enum EMaterialProperty : int
 	MP_MaterialAttributes,
 	MP_CustomOutput UMETA(Hidden),
 	MP_MAX UMETA(DisplayName = "None"),
+
+	// Aliases
+	MP_LastCustomizedUVs = MP_CustomizedUVs7,
+	MP_NumCustomizedUVs = MP_LastCustomizedUVs - MP_CustomizedUVs0 + 1,
 };
 
 /** Blend modes supported for simple element rendering */
