[Viewport ITF]
- Allow raw mouse delta events to affect the unbounded mouse position when the mouse is on the edge of the viewport.
- Fix an issue where unbounded mouse movements did not work when the viewport was in fullscreen.
- FIx an issue with scaling compensation within windows fullscreen border compensation that caused a the fullscreen editor viewport to report as being larger than the monitor when windows scaling was above 100%.

#rnx

#Synced-CL 51158691

#Synced-CL 51158691

diff --git a/Engine/Source/Editor/LevelEditor/Private/LevelViewportLayout.cpp b/Engine/Source/Editor/LevelEditor/Private/LevelViewportLayout.cpp
--- a/Engine/Source/Editor/LevelEditor/Private/LevelViewportLayout.cpp
+++ b/Engine/Source/Editor/LevelEditor/Private/LevelViewportLayout.cpp
@@ -502,16 +502,15 @@ FVector2D FLevelViewportLayout::GetMaximizedViewportSizeOnCanvas() const
 			{
 				FVector2D ClippedArea = FVector2D::ZeroVector;
 
-				const float AppScale = FSlateApplication::Get().GetApplicationScale();
+				const float ScaleFactor = OwnerWindow->GetNativeWindow()->GetDPIScaleFactor() * FSlateApplication::Get().GetApplicationScale();
 
 				if( OwnerWindow->IsWindowMaximized() )
 				{
 					// When the window is maximized and we are in immersive we size the canvas to the size of the visible area which does not include the window border
-					const FMargin& WindowContentMargin = OwnerWindow->GetWindowBorderSize() * AppScale;
+					const FMargin& WindowContentMargin = OwnerWindow->GetWindowBorderSize() * ScaleFactor;
 					ClippedArea.Set( WindowContentMargin.GetTotalSpaceAlong<Orient_Horizontal>(), WindowContentMargin.GetTotalSpaceAlong<Orient_Vertical>() );
 				}
 
-				const float ScaleFactor = OwnerWindow->GetNativeWindow()->GetDPIScaleFactor() * AppScale;
 				TargetSize = (OwnerWindow->GetSizeInScreen() - ClippedArea)/ScaleFactor;
 			}
 		}
diff --git a/Engine/Source/Runtime/Engine/Private/Slate/SceneViewport.cpp b/Engine/Source/Runtime/Engine/Private/Slate/SceneViewport.cpp
--- a/Engine/Source/Runtime/Engine/Private/Slate/SceneViewport.cpp
+++ b/Engine/Source/Runtime/Engine/Private/Slate/SceneViewport.cpp
@@ -234,17 +234,19 @@ void FSceneViewport::UpdateCachedCursorPos( const FGeometry& InGeometry, const F
 		const FIntPoint NewPos = LocalPixelMousePos.IntPoint();
 		const FIntPoint Delta = NewPos - CachedCursorPos;
 		
+		const int32 PixelWidth = FMath::RoundToInt32(FMath::Clamp(InGeometry.Size.X * CachedGeometry.Scale, (double)TNumericLimits<int32>::Min(), (double)TNumericLimits<int32>::Max()));
+		const int32 PixelHeight = FMath::RoundToInt32(FMath::Clamp(InGeometry.Size.Y * CachedGeometry.Scale, (double)TNumericLimits<int32>::Min(), (double)TNumericLimits<int32>::Max()));
+		
 		CachedCursorPos = NewPos;
 		
-		// The cached cursor value is scaled and needs to be reversed to get an accurate calculation
-		if (InGeometry.IsUnderLocation(InGeometry.LocalToAbsolute(FVector2f(CachedCursorPos) / InGeometry.Scale)))
+		if (NewPos.X > 0 && NewPos.X < PixelWidth - 1 && NewPos.Y > 0 && NewPos.Y < PixelHeight - 1)
 		{
-			// When the mouse is over the viewport, the unbound cursor should move in lock step with the normal cursor
+			// When the mouse is over the viewport and not on the edge, the unbound cursor should move in lock step with the normal cursor
 			CachedUnboundedCursorPos += Delta;
 		}
 		else if (bIsMoveEvent)
 		{
-			// When the mouse is outside the viewport, rely on straight deltas.
+			// When the mouse is outside the viewport or on the edge, rely on straight deltas.
 			// This distinction was made due to seeing duplicate OnMouseMove deltas being applied.
 			// If that problem can be ameliorated, the above branch should no longer be necessary.
 			const FIntPoint UnboundDelta = (InMouseEvent.GetCursorDelta() * InGeometry.Scale).IntPoint();
