- Fix for memory leak with FXSystemSet
- There was no reference to the set and therefore it could leak along with the containing FX systems

#Synced-CL 51164657

diff --git a/Engine/Plugins/FX/Cascade/Source/Cascade/Private/Cascade.cpp b/Engine/Plugins/FX/Cascade/Source/Cascade/Private/Cascade.cpp
--- a/Engine/Plugins/FX/Cascade/Source/Cascade/Private/Cascade.cpp
+++ b/Engine/Plugins/FX/Cascade/Source/Cascade/Private/Cascade.cpp
@@ -415,7 +415,7 @@ FFXSystemInterface* FCascade::GetFXSystem() const
 	check(PreviewViewport.IsValid());
 	auto World = PreviewViewport->GetViewportClient()->GetPreviewScene().GetWorld();
 	check(World);
-	return World->FXSystem;
+	return World->FXSystem.Get();
 }
 
 UParticleEmitter* FCascade::GetSelectedEmitter() const
diff --git a/Engine/Source/Editor/UnrealEd/Private/ThumbnailHelpers.cpp b/Engine/Source/Editor/UnrealEd/Private/ThumbnailHelpers.cpp
--- a/Engine/Source/Editor/UnrealEd/Private/ThumbnailHelpers.cpp
+++ b/Engine/Source/Editor/UnrealEd/Private/ThumbnailHelpers.cpp
@@ -208,12 +208,12 @@ FParticleSystemThumbnailScene::FParticleSystemThumbnailScene()
 	PartComponent = NULL;
 
 	ThumbnailFXSystem = FFXSystemInterface::Create(GetScene()->GetFeatureLevel(), GetScene());
-	GetScene()->SetFXSystem( ThumbnailFXSystem );
+	GetScene()->SetFXSystem( ThumbnailFXSystem.Get());
 }
 
 FParticleSystemThumbnailScene::~FParticleSystemThumbnailScene()
 {
-	FFXSystemInterface::MarkPendingKill( ThumbnailFXSystem );
+	FFXSystemInterface::MarkPendingKill( ThumbnailFXSystem.Get() );
 }
 
 void FParticleSystemThumbnailScene::SetParticleSystem(UParticleSystem* ParticleSystem)
diff --git a/Engine/Source/Editor/UnrealEd/Public/ThumbnailHelpers.h b/Engine/Source/Editor/UnrealEd/Public/ThumbnailHelpers.h
--- a/Engine/Source/Editor/UnrealEd/Public/ThumbnailHelpers.h
+++ b/Engine/Source/Editor/UnrealEd/Public/ThumbnailHelpers.h
@@ -106,7 +106,7 @@ protected:
 	class UParticleSystemComponent* PartComponent;
 
 	/** The FXSystem used to render all thumbnail particle systems */
-	class FFXSystemInterface* ThumbnailFXSystem;
+	TSharedPtr<class FFXSystemInterface> ThumbnailFXSystem;
 };
 
 class FMaterialThumbnailScene : public FThumbnailPreviewScene
diff --git a/Engine/Source/Runtime/Engine/Classes/Engine/World.h b/Engine/Source/Runtime/Engine/Classes/Engine/World.h
--- a/Engine/Source/Runtime/Engine/Classes/Engine/World.h
+++ b/Engine/Source/Runtime/Engine/Classes/Engine/World.h
@@ -1870,7 +1870,7 @@ public:
 	FURL										URL;
 
 	/** Interface to the FX system managing particles and related effects for this world.										*/
-	class FFXSystemInterface*					FXSystem;
+	TSharedPtr<class FFXSystemInterface>		FXSystem;
 
 	/** Tick function for starting physics																						*/
 	FStartPhysicsTickFunction StartPhysicsTickFunction;
diff --git a/Engine/Source/Runtime/Engine/Private/Particles/FXSystem.cpp b/Engine/Source/Runtime/Engine/Private/Particles/FXSystem.cpp
--- a/Engine/Source/Runtime/Engine/Private/Particles/FXSystem.cpp
+++ b/Engine/Source/Runtime/Engine/Private/Particles/FXSystem.cpp
@@ -54,7 +54,7 @@ void FFXSystemInterface::FFXInterfaceDeletor::operator()(FFXSystemInterface* Obj
 	External FX system interface.
 -----------------------------------------------------------------------------*/
 
-FFXSystemInterface* FFXSystemInterface::Create(ERHIFeatureLevel::Type InFeatureLevel, FSceneInterface* Scene)
+TSharedRef<FFXSystemInterface> FFXSystemInterface::Create(ERHIFeatureLevel::Type InFeatureLevel, FSceneInterface* Scene)
 {
 	check(Scene);
 	EShaderPlatform InShaderPlatform = Scene->GetShaderPlatform();
@@ -65,8 +65,8 @@ FFXSystemInterface* FFXSystemInterface::Create(ERHIFeatureLevel::Type InFeatureL
 
 	if (CreateCustomFXDelegates.Num())
 	{
-		FFXSystemSet* Set = new FFXSystemSet(GPUSortManager);
-		Scene->SetFXSystem(Set);
+		TSharedRef<FFXSystemSet> Set = MakeShareable(new FFXSystemSet(GPUSortManager), FFXInterfaceDeletor());
+		Scene->SetFXSystem(&Set.Get());
 		Set->SetSceneInterface(Scene);
 		Set->SetScene((FScene*)Scene);
 
@@ -86,8 +86,8 @@ FFXSystemInterface* FFXSystemInterface::Create(ERHIFeatureLevel::Type InFeatureL
 	}
 	else
 	{
-		FFXSystemInterface* Ret = new FFXSystem(InFeatureLevel, InShaderPlatform, GPUSortManager);
-		Scene->SetFXSystem(Ret);
+		TSharedRef<FFXSystemInterface> Ret = MakeShareable(new FFXSystem(InFeatureLevel, InShaderPlatform, GPUSortManager), FFXInterfaceDeletor());
+		Scene->SetFXSystem(&Ret.Get());
 		Ret->SetSceneInterface(Scene);
 		Ret->SetScene((FScene*)Scene);
 		return Ret;
diff --git a/Engine/Source/Runtime/Engine/Private/World.cpp b/Engine/Source/Runtime/Engine/Private/World.cpp
--- a/Engine/Source/Runtime/Engine/Private/World.cpp
+++ b/Engine/Source/Runtime/Engine/Private/World.cpp
@@ -1593,7 +1593,7 @@ void UWorld::FinishDestroy()
 
 		if (FXSystem)
 		{
-			FFXSystemInterface::MarkPendingKill( FXSystem );
+			FFXSystemInterface::MarkPendingKill( FXSystem.Get() );
 			FXSystem = NULL;
 		}
 
@@ -6486,7 +6486,7 @@ void UWorld::CleanupWorldInternal(bool bSessionEnded, bool bCleanupResources, bo
 
 	if(FXSystem && bWorldChanged)
 	{
-		FFXSystemInterface::MarkPendingKill( FXSystem );
+		FFXSystemInterface::MarkPendingKill( FXSystem.Get() );
 		Scene->SetFXSystem(NULL);
 		FXSystem = NULL;
 	}
@@ -10129,7 +10129,7 @@ void UWorld::RecreateScene(ERHIFeatureLevel::Type InFeatureLevel, bool bBroadcas
 		if (FXSystem)
 		{
 			bCreateFXSystem = true;
-			FFXSystemInterface::MarkPendingKill(FXSystem);
+			FFXSystemInterface::MarkPendingKill(FXSystem.Get());
 			FXSystem = nullptr;
 			Scene->SetFXSystem(nullptr);
 		}
diff --git a/Engine/Source/Runtime/Engine/Public/FXSystem.h b/Engine/Source/Runtime/Engine/Public/FXSystem.h
--- a/Engine/Source/Runtime/Engine/Public/FXSystem.h
+++ b/Engine/Source/Runtime/Engine/Public/FXSystem.h
@@ -138,7 +138,7 @@ public:
 	/**
 	 * Create an effects system instance.
 	 */
-	ENGINE_API static FFXSystemInterface* Create(ERHIFeatureLevel::Type InFeatureLevel, FSceneInterface* Scene);
+	ENGINE_API static TSharedRef<FFXSystemInterface> Create(ERHIFeatureLevel::Type InFeatureLevel, FSceneInterface* Scene);
 
 	/**
 	 * Notify the effect system instance that it should release resources in preperation for destruction.
diff --git a/Engine/Source/Runtime/Renderer/Private/RendererScene.cpp b/Engine/Source/Runtime/Renderer/Private/RendererScene.cpp
--- a/Engine/Source/Runtime/Renderer/Private/RendererScene.cpp
+++ b/Engine/Source/Runtime/Renderer/Private/RendererScene.cpp
@@ -1200,7 +1200,7 @@ FScene::FScene(UWorld* InWorld, bool bInRequiresHitProxies, bool bInIsEditorScen
 
 	if (World->FXSystem)
 	{
-		FFXSystemInterface::MarkPendingKill(World->FXSystem);
+		FFXSystemInterface::MarkPendingKill(World->FXSystem.Get());
 	}
 
 	if (bCreateFXSystem)
