Add instrumentation to try to find UGameEngine::SceneViewport lifetime issue

#Synced-CL 51078227

diff --git a/Engine/Source/Runtime/Engine/Private/GameEngine.cpp b/Engine/Source/Runtime/Engine/Private/GameEngine.cpp
--- a/Engine/Source/Runtime/Engine/Private/GameEngine.cpp
+++ b/Engine/Source/Runtime/Engine/Private/GameEngine.cpp
@@ -294,6 +294,7 @@ void UGameEngine::CreateGameViewport( UGameViewportClient* GameViewportClient )
 	}
 
 	SceneViewport = MakeShareable( GameViewportClient->CreateGameViewport(GameViewportWidgetRef) );
+	SceneViewport->bCrashOnDestruction = true; // TEMP: Check for premature destruction
 	GameViewportClient->Viewport = SceneViewport.Get();
 	//GameViewportClient->CreateHighresScreenshotCaptureRegionWidget(); //  Disabled until mouse based input system can be made to work correctly.
 
@@ -807,12 +808,16 @@ void UGameEngine::SwitchGameWindowToUseGameViewport()
 	}
 }
 
-void UGameEngine::OnGameWindowClosed( const TSharedRef<SWindow>& WindowBeingClosed )
+void UGameEngine::OnGameWindowClosed(const TSharedRef<SWindow>& WindowBeingClosed)
 {
 	FSlateApplication::Get().UnregisterGameViewport();
-	// This will shutdown the game
-	GameViewport->CloseRequested(SceneViewport.Get());
-	SceneViewport.Reset();
+	if (ensure(SceneViewport.IsValid()))
+	{
+		// This will shutdown the game
+		GameViewport->CloseRequested(SceneViewport.Get());
+		SceneViewport->bCrashOnDestruction = false; // TEMP: Check for premature destruction
+		SceneViewport.Reset();
+	}
 }
 
 void UGameEngine::OnGameWindowMoved_SavePosition(const TSharedRef<SWindow>& WindowBeingMoved, int32 SaveWinPosIndex)
diff --git a/Engine/Source/Runtime/Engine/Private/Slate/SceneViewport.cpp b/Engine/Source/Runtime/Engine/Private/Slate/SceneViewport.cpp
--- a/Engine/Source/Runtime/Engine/Private/Slate/SceneViewport.cpp
+++ b/Engine/Source/Runtime/Engine/Private/Slate/SceneViewport.cpp
@@ -78,6 +78,7 @@ FSceneViewport::FSceneViewport( FViewportClient* InViewportClient, TSharedPtr<SV
 
 FSceneViewport::~FSceneViewport()
 {
+	ensureAlwaysMsgf(!bCrashOnDestruction, TEXT("Freeing FSceneViewport too early.")); // TEMP: Check for premature destruction
 	Destroy();
 	// Wait for resources to be deleted
 	FlushRenderingCommands();
diff --git a/Engine/Source/Runtime/Engine/Public/Slate/SceneViewport.h b/Engine/Source/Runtime/Engine/Public/Slate/SceneViewport.h
--- a/Engine/Source/Runtime/Engine/Public/Slate/SceneViewport.h
+++ b/Engine/Source/Runtime/Engine/Public/Slate/SceneViewport.h
@@ -35,6 +35,9 @@ public:
 	ENGINE_API FSceneViewport( FViewportClient* InViewportClient, TSharedPtr<SViewport> InViewportWidget );
 	ENGINE_API ~FSceneViewport();
 
+	/** TEMP: Check for premature destruction */
+	UE_INTERNAL bool bCrashOnDestruction = false;
+
 	virtual void* GetWindow() override { return NULL; }
 
 	/** FViewport interface */
