[Blueprint] Fixed Rename function to rename the BlueprintClass as opposed to CDO only.

#rnx
#Synced-CL 51164812

diff --git a/Engine/Source/Runtime/Engine/Private/Blueprint.cpp b/Engine/Source/Runtime/Engine/Private/Blueprint.cpp
--- a/Engine/Source/Runtime/Engine/Private/Blueprint.cpp
+++ b/Engine/Source/Runtime/Engine/Private/Blueprint.cpp
@@ -524,14 +524,12 @@ bool UBlueprint::RenameGeneratedClasses(const TCHAR* NewName, UObject* NewOuter,
 
 	if (bRenameGeneratedClasses)
 	{
-		const auto TryFreeCDOName = [](FName NewClassName, UObject* NewOuter, ERenameFlags InFlags)
+		const auto TryFreeName = [](const FString& NewClassString, UObject* NewOuter, ERenameFlags InFlags)
 			{
-				const FString& NewCDOName = FString(DEFAULT_OBJECT_PREFIX) + NewClassName.ToString();
-
-				if (UObject* ExistingCDO = StaticFindObjectFast(UObject::StaticClass(), NewOuter, *NewCDOName))
+				if (UObject* ExistingObject = StaticFindObjectFast(UObject::StaticClass(), NewOuter, *NewClassString))
 				{
-					const FName NewName = MakeUniqueObjectName(NewOuter, ExistingCDO->GetClass(), *NewCDOName);
-					ExistingCDO->Rename(*(NewName.ToString()), NewOuter, InFlags | REN_DontCreateRedirectors | REN_AllowPackageLinkerMismatch);
+					const FName NewName = MakeUniqueObjectName(NewOuter, ExistingObject->GetClass(), *NewClassString);
+					ExistingObject->Rename(*(NewName.ToString()), NewOuter, InFlags | REN_DontCreateRedirectors | REN_AllowPackageLinkerMismatch);
 				}
 			};
 
@@ -559,10 +557,11 @@ bool UBlueprint::RenameGeneratedClasses(const TCHAR* NewName, UObject* NewOuter,
 		UPackage* NewTopLevelObjectOuter = NewOuter ? NewOuter->GetOutermost() : nullptr;
 		if (GeneratedClass != nullptr)
 		{
-			// check for collision of CDO name, move aside if necessary:
-			TryFreeCDOName(NewGenClassName, NewTopLevelObjectOuter, Flags);
+			const FString NewGenClassNameString = NewGenClassName.ToString();
+
+			TryFreeName(NewGenClassNameString, NewTopLevelObjectOuter, Flags);
 			CheckRedirectors(NewGenClassName, GeneratedClass, NewTopLevelObjectOuter);
-			const bool bMovedOK = GeneratedClass->Rename(*NewGenClassName.ToString(), NewTopLevelObjectOuter, Flags);
+			const bool bMovedOK = GeneratedClass->Rename(*NewGenClassNameString, NewTopLevelObjectOuter, Flags);
 			if (!bMovedOK)
 			{
 				return false;
@@ -572,14 +571,15 @@ bool UBlueprint::RenameGeneratedClasses(const TCHAR* NewName, UObject* NewOuter,
 		// Also move skeleton class, if different from generated class, to new package (again, to create redirector)
 		const bool bShouldMoveSkeleton =
 			(SkeletonGeneratedClass != nullptr) && 
-			(SkeletonGeneratedClass != GeneratedClass)
-		;
+			(SkeletonGeneratedClass != GeneratedClass);
 
 		if (bShouldMoveSkeleton)
 		{
-			TryFreeCDOName(NewSkelClassName, NewTopLevelObjectOuter, Flags);
+			const FString NewSkelClassNameString = NewSkelClassName.ToString();
+
+			TryFreeName(NewSkelClassNameString, NewTopLevelObjectOuter, Flags);
 			CheckRedirectors(NewSkelClassName, SkeletonGeneratedClass, NewTopLevelObjectOuter);
-			const bool bMovedOK = SkeletonGeneratedClass->Rename(*NewSkelClassName.ToString(), NewTopLevelObjectOuter, Flags);
+			const bool bMovedOK = SkeletonGeneratedClass->Rename(*NewSkelClassNameString, NewTopLevelObjectOuter, Flags);
 			if (!bMovedOK)
 			{
 				return false;
