New material translator, minor change:
- Module interned strings are now automatically de-duplicated.
- Bug fixed with FunctionHLSL not allowing two functions with the same name to exist (weren't made unique).
- Some minor function renaming.

#rnx

#Synced-CL 51020780

diff --git a/Engine/Source/Runtime/Engine/Private/Materials/MaterialIREmitter.cpp b/Engine/Source/Runtime/Engine/Private/Materials/MaterialIREmitter.cpp
--- a/Engine/Source/Runtime/Engine/Private/Materials/MaterialIREmitter.cpp
+++ b/Engine/Source/Runtime/Engine/Private/Materials/MaterialIREmitter.cpp
@@ -4203,6 +4203,7 @@ const FFunction* FEmitter::FunctionHLSL(const FFunctionHLSLDesc& Desc)
 	Function->Code = Module->InternString(Desc.Code);
 	Function->Defines = MakeArrayCopy(*this, Desc.Defines);
 	Function->Includes = MakeArrayCopy(*this, Desc.Includes);
+	Function->UniqueId = Module->FunctionHLSLs.Num();
 
 	// Add it to the module list
 	Module->FunctionHLSLs.Push(Function);
diff --git a/Engine/Source/Runtime/Engine/Private/Materials/MaterialIRModule.cpp b/Engine/Source/Runtime/Engine/Private/Materials/MaterialIRModule.cpp
--- a/Engine/Source/Runtime/Engine/Private/Materials/MaterialIRModule.cpp
+++ b/Engine/Source/Runtime/Engine/Private/Materials/MaterialIRModule.cpp
@@ -27,7 +27,7 @@ void FMaterialIRModule::Empty()
 	ParameterCollections.Empty();
 	EnvironmentDefines.Empty();
 	IntegerEnvironmentDefines.Empty();
-	UserStrings.Empty();
+	InternedStrings.Empty();
 	ParameterIdToData.Empty();
 	ParameterInfoToId.Empty();
 	Statistics = {};
@@ -63,6 +63,16 @@ void FMaterialIRModule::BeginCustomOutputGroup(FStringView Name, int32 NumOutput
 
 FStringView FMaterialIRModule::InternString(FStringView InString)
 {
+	// Precompute the hash first, to be used in both finding and inserting the string into the set of interned strings
+	uint32 KeyHash = GetTypeHash(InString);
+
+	// Look for the string in the set of interned strings first, and return existing if found.
+	const FStringView* ExistingString = InternedStrings.FindByHash(KeyHash, InString);
+	if (ExistingString)
+	{
+		return *ExistingString;
+	}
+
 	// Allocate the buffer to contain the string
 	FStringView::ElementType* String = (FStringView::ElementType*)Allocator.PushBytes(InString.NumBytes() + sizeof(FStringView::ElementType), alignof(FStringView::ElementType));
 
@@ -72,7 +82,13 @@ FStringView FMaterialIRModule::InternString(FStringView InString)
 	// Mark the null-character at the end
 	String[InString.Len()] = 0;
 
-	return { String, InString.Len() };
+	// Make the view
+	FStringView View = { String, InString.Len() };
+
+	// Add the new string into the set of interned strings so future fetches of this string will result in the same unique string
+	InternedStrings.AddByHash(KeyHash, View);
+
+	return View;
 }
 
 void FMaterialIRModule::AddError(UMaterialExpression* Source, FString Message)
diff --git a/Engine/Source/Runtime/Engine/Private/Materials/MaterialIRToHLSLTranslator.cpp b/Engine/Source/Runtime/Engine/Private/Materials/MaterialIRToHLSLTranslator.cpp
--- a/Engine/Source/Runtime/Engine/Private/Materials/MaterialIRToHLSLTranslator.cpp
+++ b/Engine/Source/Runtime/Engine/Private/Materials/MaterialIRToHLSLTranslator.cpp
@@ -995,7 +995,7 @@ struct FEntryPointCodeGen
 
 	void LowerSetMaterialOutput(const MIR::FSetMaterialOutput* SetMaterialOutput)
 	{
-		if (SetMaterialOutput->Property >= MP_CustomizedUVs0 && SetMaterialOutput->Property <= MP_CustomizedUVs7)
+		if (SetMaterialOutput->Property >= MP_CustomizedUVs0 && SetMaterialOutput->Property <= MP_LastCustomizedUVs)
 		{
 			Printer << TEXTVIEW("Parameters.CustomizedUVs[") << (int)(SetMaterialOutput->Property - MP_CustomizedUVs0) << TEXT(']');
 		}
@@ -2269,7 +2269,7 @@ static void GenerateOtherTemplateStringParameters(const FMaterial* Material, con
 		FPrinter Printer { MaterialAggregatesRequiringDeclaration, 1 };
 		for (EMaterialProperty Property : MaterialAttributesAggregate::GetMaterialProperties())
 		{
-			if (Property == MP_WorldPositionOffset || (Property >= MP_CustomizedUVs0 && Property <= MP_CustomizedUVs7)) // todo: fixme
+			if (Property == MP_WorldPositionOffset || (Property >= MP_CustomizedUVs0 && Property <= MP_LastCustomizedUVs)) // todo: fixme
 			{
 				continue;
 			}
diff --git a/Engine/Source/Runtime/Engine/Private/Materials/MaterialIRValueAnalyzer.cpp b/Engine/Source/Runtime/Engine/Private/Materials/MaterialIRValueAnalyzer.cpp
--- a/Engine/Source/Runtime/Engine/Private/Materials/MaterialIRValueAnalyzer.cpp
+++ b/Engine/Source/Runtime/Engine/Private/Materials/MaterialIRValueAnalyzer.cpp
@@ -25,12 +25,12 @@
 #include "UObject/Package.h"
 #include "Shader/Preshader2.h"
 
-static void AnalyzeAnalysisError(FMaterialIRValueAnalyzer& Analyzer, MIR::FAnalysisError* AnalysisError)
+static void AnalysisError_Analyze(FMaterialIRValueAnalyzer& Analyzer, MIR::FAnalysisError* AnalysisError)
 {
 	Analyzer.Module->AddError(AnalysisError->Expression, AnalysisError->Message.ToString());
 }
 
-static void AnalyzeInStageHardwarePartialDerivative(FMaterialIRValueAnalyzer& Analyzer, MIR::FHardwarePartialDerivative* PartialDerivative, MIR::EStage Stage)
+static void HardwarePartialDerivative_AnalyzeInStage(FMaterialIRValueAnalyzer& Analyzer, MIR::FHardwarePartialDerivative* PartialDerivative, MIR::EStage Stage)
 {
 	if (Stage == MIR::Stage_Vertex)
 	{
@@ -56,7 +56,7 @@ static EMaterialShaderFrequency MapToMaterialShaderFrequencyOrAny(MIR::EStage St
 	return EMaterialShaderFrequency::Any;
 }
 
-static void AnalyzeExtern(FMaterialIRValueAnalyzer& Analyzer, MIR::FExtern* Extern)
+static void Extern_Analyze(FMaterialIRValueAnalyzer& Analyzer, MIR::FExtern* Extern)
 {
 	MIR::FExternAnalysisContext Context = {
 		.MaterialInterface = Analyzer.MaterialInterface,
@@ -66,7 +66,7 @@ static void AnalyzeExtern(FMaterialIRValueAnalyzer& Analyzer, MIR::FExtern* Exte
 	Extern->Analyze(Context);
 }
 
-static void AnalyzeInStageExtern(FMaterialIRValueAnalyzer& Analyzer, MIR::FExtern* Extern, MIR::EStage Stage)
+static void Extern_AnalyzeInStage(FMaterialIRValueAnalyzer& Analyzer, MIR::FExtern* Extern, MIR::EStage Stage)
 {
 	MIR::FExternAnalysisContext Context = {
 		.MaterialInterface = Analyzer.MaterialInterface,
@@ -198,7 +198,7 @@ static int32 GetPageTableLayerIndex(MIR::FValue* TextureObject)
 	return INDEX_NONE;
 }
 
-static void AnalyzeInStageVTPageTableRead(FMaterialIRValueAnalyzer& Analyzer, MIR::FVTPageTableRead* VTPageTableRead, MIR::EStage Stage)
+static void VTPageTableRead_AnalyzeInStage(FMaterialIRValueAnalyzer& Analyzer, MIR::FVTPageTableRead* VTPageTableRead, MIR::EStage Stage)
 {
 	check(VTPageTableRead->TextureObject);
 
@@ -232,7 +232,7 @@ static void AnalyzeInStageVTPageTableRead(FMaterialIRValueAnalyzer& Analyzer, MI
 	}
 }
 
-static void AnalyzeBranch(FMaterialIRValueAnalyzer& Analyzer, MIR::FBranch* Branch)
+static void Branch_Analyze(FMaterialIRValueAnalyzer& Analyzer, MIR::FBranch* Branch)
 {
 	Branch->TrueBlock = Analyzer.Module->AllocateArray<MIR::FBlock>(Analyzer.Module->GetNumEntryPoints());
 	MIR::ZeroArray(Branch->TrueBlock);
@@ -241,7 +241,7 @@ static void AnalyzeBranch(FMaterialIRValueAnalyzer& Analyzer, MIR::FBranch* Bran
 	MIR::ZeroArray(Branch->FalseBlock);
 }
 
-static void AnalyzeTextureRead(FMaterialIRValueAnalyzer& Analyzer, MIR::FTextureRead* TextureRead)
+static void TextureRead_Analyze(FMaterialIRValueAnalyzer& Analyzer, MIR::FTextureRead* TextureRead)
 {
 	// Ensure a uniform parameter is allocated when a texture read instruction accesses the texture.
 	// Otherwise, no uniform parameter must be allocated for the texture object in case it's used for other nodes only such as texture properties.
@@ -372,7 +372,7 @@ static void AnalyzePrimitiveUniformParameter(FMaterialIRValueAnalyzer& Analyzer,
 	}
 }
 
-static void AnalyzeUniformParameter(FMaterialIRValueAnalyzer& Analyzer, MIR::FUniformParameter* Parameter)
+static void UniformParameter_Analyze(FMaterialIRValueAnalyzer& Analyzer, MIR::FUniformParameter* Parameter)
 {
 	if (!Parameter->Type.IsTexture())
 	{
@@ -381,7 +381,7 @@ static void AnalyzeUniformParameter(FMaterialIRValueAnalyzer& Analyzer, MIR::FUn
 	}
 }
 
-static void AnalyzePreshaderParameter(FMaterialIRValueAnalyzer& Analyzer, MIR::FPreshaderParameter* Parameter)
+static void PreshaderParameter_Analyze(FMaterialIRValueAnalyzer& Analyzer, MIR::FPreshaderParameter* Parameter)
 {
 	FUniformExpressionSet& UniformExpressionSet = Analyzer.CompilationOutput->UniformExpressionSet;
 
@@ -417,7 +417,7 @@ static void AnalyzePreshaderParameter(FMaterialIRValueAnalyzer& Analyzer, MIR::F
 	Parameter->Analysis_PreshaderOffset = GlobalComponentOffset;
 }
 
-static void AnalyzeSetMaterialOutput(FMaterialIRValueAnalyzer& Analyzer, MIR::FSetMaterialOutput* SetMaterialOutput)
+static void SetMaterialOutput_Analyze(FMaterialIRValueAnalyzer& Analyzer, MIR::FSetMaterialOutput* SetMaterialOutput)
 {
 	// If this output is evaluated in the vertex stage, ensure its input is also valid there.
 	// Inputs marked "PixelStageOnly" cannot be used when generating vertex-stage code
@@ -441,7 +441,7 @@ static void AnalyzeSetMaterialOutput(FMaterialIRValueAnalyzer& Analyzer, MIR::FS
 	}
 }
 
-static void AnalyzeMaterialParameterCollection(FMaterialIRValueAnalyzer& Analyzer, MIR::FMaterialParameterCollection* MaterialParameterCollection)
+static void MaterialParameterCollection_Analyze(FMaterialIRValueAnalyzer& Analyzer, MIR::FMaterialParameterCollection* MaterialParameterCollection)
 {
 	int32 CollectionIndex = Analyzer.Module->FindOrAddParameterCollection(MaterialParameterCollection->Collection);
 		
@@ -453,7 +453,7 @@ static void AnalyzeMaterialParameterCollection(FMaterialIRValueAnalyzer& Analyze
 	MaterialParameterCollection->Analysis_CollectionIndex = CollectionIndex;
 }
 
-static void AnalyzeGetVertexInterpolator(FMaterialIRValueAnalyzer& Analyzer, MIR::FGetVertexInterpolator* GetVertexInterpolator)
+static void GetVertexInterpolator_Analyze(FMaterialIRValueAnalyzer& Analyzer, MIR::FGetVertexInterpolator* GetVertexInterpolator)
 {
 	// Record in the compilation output that the material is using vertex interpolators
 	Analyzer.Module->GetCompilationOutput().bUsesVertexInterpolator = true;
@@ -463,7 +463,7 @@ static void AnalyzeGetVertexInterpolator(FMaterialIRValueAnalyzer& Analyzer, MIR
 	Analyzer.Module->GetEntryPoint(MIR::EStage::Stage_Vertex).Outputs.Push(GetVertexInterpolator->SetInstr);
 }
 
-static void AnalyzeSetVertexInterpolator(FMaterialIRValueAnalyzer& Analyzer, MIR::FSetVertexInterpolator* SetVertexInterpolator)
+static void SetVertexInterpolator_Analyze(FMaterialIRValueAnalyzer& Analyzer, MIR::FSetVertexInterpolator* SetVertexInterpolator)
 {
 	int32& NumInterpRegisters = Analyzer.Module->GetStatistics().NumUsedVertexInterpolatorFloatRegisters;
 
@@ -504,7 +504,7 @@ void FMaterialIRValueAnalyzer::Analyze(MIR::FValue* Value)
 	#define EXPAND_CASE(ValueType) \
 		case MIR::VK_##ValueType: \
 		{ \
-			Analyze##ValueType(*this, static_cast<MIR::F##ValueType*>(Value)); \
+			ValueType##_Analyze(*this, static_cast<MIR::F##ValueType*>(Value)); \
 			break; \
 		}
 
@@ -532,7 +532,7 @@ void FMaterialIRValueAnalyzer::AnalyzeInStage(MIR::FValue* Value, MIR::EStage St
 	#define EXPAND_CASE(ValueType) \
 	case MIR::VK_##ValueType: \
 	{ \
-		AnalyzeInStage##ValueType(*this, static_cast<MIR::F##ValueType*>(Value), Stage); \
+		ValueType##_AnalyzeInStage(*this, static_cast<MIR::F##ValueType*>(Value), Stage); \
 		break; \
 	}
 
diff --git a/Engine/Source/Runtime/Engine/Public/Materials/MaterialIRModule.h b/Engine/Source/Runtime/Engine/Public/Materials/MaterialIRModule.h
--- a/Engine/Source/Runtime/Engine/Public/Materials/MaterialIRModule.h
+++ b/Engine/Source/Runtime/Engine/Public/Materials/MaterialIRModule.h
@@ -265,8 +265,8 @@ private:
 	// Parameter metadata.
 	TArray<TPair<FMaterialParameterInfo, FMaterialParameterMetadata>> ParameterIdToData;
 
-	// Stores user-defined strings in this module.
-	TArray<FString> UserStrings;
+	// Stores the set of user generated interned strings in this module.
+	TSet<FStringView> InternedStrings;
 
 	// Environment define names for shader compilation in this module.
 	// @massimo.tristano refactor this into a better way to carry output flags from the builder to the backend.
