[HLODActor/ALT]
- Register to delegate in postload when constructor is called from ALT

#Synced-CL 51209060

diff --git a/Engine/Source/Runtime/Engine/Private/WorldPartition/HLOD/HLODActor.cpp b/Engine/Source/Runtime/Engine/Private/WorldPartition/HLOD/HLODActor.cpp
--- a/Engine/Source/Runtime/Engine/Private/WorldPartition/HLOD/HLODActor.cpp
+++ b/Engine/Source/Runtime/Engine/Private/WorldPartition/HLOD/HLODActor.cpp
@@ -79,7 +79,10 @@ AWorldPartitionHLOD::AWorldPartitionHLOD(const FObjectInitializer& ObjectInitial
 #endif
 
 #if WITH_EDITOR
-	FWorldDelegates::OnWorldCleanup.AddUObject(this, &AWorldPartitionHLOD::OnWorldCleanup);
+	if (!IsInAsyncLoadingThread())
+	{
+		FWorldDelegates::OnWorldCleanup.AddUObject(this, &AWorldPartitionHLOD::OnWorldCleanup);
+	}
 #endif
 }
 
@@ -331,6 +334,8 @@ void AWorldPartitionHLOD::PostLoad()
 			check(SourceCellGuid.IsValid());
 		}
 	}
+
+	FWorldDelegates::OnWorldCleanup.AddUObject(this, &AWorldPartitionHLOD::OnWorldCleanup);
 #endif
 }
 
