New material translator small fixes:
- Operators with zero derivative now are cast to the correct number of dimensions
- Fixed material attributes blending again by disabling blending on the FrontMaterial attribute.

#rnx

#Synced-CL 51213401

diff --git a/Engine/Source/Runtime/Engine/Private/Materials/MaterialExpressionsIR.cpp b/Engine/Source/Runtime/Engine/Private/Materials/MaterialExpressionsIR.cpp
--- a/Engine/Source/Runtime/Engine/Private/Materials/MaterialExpressionsIR.cpp
+++ b/Engine/Source/Runtime/Engine/Private/Materials/MaterialExpressionsIR.cpp
@@ -3728,6 +3728,11 @@ static FValueRef Blend(MIR::FEmitter& Em, EMaterialExpressionBlendMode PixelAttr
 		// Arguments are aggregates, so recursively blend each attribute pair.
 		return BlendAggregate(Em, PixelAttributeBlendMode, VertexAttributeBlendMode, A, B, Alpha);
 	}
+	else if (A->Type.IsSubstrateData())
+	{
+		// Ignore this attribute for now. FrontMaterial should be removed from MaterialAttributes. @massimo.tristano, @charles.derousiers, @sebastien.hillaire
+		return Em.SubstrateDefaultSlab();
+	}
 	else
 	{
 		Em.Errorf(TEXT("Cannot blend values of type '%s'."), *A->Type.GetSpelling());
diff --git a/Engine/Source/Runtime/Engine/Private/Materials/MaterialIREmitter.cpp b/Engine/Source/Runtime/Engine/Private/Materials/MaterialIREmitter.cpp
--- a/Engine/Source/Runtime/Engine/Private/Materials/MaterialIREmitter.cpp
+++ b/Engine/Source/Runtime/Engine/Private/Materials/MaterialIREmitter.cpp
@@ -3903,7 +3903,7 @@ static FValue* DifferentiateOperator(FEmitter& E, FOperator* Op, EDerivativeAxis
 		case UO_Floor:
 		case UO_Round:
 		case UO_Truncate:
-			return Zero();
+			return E.Cast(Zero(), Op->Type);
 
 		// d/dx (f(x) + g(x)) = f'(x) + g'(x)
 		case BO_Add:
@@ -3983,7 +3983,7 @@ static FValue* DifferentiateOperator(FEmitter& E, FOperator* Op, EDerivativeAxis
 					   E.Add(E.Multiply(dH, E.Subtract(G, F)), 
 						   E.Multiply(H, E.Subtract(dG, dF))));
 
-		// d/dx select(F, g(x), h(x)) ˜ select(F, g'(x), h'(x))
+		// d/dx select(F, g(x), h(x)) = select(F, g'(x), h'(x))
 		case TO_Select:
 			return E.Select(F, dG, dH);
 
@@ -4010,7 +4010,7 @@ static FValue* DifferentiateOperator(FEmitter& E, FOperator* Op, EDerivativeAxis
 		case BO_BitShiftLeft:
 		case BO_BitShiftRight:
 		case BO_Step:
-			return Zero();
+			return E.Cast(Zero(), Op->Type);
 
 		default:
 			UE_MIR_UNREACHABLE();
