[Backout] - CL50487808
#rnx
Original CL Desc
-----------------------------------------------------------------
[Functional Tests] Move Rescan in PostLoad to allow async loading with world composition

diff --git a/Engine/Source/Editor/UnrealEd/Private/UnrealEdMisc.cpp b/Engine/Source/Editor/UnrealEd/Private/UnrealEdMisc.cpp
--- a/Engine/Source/Editor/UnrealEd/Private/UnrealEdMisc.cpp
+++ b/Engine/Source/Editor/UnrealEd/Private/UnrealEdMisc.cpp
@@ -992,11 +992,10 @@ bool FUnrealEdMisc::EnableWorldComposition(UWorld* InWorld, bool bEnable)
 			return false;
 		}
 			
-		auto WorldComposition = NewObject<UWorldComposition>(InWorld);
+		auto WorldCompostion = NewObject<UWorldComposition>(InWorld);
 		// All map files found in the same and folder and all sub-folders will be added ass sub-levels to this map
 		// Make sure user understands this
-		WorldComposition->Rescan();
-		int32 NumFoundSublevels = WorldComposition->GetTilesList().Num();
+		int32 NumFoundSublevels = WorldCompostion->GetTilesList().Num();
 		if (NumFoundSublevels)
 		{
 			FFormatNamedArguments Arguments;
@@ -1007,13 +1006,13 @@ bool FUnrealEdMisc::EnableWorldComposition(UWorld* InWorld, bool bEnable)
 			auto AppResult = FMessageDialog::Open(EAppMsgType::OkCancel, Message);
 			if (AppResult != EAppReturnType::Ok)
 			{
-				WorldComposition->MarkAsGarbage();
+				WorldCompostion->MarkAsGarbage();
 				return false;
 			}
 		}
 			
 		// 
-		InWorld->WorldComposition = WorldComposition;
+		InWorld->WorldComposition = WorldCompostion;
 		UWorldComposition::WorldCompositionChangedEvent.Broadcast(InWorld);
 	}
 	
diff --git a/Engine/Source/Runtime/Engine/Classes/Engine/WorldComposition.h b/Engine/Source/Runtime/Engine/Classes/Engine/WorldComposition.h
--- a/Engine/Source/Runtime/Engine/Classes/Engine/WorldComposition.h
+++ b/Engine/Source/Runtime/Engine/Classes/Engine/WorldComposition.h
@@ -169,6 +169,7 @@ class UWorldComposition : public UObject
 
 private:
 	// UObject interface
+	ENGINE_API virtual void PostInitProperties() override;
 	ENGINE_API virtual void Serialize(FArchive& Ar) override;
 	ENGINE_API virtual void PostDuplicate(bool bDuplicateForPIE) override;
 	ENGINE_API virtual void PostLoad() override;
@@ -178,7 +179,7 @@ private:
 	void PopulateStreamingLevels();
 
 	/** Calculates tiles absolute positions based on relative positions */
-	void CalculateTilesAbsolutePositions();
+	void CaclulateTilesAbsolutePositions();
 
 	/** Resets world composition structures */
 	void Reset();
diff --git a/Engine/Source/Runtime/Engine/Private/WorldComposition.cpp b/Engine/Source/Runtime/Engine/Private/WorldComposition.cpp
--- a/Engine/Source/Runtime/Engine/Private/WorldComposition.cpp
+++ b/Engine/Source/Runtime/Engine/Private/WorldComposition.cpp
@@ -42,6 +42,17 @@ UWorldComposition::UWorldComposition(const FObjectInitializer& ObjectInitializer
 {
 }
 
+void UWorldComposition::PostInitProperties()
+{
+	Super::PostInitProperties();
+
+	if (!IsTemplate() && !GetOutermost()->HasAnyPackageFlags(PKG_PlayInEditor) && !FPlatformProperties::RequiresCookedData())
+	{
+		// Tiles information is not serialized to disk in the editor, and should be regenerated on world composition object construction
+		Rescan();
+	}
+}
+
 void UWorldComposition::Serialize( FArchive& Ar )
 {
 	Super::Serialize(Ar);
@@ -64,7 +75,6 @@ void UWorldComposition::PostDuplicate(bool bDuplicateForPIE)
 
 	if (bDuplicateForPIE)
 	{
-		// When duplicating for PIE, we assume the source was up-to-date and has fresh Tiles, so no need to rescan here. But Tiles need PIE package name fixup.
 		FixupForPIE(GetOutermost()->GetPIEInstanceID());	
 	}
 }
@@ -73,18 +83,12 @@ void UWorldComposition::PostLoad()
 {
 	Super::PostLoad();
 
-	if (!IsTemplate() && !GetOutermost()->HasAnyPackageFlags(PKG_PlayInEditor) && !FPlatformProperties::RequiresCookedData())
-	{
-		// Tiles information is not serialized to disk in the editor, and should be regenerated after world composition object load
-		Rescan();
-	}
-
 	if (FPlatformProperties::RequiresCookedData())
 	{
 		// Create streaming levels for each Tile
 		PopulateStreamingLevels();
 		// Calculate absolute positions since they are not serialized to disk
-		CalculateTilesAbsolutePositions();
+		CaclulateTilesAbsolutePositions();
 	}
 
 	UWorld* World = GetWorld();
@@ -321,7 +325,7 @@ void UWorldComposition::Rescan()
 	PopulateStreamingLevels();
 
 	// Calculate absolute positions since they are not serialized to disk
-	CalculateTilesAbsolutePositions();
+	CaclulateTilesAbsolutePositions();
 }
 
 void UWorldComposition::ReinitializeForPIE()
@@ -358,7 +362,7 @@ ULevelStreaming* UWorldComposition::CreateStreamingLevel(const FWorldComposition
 	return StreamingLevel;
 }
 
-void UWorldComposition::CalculateTilesAbsolutePositions()
+void UWorldComposition::CaclulateTilesAbsolutePositions()
 {
 	for (FWorldCompositionTile& Tile : Tiles)
 	{
