Fix virtual texture ensure for RDG external access mode.
Could trigger if VT update/finalize happened more than once in the same RDG graph.
#rnx
#Synced-CL 51204827

diff --git a/Engine/Source/Runtime/Engine/Private/VT/VirtualTextureUploadCache.cpp b/Engine/Source/Runtime/Engine/Private/VT/VirtualTextureUploadCache.cpp
--- a/Engine/Source/Runtime/Engine/Private/VT/VirtualTextureUploadCache.cpp
+++ b/Engine/Source/Runtime/Engine/Private/VT/VirtualTextureUploadCache.cpp
@@ -480,6 +480,7 @@ void FVirtualTextureUploadCache::FinalizeWithUpdateTexture(FRDGBuilder& GraphBui
 			else
 			{
 				TileEntry.Texture = GraphBuilder.RegisterExternalTexture(TileEntry.PooledRenderTarget, ERDGTextureFlags::None);
+				GraphBuilder.UseInternalAccessMode(TileEntry.Texture);
 				UploadParameters->TextureArray.AddUnique(FRDGTextureAccess(TileEntry.Texture, ERHIAccess::CopyDest));
 			}
 			
