Fixing TestFramework test failure due to missing tag.

#rnx
#Synced-CL 51211449

diff --git a/Engine/Source/Runtime/Engine/Private/Tests/EngineAutomationTests.cpp b/Engine/Source/Runtime/Engine/Private/Tests/EngineAutomationTests.cpp
--- a/Engine/Source/Runtime/Engine/Private/Tests/EngineAutomationTests.cpp
+++ b/Engine/Source/Runtime/Engine/Private/Tests/EngineAutomationTests.cpp
@@ -2163,7 +2163,7 @@ bool FAutomationAttachment::RunTest(const FString& Parameters)
 	return false;
 }
 
-constexpr auto ExampleTag = "[TestExampleTag]";
+constexpr auto ExampleTag = "[TestExampleTag][MustPass]";
 constexpr auto OtherTag = "[SomeOtherTag]";
 constexpr auto NegativeTag = "[DoNotWant]";
 constexpr auto NegativeAndExampleTags = "[TestExampleTag][DoNotWant]";
@@ -2318,8 +2318,8 @@ bool FAutomationTagsFilterBothReverseOrder::RunTest(const FString& Parameters)
 	return true;
 }
 
-constexpr auto ComplexTags1 = "[TestExampleTagComplex][FirstExample]";
-constexpr auto ComplexTags2 = "[TestExampleTagComplex][SecondExample]";
+constexpr auto ComplexTags1 = "[TestExampleTagComplex][FirstExample][MustPass]";
+constexpr auto ComplexTags2 = "[TestExampleTagComplex][SecondExample][MustPass]";
 constexpr auto ComplexTagTestPath = "TestFramework.Tags.Complex";
 IMPLEMENT_COMPLEX_AUTOMATION_TEST(FAutomationTagsForComplexSuite, ComplexTagTestPath, EAutomationTestFlags::EditorContext | EAutomationTestFlags::EngineFilter)
 void FAutomationTagsForComplexSuite::GetTests(TArray<FString>& OutBeautifiedNames, TArray<FString>& OutTestCommands) const
