Prevent auto-centering for explicit window positions from -WinX=/-WinY= and -SAVEWINPOS= command-line arguments by locking the monitor index to -1

#Synced-CL 51150456

diff --git a/Engine/Source/Runtime/Engine/Private/GameEngine.cpp b/Engine/Source/Runtime/Engine/Private/GameEngine.cpp
--- a/Engine/Source/Runtime/Engine/Private/GameEngine.cpp
+++ b/Engine/Source/Runtime/Engine/Private/GameEngine.cpp
@@ -503,7 +503,7 @@ TSharedRef<SWindow> UGameEngine::CreateGameWindow()
 	int32 ResX = GSystemResolution.ResX;
 	int32 ResY = GSystemResolution.ResY;
 	EWindowMode::Type WindowMode = GSystemResolution.WindowMode;
-	const int32 MonitorIndex = GSystemResolution.MonitorIndex;
+	int32 MonitorIndex = GSystemResolution.MonitorIndex;
 	
 #if WITH_EDITOR
 	/**************************************************************/
@@ -633,6 +633,13 @@ TSharedRef<SWindow> UGameEngine::CreateGameWindow()
 		}
 	}
 
+	if (AutoCenterType == EAutoCenter::None)
+	{
+		// Position has been set explicitly, so disable multiple monitor behavior.
+		MonitorIndex = -1;
+		GSystemResolution.LockUnknownMonitor();
+	}
+
 	// Give the window the max width/height of either the requested resolution, or your available desktop resolution
 	// We need to do this as we request some 4K windows when rendering sequences, and the OS may try and clamp that
 	// window to your available desktop resolution
diff --git a/Engine/Source/Runtime/Engine/Private/UnrealEngine.cpp b/Engine/Source/Runtime/Engine/Private/UnrealEngine.cpp
--- a/Engine/Source/Runtime/Engine/Private/UnrealEngine.cpp
+++ b/Engine/Source/Runtime/Engine/Private/UnrealEngine.cpp
@@ -18920,9 +18920,19 @@ bool AllowHighQualityLightmaps(const FStaticFeatureLevel FeatureLevel)
 		&& !IsMobilePlatform(GShaderPlatformForFeatureLevel[FeatureLevel]);
 }
 
+static bool bLockUnknownMonitor = false;
+void FSystemResolution::LockUnknownMonitor()
+{
+	SetMonitorIndex(-1);
+	bLockUnknownMonitor = true;
+}
+
 void FSystemResolution::SetMonitorIndex(int32 MonitorIndex)
 {
-	CVarMonitorIndex->Set(MonitorIndex);
+	if (!bLockUnknownMonitor)
+	{
+		CVarMonitorIndex->Set(MonitorIndex);
+	}
 }
 
 // Helper function for changing system resolution via the r.setres console command
diff --git a/Engine/Source/Runtime/Engine/Public/UnrealEngine.h b/Engine/Source/Runtime/Engine/Public/UnrealEngine.h
--- a/Engine/Source/Runtime/Engine/Public/UnrealEngine.h
+++ b/Engine/Source/Runtime/Engine/Public/UnrealEngine.h
@@ -452,6 +452,7 @@ struct FSystemResolution
 	int32 MonitorIndex = -1;
 	bool bForceRefresh = false;
 
+	static void LockUnknownMonitor();
 	static ENGINE_API void SetMonitorIndex(int32 MonitorIndex);
 
 	// Helper function for changing system resolution via the r.setres console command
