[StaticMesh] When saving static mesh render data to DDC cache, override the locally stored data if it differs from what is in the cache. This fixes cooking errors that had recently cropped up due to apparent non-determinism in some static meshes (cause unclear) with DDC disallowing overwrites.

#rnx

#Synced-CL 51222305

diff --git a/Engine/Source/Runtime/Engine/Private/StaticMesh.cpp b/Engine/Source/Runtime/Engine/Private/StaticMesh.cpp
--- a/Engine/Source/Runtime/Engine/Private/StaticMesh.cpp
+++ b/Engine/Source/Runtime/Engine/Private/StaticMesh.cpp
@@ -2350,6 +2350,16 @@ FStaticMeshRenderData::FStaticMeshRenderData()
 }
 
 FStaticMeshRenderData::~FStaticMeshRenderData()
+{
+	FreeLODResources();
+
+	if (RayTracingProxy != nullptr)
+	{
+		delete RayTracingProxy;
+	}
+}
+
+void FStaticMeshRenderData::FreeLODResources()
 {
 	FStaticMeshLODResources** LODResourcesArray = LODResources.GetData();
 	for (int32 LODIndex = 0; LODIndex < LODResources.Num(); ++LODIndex)
@@ -2360,11 +2370,6 @@ FStaticMeshRenderData::~FStaticMeshRenderData()
 		LODResourcesArray[LODIndex] = nullptr;
 	}
 	LODResources.Empty();
-
-	if (RayTracingProxy != nullptr)
-	{
-		delete RayTracingProxy;
-	}
 }
 
 int32 FStaticMeshRenderData::GetNumNonStreamingLODs() const
@@ -2536,6 +2541,10 @@ void FStaticMeshRenderData::Serialize(FArchive& Ar, UStaticMesh* Owner, bool bCo
 	}
 #endif // PLATFORM_DESKTOP
 
+	if (Ar.IsLoading() && !LODResources.IsEmpty())
+	{
+		FreeLODResources();
+	}
 	LODResources.Serialize(Ar, Owner);
 
 #if WITH_EDITOR
@@ -2591,6 +2600,11 @@ void FStaticMeshRenderData::Serialize(FArchive& Ar, UStaticMesh* Owner, bool bCo
 
 			RayTracingProxy->Serialize(Ar, Owner, this, bCooked);
 		}
+		else if (RayTracingProxy != nullptr)
+		{
+			delete RayTracingProxy;
+			RayTracingProxy = nullptr;
+		}
 	}
 
 	// Inline the distance field derived data for cooked builds
@@ -4428,19 +4442,41 @@ void FStaticMeshRenderData::Cache(const ITargetPlatform* TargetPlatform, UStatic
 			{
 				TRACE_CPUPROFILER_EVENT_SCOPE(SaveDDC);
 
+				bool bInvalidateNaniteStreaming = false;
 				FValue Value = FValue::Compress(FSharedBuffer::MakeView(Ar.GetData(), Ar.TotalSize()));
 				RecordBuilder.AddValue(MeshDataId, Value);
 				TotalPushedBytes += Ar.TotalSize();
-				EstimatedCompressedSize = uint64(Value.GetData().GetCompressedSize() * DDCSizeToEstimateFactor) + EstimatedNaniteStreamingCompressedSize;
+				EstimatedMeshDataCompressedSize = uint64(Value.GetData().GetCompressedSize() * DDCSizeToEstimateFactor);
+
+				FCacheRecordPolicyBuilder PolicyBuilder(ECachePolicy::Default | ECachePolicy::SkipData | ECachePolicy::KeepAlive);
+				PolicyBuilder.AddValuePolicy(MeshDataId, ECachePolicy::Default);
 
 				FRequestOwner RequestOwner(UE::DerivedData::EPriority::Blocking);
-				const FCachePutRequest PutRequest = { FSharedString(Owner->GetPathName()), RecordBuilder.Build(), ECachePolicy::Default | ECachePolicy::KeepAlive };
+				const FCachePutRequest PutRequest = { FSharedString(Owner->GetPathName()), RecordBuilder.Build(), PolicyBuilder.Build() };
 				GetCache().Put(MakeArrayView(&PutRequest, 1), RequestOwner,
-					[&bSavedToDDC](FCachePutResponse&& Response)
+					[&bSavedToDDC, &bInvalidateNaniteStreaming, &PutRequest, &MeshDataBuffer,
+					 &EstimatedMeshDataCompressedSize, &NaniteStreamingDataHash](FCachePutResponse&& Response)
 					{
 						if (Response.Status == EStatus::Ok)
 						{
 							bSavedToDDC = true;
+
+							const FValueWithId& ResponseMeshValue = Response.Record.GetValue(MeshDataId);
+							const FValueWithId& ResponseNaniteStreamingValue = Response.Record.GetValue(NaniteStreamingDataId);
+
+							if (ResponseMeshValue != PutRequest.Record.GetValue(MeshDataId))
+							{
+								// The request and response differ, so we're going to reserialize with what is actually in the cache
+								const FCompressedBuffer& CompressedBuffer = ResponseMeshValue.GetData();
+								MeshDataBuffer = CompressedBuffer.Decompress();
+								EstimatedMeshDataCompressedSize = uint64(CompressedBuffer.GetCompressedSize() * DDCSizeToEstimateFactor);
+							}
+
+							if (ResponseNaniteStreamingValue != PutRequest.Record.GetValue(NaniteStreamingDataId))
+							{
+								NaniteStreamingDataHash = ResponseNaniteStreamingValue.GetRawHash();
+								bInvalidateNaniteStreaming = true;
+							}
 						}
 					});
 
@@ -4449,12 +4485,33 @@ void FStaticMeshRenderData::Cache(const ITargetPlatform* TargetPlatform, UStatic
 				if (bSavedToDDC)
 				{
 					NaniteResources.SetHasBuildFromDDCError(false);
-					if (NaniteResources.HasStreamingData() && !IsRunningCookCommandlet())
+					if (bInvalidateNaniteStreaming || (NaniteResources.HasStreamingData() && !IsRunningCookCommandlet()))
 					{
-						// Drop streaming data from memory when it has been successfully committed to DDC
+						// Drop streaming data from memory when it has been successfully committed to DDC or if the streaming data we have doesn't match the cache record
 						NaniteResources.DropBulkData();
+
+						if (bInvalidateNaniteStreaming)
+						{
+							UE_LOGFMT(LogStaticMesh, Display,
+								"StaticMesh Nanite streaming data saved to DDC mismatched the data that was already in the cache. Dropping data and replacing the hash. Asset: {Asset}. Cache Key: {CacheKey}",
+								UE::FAssetLog(Owner), CacheKey);
+							NaniteResources.DDCRawHash = NaniteStreamingDataHash;
+						}
+					}
+
+					if (!MeshDataBuffer.IsNull())
+					{
+						// The mesh data record in cache differs from what we tried to put, so we'll override it with what's in the cache
+						UE_LOGFMT(LogStaticMesh, Display,
+							"StaticMesh data saved to DDC mismatched the data that was already in the cache. Serializing from cached data. Asset: {Asset}. Cache Key: {CacheKey}",
+							UE::FAssetLog(Owner), CacheKey);
+
+						FMemoryReaderView ReadAr(MeshDataBuffer.GetView(), /*bIsPersistent=*/ true);
+						Serialize(ReadAr, Owner, /*bCooked=*/ false);
 					}
 				}
+
+				EstimatedCompressedSize = EstimatedMeshDataCompressedSize + EstimatedNaniteStreamingCompressedSize;
 			}
 
 			if (NaniteResources.HasStreamingData() && !bSavedToDDC)
diff --git a/Engine/Source/Runtime/Engine/Public/StaticMeshResources.h b/Engine/Source/Runtime/Engine/Public/StaticMeshResources.h
--- a/Engine/Source/Runtime/Engine/Public/StaticMeshResources.h
+++ b/Engine/Source/Runtime/Engine/Public/StaticMeshResources.h
@@ -928,6 +928,8 @@ public:
 	}
 
 private:
+	void FreeLODResources();
+
 #if WITH_EDITORONLY_DATA
 	/** Allow the editor to explicitly update section information. */
 	friend class FLevelOfDetailSettingsLayout;
