- Set the log in UInputComponent::LogDeprecatedBindingWarning to Verbose because it was spamming the log multiple times in a frame.

#Synced-CL 51105393

diff --git a/Engine/Source/Runtime/Engine/Private/Components/InputComponent.cpp b/Engine/Source/Runtime/Engine/Private/Components/InputComponent.cpp
--- a/Engine/Source/Runtime/Engine/Private/Components/InputComponent.cpp
+++ b/Engine/Source/Runtime/Engine/Private/Components/InputComponent.cpp
@@ -130,7 +130,7 @@ void UInputComponent::LogDeprecatedBindingWarning(const FName BindingName)
 	const UInputSettings* Settings = GetDefault<UInputSettings>();
 	if (Settings->bEnabledLegacyMappingDeprecationWarnings && Settings->IsActionMappingDeprecated(BindingName))
 	{
-		UE_LOG(LogPlayerInput, Warning, TEXT("Attempting to access a deprecated action/axis mapping '%s'"), *BindingName.ToString());
+		UE_LOG(LogPlayerInput, Verbose, TEXT("Attempting to access a deprecated action/axis mapping '%s'"), *BindingName.ToString());
 	}
 }
 
