Fixed restarting workers with the incorrect number of task workers if the platform has a thread config. Changed the code to reinitialize with the existing number of workers.

diff --git a/Engine/Source/Runtime/Engine/Private/UnrealEngine.cpp b/Engine/Source/Runtime/Engine/Private/UnrealEngine.cpp
--- a/Engine/Source/Runtime/Engine/Private/UnrealEngine.cpp
+++ b/Engine/Source/Runtime/Engine/Private/UnrealEngine.cpp
@@ -20127,9 +20127,11 @@ static void SetupThreadConfig(const TArray<FString>& Args)
 	// Wait until the RT is really idle
 	ReadyEvent.Wait();
 
-	int32 NumWorkerThreads = FPlatformMisc::NumberOfWorkerThreadsToSpawn();
-	int32 NumBackgroundWorkers = FMath::Max(1, NumWorkerThreads - FMath::Min<int>(2, NumWorkerThreads));
-	int32 NumForegroundWorkers = FMath::Max(1, NumWorkerThreads - NumBackgroundWorkers);
+	// Try to retain the existing number of workers
+	FTaskGraphInterface& TaskGraph = FTaskGraphInterface::Get();
+	const int32 NumWorkerThreads = FPlatformMisc::NumberOfWorkerThreadsToSpawn();
+	const int32 NumBackgroundWorkers = FMath::Max(1, TaskGraph.GetNumBackgroundThreads());
+	const int32 NumForegroundWorkers = FMath::Max(1, FMath::Max(TaskGraph.GetNumForegroundThreads(), NumWorkerThreads - NumBackgroundWorkers));
 
 	// Setting the stack size won't have an effect for the threads created before InitializeConfigSystem() is called
 	int32 TaskGraphStackSize = 0;
