Iris
* Changed NetRootObjectAdapter to track dependencies via ObjectKeys instead of NetRefHandle.
* This was made to allow actor components to remove dependencies in their EndPlay after they have stopped replication.
* The true root object (Actor!) of the dependent must be passed to RemoveRelevantWith to work as we can't ask Iris for the component's or subobject's root owner.

#rnx

#Synced-CL 51003732

diff --git a/Engine/Source/Runtime/Engine/Private/Net/Iris/ReplicationSystem/NetRootObjectAdapter.cpp b/Engine/Source/Runtime/Engine/Private/Net/Iris/ReplicationSystem/NetRootObjectAdapter.cpp
--- a/Engine/Source/Runtime/Engine/Private/Net/Iris/ReplicationSystem/NetRootObjectAdapter.cpp
+++ b/Engine/Source/Runtime/Engine/Private/Net/Iris/ReplicationSystem/NetRootObjectAdapter.cpp
@@ -286,7 +286,7 @@ void FNetRootObjectAdapter::RelevantWith(const UObject* OtherReplicatedObject)
 	ULevel* Level = WeakLevel.Get();
 	if (!Level)
 	{
-		ensureMsgf(Level, TEXT("RootObjectAdapter cannot replicated without a valid level associated to it: %s."), *GetNameSafe(WeakReplicatedObject.Get()));
+		ensureMsgf(Level, TEXT("RootObjectAdapter cannot replicate without a valid level associated to it: %s."), *GetNameSafe(WeakReplicatedObject.Get()));
 		return;
 	}
 
@@ -347,7 +347,9 @@ void FNetRootObjectAdapter::RelevantWith(const UObject* OtherReplicatedObject)
 			return;
 		}
 
-		uint32& RootObjectRefCount = DependentRootObjects.FindOrAdd(ParentHandle);
+		const FObjectKey DependentKey = Bridge->GetReplicatedObject(ParentHandle);
+
+		uint32& RootObjectRefCount = DependentRootObjects.FindOrAdd(DependentKey);
 		++RootObjectRefCount;
 
 		// Add the Iris dependency the first time we are linked with that root object
@@ -391,40 +393,48 @@ void FNetRootObjectAdapter::RemoveRelevantWith(const UObject* OtherReplicatedObj
 		return;
 	}
 
+
 	FReplicationSystemUtil::ForEachServerBridge(World, [&](UEngineReplicationBridge* Bridge)
 	{
+		FObjectKey DependentKey;
+
 		// Find the handle of the object we will be dependent to
 		const FNetRefHandle OtherRefHandle = Bridge->GetReplicatedRefHandle(OtherReplicatedObject);
-		if (!OtherRefHandle.IsValid())
+		const FNetRefHandle ParentHandle = Bridge->GetRootObjectOfAnyObject(OtherRefHandle);
+
+		if (ParentHandle.IsValid())
 		{
-			UE_LOG(LogIris, Error, TEXT("%hs ignored for: %s with: %s because the other replicated object is not replicating"), __FUNCTION__, *GetNameSafe(MyReplicatedObject), *GetNameSafe(OtherReplicatedObject));
-			ensure(false);
-			return;
+			DependentKey = Bridge->GetReplicatedObject(ParentHandle);
 		}
-
-		// Find the handle of it's root object
-		const FNetRefHandle ParentHandle = Bridge->GetRootObjectOfAnyObject(OtherRefHandle);
-		if (!ParentHandle.IsValid())
+		else
 		{
-			UE_LOG(LogIris, Error, TEXT("%hs Could not find root object handle to replicate %s with %s"), __FUNCTION__, *GetNameSafe(MyReplicatedObject), *GetNameSafe(OtherReplicatedObject));
-			ensure(false);
-			return;
+			// This usually means the other object has stopped replicating before this.
+			// This is not directly a problem because Iris will have removed the links at that time.
+			// But it also means we can't find the root object pointer via Iris
+			UE_LOG(LogIris, VeryVerbose, TEXT("%hs Could not find root object handle to replicate %s with %s"), __FUNCTION__, *GetNameSafe(MyReplicatedObject), *GetNameSafe(OtherReplicatedObject));
+
+			DependentKey = OtherReplicatedObject;
 		}
 
 		{
-			uint32* RootObjectRefCount = DependentRootObjects.Find(ParentHandle);
+			uint32* RootObjectRefCount = DependentRootObjects.Find(DependentKey);
+
+			ensureMsgf(RootObjectRefCount != nullptr, TEXT("Could not find dependent object: %s in dependency map of %s"), *GetNameSafe(DependentKey.ResolveObjectPtr()), *GetNameSafe(MyReplicatedObject));
 	
 			// Remove the Iris dependency if it's the last reference
 			if (!RootObjectRefCount || (*RootObjectRefCount) <= 1)
 			{
 				// Find our handle
 				const FNetRefHandle MyRefHandle = Bridge->GetReplicatedRefHandle(MyReplicatedObject);
-				if (ensureMsgf(MyRefHandle.IsValid(), TEXT("No netrefhandle assigned to %s"), *GetNameSafe(MyReplicatedObject)))
+				ensureMsgf(MyRefHandle.IsValid(), TEXT("No netrefhandle assigned to %s"), *GetNameSafe(MyReplicatedObject));
+
+				// Only need to remove the dependency if both objects are still replicating
+				if (MyRefHandle.IsValid() && ParentHandle.IsValid())
 				{
 					Bridge->RemoveDependentObject(ParentHandle, MyRefHandle);
 				}
 				
-				DependentRootObjects.Remove(ParentHandle);
+				DependentRootObjects.Remove(DependentKey);
 			}
 			else
 			{
diff --git a/Engine/Source/Runtime/Engine/Public/Net/Iris/ReplicationSystem/NetRootObjectAdapter.h b/Engine/Source/Runtime/Engine/Public/Net/Iris/ReplicationSystem/NetRootObjectAdapter.h
--- a/Engine/Source/Runtime/Engine/Public/Net/Iris/ReplicationSystem/NetRootObjectAdapter.h
+++ b/Engine/Source/Runtime/Engine/Public/Net/Iris/ReplicationSystem/NetRootObjectAdapter.h
@@ -2,8 +2,11 @@
 
 #pragma once
 
+#include "UObject/ObjectKey.h"
 #include "UObject/WeakObjectPtr.h"
+
 #include "Misc/NotNull.h"
+
 #include "Iris/ReplicationSystem/NetRefHandle.h"
 
 class UWorld;
@@ -130,7 +133,7 @@ public:
 private:
 
 	/** Map of root objects and how many times we are dependent to each. */
-	TMap<FNetRefHandle /* RootObject */, uint32 /* DependentCount */> DependentRootObjects;
+	TMap<FObjectKey /* RootObject */, uint32 /* DependentCount */> DependentRootObjects;
 
 	FName ExplicitNetFactoryName;	
 
