Fixed unfiltered montage notifies not being dispatched during an autonomouse pose tick.
- Added a.Notifies.ApplyUnfilteredOnAutonomousTickPose console variable

[RN] minor

#Synced-CL 51072362

diff --git a/Engine/Source/Runtime/Engine/Private/Animation/AnimInstance.cpp b/Engine/Source/Runtime/Engine/Private/Animation/AnimInstance.cpp
--- a/Engine/Source/Runtime/Engine/Private/Animation/AnimInstance.cpp
+++ b/Engine/Source/Runtime/Engine/Private/Animation/AnimInstance.cpp
@@ -158,15 +158,27 @@ enum class EAnimNotifyDispatchPolicy : int32
 	UpdatedSourcesOnly = 1,
 };
 
-int32 GAnimNotifyDispatchPolicy = static_cast<int32>(EAnimNotifyDispatchPolicy::UpdatedSourcesOnly);
-static FAutoConsoleVariableRef CVarNotifiesDispatchPolicy(
-	TEXT("a.Notifies.DispatchPolicy"),
-	GAnimNotifyDispatchPolicy,
-	TEXT("Controls notify dispatch behavior when it's source (montage, anim graph) may tick at different times within a frame.\n")
-		TEXT("  0 = AllSources: Update active notify buffers and dispatch from any source, even if that source has not updated yet this tick (previous behavior).\n")
-		TEXT("  1 = UpdatedSourcesOnly: Only dispatch from sources that have updated this tick. Anim graph notifies/states are not triggered/ticked if the graph has not updated yet, ")
-		TEXT("but active notify states from skipped sources remain valid. Useful for multi-tick scenarios such as MoveAutonomous."),
-	ECVF_Default);
+namespace AnimNotifyCVars
+{
+	static int32 bApplyUnfilteredMontageAnimNotifiesOnAutonomousTickPose = 1;
+	FAutoConsoleVariableRef CVarApplyUnfilteredMontageAnimNotifiesOnAutonomousTickPose(
+		TEXT("a.Notifies.ApplyUnfilteredOnAutonomousTickPose"),
+		bApplyUnfilteredMontageAnimNotifiesOnAutonomousTickPose,
+		TEXT("Whether to apply unfiltered montage notifies when ticking character pose during autonomous move updates. \n")
+		TEXT("0: Disable, 1: Enable"),
+		ECVF_Default);
+	
+	static int32 NotifiesDispatchPolicy = static_cast<int32>(EAnimNotifyDispatchPolicy::UpdatedSourcesOnly);
+	static FAutoConsoleVariableRef CVarNotifiesDispatchPolicy(
+		TEXT("a.Notifies.DispatchPolicy"),
+		NotifiesDispatchPolicy,
+		TEXT("Controls notify dispatch behavior when it's source (montage, anim graph) may tick at different times within a frame.\n")
+			TEXT("  0 = AllSources: Update active notify buffers and dispatch from any source, even if that source has not updated yet this tick (previous behavior).\n")
+			TEXT("  1 = UpdatedSourcesOnly: Only dispatch from sources that have updated this tick. Anim graph notifies/states are not triggered/ticked if the graph has not updated yet, ")
+			TEXT("but active notify states from skipped sources remain valid. Useful for multi-tick scenarios such as MoveAutonomous."),
+		ECVF_Default);
+}
+
 
 /////////////////////////////////////////////////////
 // UAnimInstance
@@ -921,8 +933,18 @@ void UAnimInstance::PostUpdateAnimation()
 
 void UAnimInstance::DispatchQueuedAnimEvents()
 {
+	// Acquire the proxy as we need to query notifies and delta time
+	FAnimInstanceProxy& Proxy = GetProxyOnGameThread<FAnimInstanceProxy>();
+	USkeletalMeshComponent* SkelMeshComp = GetSkelMeshComponent();
+	
+	// When ticking the character pose via the CharacterMovementComponent we need to manually apply the montage notifies after each replicated move since it can queue unfiltered montage notifies that will get skipped otherwise. 
+	if (SkelMeshComp && SkelMeshComp->bIsAutonomousTickPose && AnimNotifyCVars::bApplyUnfilteredMontageAnimNotifiesOnAutonomousTickPose)
+	{
+		NotifyQueue.ApplyMontageNotifies(Proxy);
+	}
+	
 	// now trigger Notifies
-	TriggerAnimNotifies(GetProxyOnGameThread<FAnimInstanceProxy>().GetDeltaSeconds());
+	TriggerAnimNotifies(Proxy.GetDeltaSeconds());
 
 	// Trigger Montage end events after notifies. In case Montage ending ends abilities or other states, we make sure notifies are processed before montage events.
 	TriggerQueuedMontageEvents();
@@ -952,7 +974,7 @@ void UAnimInstance::DispatchQueuedAnimEvents()
 
 	// We need to determine which anim update path's notifies have we triggered in the case we get multiple dispatch request per frame
 	// (from the pose being ticked in CharacterMovementComponent) to avoid clearing the ActiveNotify arrays after parallel anim eval is done.
-	if (USkeletalMeshComponent* SkelMeshComp = GetSkelMeshComponent())
+	if (SkelMeshComp)
 	{
 		bHasDispatchedMontageNotifiesThisTick = SkelMeshComp->bIsAutonomousTickPose && SkelMeshComp->PoseTickedThisFrame();
 		bHasDispatchedAnimGraphNotifiesThisTick = SkelMeshComp->bIsAutonomousTickPose && !NeedsUpdate() && SkelMeshComp->PoseTickedThisFrame();
@@ -1839,7 +1861,7 @@ void UAnimInstance::TriggerAnimNotifies(float DeltaSeconds, ENotifyTriggerMode M
 		NotifyQueue.AnimNotifies = MoveTemp(FilteredAnimNotifies);
 	}
 	
-	const bool bDispatchFromSourcesOnly = static_cast<EAnimNotifyDispatchPolicy>(GAnimNotifyDispatchPolicy) == EAnimNotifyDispatchPolicy::UpdatedSourcesOnly && Mode == ENotifyTriggerMode::Default;
+	const bool bDispatchFromSourcesOnly = static_cast<EAnimNotifyDispatchPolicy>(AnimNotifyCVars::NotifiesDispatchPolicy) == EAnimNotifyDispatchPolicy::UpdatedSourcesOnly && Mode == ENotifyTriggerMode::Default;
 
 	// We need to skip triggering notifies if we haven't updated the source (due to parallel update) or if we have already processed the notifies this tick (multiple anim updates/dispatches only allowed from autonomous tick/move requests). 
 	const bool bSkipNotifiesSourcedFromAnimGraph = bDispatchFromSourcesOnly && (NeedsUpdate() || (bHasDispatchedAnimGraphNotifiesThisTick && !SkelMeshComp->bIsAutonomousTickPose));
