Streaming system fixes
- FRenderAssetStreamingManager::UpdateResourceStreaming bProcessEverything path wasn't releasing async views like the incremental paths were.
- FDynamicRenderAssetInstanceManager::OnCreateViewDone wasn't resetting DirtyIndex after getting brand new, uninitialized data. This can cause the streaming calculations to run against invalid data, resulting in resources being fully loaded when they shouldn't be.

#Synced-CL 51242114

diff --git a/Engine/Source/Runtime/Engine/Private/Streaming/DynamicTextureInstanceManager.cpp b/Engine/Source/Runtime/Engine/Private/Streaming/DynamicTextureInstanceManager.cpp
--- a/Engine/Source/Runtime/Engine/Private/Streaming/DynamicTextureInstanceManager.cpp
+++ b/Engine/Source/Runtime/Engine/Private/Streaming/DynamicTextureInstanceManager.cpp
@@ -159,6 +159,9 @@ void FDynamicRenderAssetInstanceManager::OnCreateViewDone(FRenderAssetInstanceVi
 	// The incremental update will refresh the bounds within the next update loop.
 	FRenderAssetInstanceView::SwapData(InView, State);
 
+	// All of the data is dirty now since we just got uninitialized data.
+	DirtyIndex = 0;
+
 	AsyncView = TRefCountPtr<const FRenderAssetInstanceView>(InView);
 }
 
diff --git a/Engine/Source/Runtime/Engine/Private/Streaming/StreamingManagerTexture.cpp b/Engine/Source/Runtime/Engine/Private/Streaming/StreamingManagerTexture.cpp
--- a/Engine/Source/Runtime/Engine/Private/Streaming/StreamingManagerTexture.cpp
+++ b/Engine/Source/Runtime/Engine/Private/Streaming/StreamingManagerTexture.cpp
@@ -1962,6 +1962,9 @@ void FRenderAssetStreamingManager::UpdateResourceStreaming( float DeltaTime, boo
 
 		StreamRenderAssets(bProcessEverything);
 
+		// Release the old view now as the destructors can be expensive. Now only the dynamic manager holds a ref.
+		AsyncWork->GetTask().ReleaseAsyncViews();
+
 		STAT(GatheredStats.SetupAsyncTaskCycles = 0);
 		STAT(GatheredStats.UpdateStreamingDataCycles = 0);
 		STAT(GatheredStats.StreamRenderAssetsCycles = 0);
