[HWRT] Couple of new tools to identify issues with always-resident ray tracing geometries.
- Command to dump list of unreferenced always-resident ray tracing geometries to CSV.
- On-screen warning when always-resident size exceeds a percentage of the total pool size.

#Synced-CL 51150711

diff --git a/Engine/Source/Runtime/Engine/Private/Rendering/RayTracingGeometryManager.cpp b/Engine/Source/Runtime/Engine/Private/Rendering/RayTracingGeometryManager.cpp
--- a/Engine/Source/Runtime/Engine/Private/Rendering/RayTracingGeometryManager.cpp
+++ b/Engine/Source/Runtime/Engine/Private/Rendering/RayTracingGeometryManager.cpp
@@ -21,6 +21,11 @@
 
 #include "Engine/Engine.h"
 
+#if !UE_BUILD_SHIPPING
+#include "HAL/FileManager.h"
+#include "Misc/Paths.h"
+#endif // !UE_BUILD_SHIPPING
+
 #if RHI_RAYTRACING
 
 DECLARE_LOG_CATEGORY_CLASS(LogRayTracingGeometryManager, Log, All);
@@ -107,6 +112,13 @@ static FAutoConsoleVariableRef CVarRayTracingNumAlwaysResidentLODs(
 	ECVF_RenderThreadSafe
 );
 
+static TAutoConsoleVariable<int> CVarRayTracingAlwaysResidentWarningPercentage(
+	TEXT("r.RayTracing.Debug.GeometryMemoryPool.AlwaysResidentWarningPercentage"),
+	20,
+	TEXT("A warning is triggered when always-resident size exceeds this percentage of the geometry memory pool [0, 100]."),
+	ECVF_RenderThreadSafe
+);
+
 static TAutoConsoleVariable<bool> CVarRayTracingOnDemandGeometryBuffersStreaming(
 	TEXT("r.RayTracing.OnDemandGeometryBuffersStreaming"),
 	true,
@@ -148,6 +160,16 @@ static FAutoConsoleVariableRef CVarRayTracingTestCheckIntegrity(
 );
 #endif
 
+#if !UE_BUILD_SHIPPING
+static bool bDumpUnreferencedAlwaysResidentGeometriesNextFrame = false;
+
+static FAutoConsoleCommand CmdDumpUnreferencedAlwaysResidentGeometries(
+	TEXT("r.RayTracing.DumpUnreferencedAlwaysResidentGeometries"),
+	TEXT("Writes the list of always-resident ray tracing geometries not currently referenced in the ray tracing scene to a CSV file"),
+	FConsoleCommandDelegate::CreateStatic([] { bDumpUnreferencedAlwaysResidentGeometriesNextFrame = true; }),
+	ECVF_Default);
+#endif // !UE_BUILD_SHIPPING
+
 DECLARE_STATS_GROUP(TEXT("Ray Tracing Geometry"), STATGROUP_RayTracingGeometry, STATCAT_Advanced);
 
 DECLARE_DWORD_ACCUMULATOR_STAT(TEXT("Geometry Count"), STAT_RayTracingGeometryCount, STATGROUP_RayTracingGeometry);
@@ -932,14 +954,17 @@ void FRayTracingGeometryManager::UpdateReferenceBasedResidency(FRHICommandList&
 #if !UE_BUILD_SHIPPING
 	if (GRayTracingShowOnScreenWarnings)
 	{
-		if (TotalAlwaysResidentSize > ResidentGeometryMemoryPoolSize)
+		const int ThresholdPercentage = FMath::Clamp(CVarRayTracingAlwaysResidentWarningPercentage.GetValueOnRenderThread(), 0, 100);
+		if (TotalAlwaysResidentSize > ResidentGeometryMemoryPoolSize * ThresholdPercentage / 100.0f)
 		{
 			GEngine->AddOnScreenDebugMessage(uint64(this), 1.f, FColor::Red,
-				*FString::Printf(TEXT("RAY TRACING GEOMETRY - ALWAYS RESIDENT MEMORY OVER BUDGET %s / %s"),
+				*FString::Printf(TEXT("RAY TRACING GEOMETRY - ALWAYS RESIDENT MEMORY EXCEEDS %d%% OF THE BUDGET (%s / %s)"),
+					ThresholdPercentage,
 					*FText::AsMemory(TotalAlwaysResidentSize).ToString(),
 					*FText::AsMemory(ResidentGeometryMemoryPoolSize).ToString()));
 		}
-		else if (RequestedSize > ResidentGeometryMemoryPoolSize)
+
+		if (RequestedSize > ResidentGeometryMemoryPoolSize)
 		{
 			GEngine->AddOnScreenDebugMessage(uint64(this), 1.f, FColor::Yellow,
 				*FString::Printf(TEXT("RAY TRACING GEOMETRY - REQUESTED MEMORY OVER BUDGET %s / %s"),
@@ -947,6 +972,13 @@ void FRayTracingGeometryManager::UpdateReferenceBasedResidency(FRHICommandList&
 					*FText::AsMemory(ResidentGeometryMemoryPoolSize).ToString()));
 		}
 	}
+
+	if (bDumpUnreferencedAlwaysResidentGeometriesNextFrame)
+	{
+		bDumpUnreferencedAlwaysResidentGeometriesNextFrame = false;
+
+		DumpUnreferencedAlwaysResidentRayTracingGeometries();
+	}
 #endif
 }
 
@@ -1649,6 +1681,83 @@ bool FRayTracingGeometryManager::IsGeometryGroupReferenced(RayTracing::FGeometry
 {
 	return ReferencedGeometryGroups.Contains(GeometryGroup);
 }
-#endif
+#endif // DO_CHECK
+
+#if !UE_BUILD_SHIPPING
+bool FRayTracingGeometryManager::DumpRayTracingGeometries(TArrayView<RayTracing::FGeometryHandle> InGeometries, const FString& Filename)
+{
+	Algo::Sort(InGeometries, [this](RayTracing::FGeometryHandle& LHSHandle, RayTracing::FGeometryHandle& RHSHandle)
+		{
+			FRegisteredGeometry& LHS = GetRegisteredGeometry(LHSHandle);
+			FRegisteredGeometry& RHS = GetRegisteredGeometry(RHSHandle);
+
+			return LHS.Size > RHS.Size;
+		});
+
+	TUniquePtr<FArchive> CSVFile = TUniquePtr<FArchive>(IFileManager::Get().CreateFileWriter(*Filename, FILEWRITE_AllowRead));
+	if (!CSVFile)
+	{
+		return false;
+	}
+
+	const TCHAR* Header = TEXT("Name,Size (MBs),Prims,Segments,Update\n");
+	CSVFile->Serialize(TCHAR_TO_ANSI(Header), FPlatformString::Strlen(Header));
+
+	for (RayTracing::FGeometryHandle GeometryHandle : InGeometries)
+	{
+		FRegisteredGeometry& RegisteredGeometry = GetRegisteredGeometry(GeometryHandle);
+		FRayTracingGeometryInitializer Initializer = RegisteredGeometry.Geometry->GetInitializer();
+
+		const FString Row = FString::Printf(TEXT("%s,%.3f,%d,%d,%d\n"),
+			!Initializer.DebugName.IsNone() ? *Initializer.DebugName.ToString() : TEXT("*UNKNOWN*"),
+			RegisteredGeometry.Size / double(1 << 20),
+			Initializer.TotalPrimitiveCount,
+			Initializer.Segments.Num(),
+			(int32)Initializer.bAllowUpdate);
+		CSVFile->Serialize(TCHAR_TO_ANSI(*Row), Row.Len());
+	}
+
+	return true;
+}
+
+void FRayTracingGeometryManager::DumpUnreferencedAlwaysResidentRayTracingGeometries()
+{
+	TArray<RayTracing::FGeometryHandle> UnreferencedAlwaysResidentGeometries;
+
+	for (RayTracing::FGeometryHandle GeometryHandle : AlwaysResidentGeometries)
+	{
+		const FRegisteredGeometry& RegisteredGeometry = GetRegisteredGeometry(GeometryHandle);
+		const RayTracing::FGeometryGroupHandle GroupHandle = RegisteredGeometry.Geometry->GroupHandle;
+
+		bool bIsReferenced = false;
+
+		if (ReferencedGeometries.Contains(GeometryHandle))
+		{
+			bIsReferenced = true;
+		}
+
+		if (!GroupHandle.IsNull() && ReferencedGeometryGroups.Contains(GroupHandle))
+		{
+			bIsReferenced = true;
+		}
+
+		if (!bIsReferenced)
+		{
+			UnreferencedAlwaysResidentGeometries.Add(GeometryHandle);
+		}
+	}
+
+	FString Filename = FPaths::ProfilingDir() + FString::Printf(TEXT("UnreferencedAlwaysResidentGeometries(%s).csv"), *FDateTime::Now().ToString(TEXT("%Y%m%d_%H%M%S")));
+
+	if (DumpRayTracingGeometries(UnreferencedAlwaysResidentGeometries, Filename))
+	{
+		UE_LOGF(LogRayTracingGeometryManager, Log, "List of unreferenced always-resident ray tracing geometries written to file '%ls'", *Filename);
+	}
+	else
+	{
+		UE_LOGF(LogRayTracingGeometryManager, Log, "Unable to write list of unreferenced always-resident ray tracing geometries to file '%ls'", *Filename);
+	}
+}
+#endif // !UE_BUILD_SHIPPING
 
 #endif // RHI_RAYTRACING
diff --git a/Engine/Source/Runtime/Engine/Public/Rendering/RayTracingGeometryManager.h b/Engine/Source/Runtime/Engine/Public/Rendering/RayTracingGeometryManager.h
--- a/Engine/Source/Runtime/Engine/Public/Rendering/RayTracingGeometryManager.h
+++ b/Engine/Source/Runtime/Engine/Public/Rendering/RayTracingGeometryManager.h
@@ -190,7 +190,12 @@ private:
 
 #if DO_CHECK
 	void CheckIntegrity();
-#endif
+#endif // DO_CHECK
+
+#if !UE_BUILD_SHIPPING
+	bool DumpRayTracingGeometries(TArrayView<RayTracing::FGeometryHandle> Geometries, const FString& Filename);
+	void DumpUnreferencedAlwaysResidentRayTracingGeometries();
+#endif // !UE_BUILD_SHIPPING
 
 	FBuildRequest& GetBuildRequest(RayTracing::FBuildRequestHandle Handle);
 	FRegisteredGeometry& GetRegisteredGeometry(RayTracing::FGeometryHandle Handle);
