Scene Capture:  fix for sky light not working -- need to set sky light scale to one in the constructor (used internally for runtime reflection captures, not a serialized UPROPERTY, so no versioning needed).

#rnx

#Synced-CL 51249996

diff --git a/Engine/Source/Runtime/Engine/Private/Components/SceneCaptureComponent.cpp b/Engine/Source/Runtime/Engine/Private/Components/SceneCaptureComponent.cpp
--- a/Engine/Source/Runtime/Engine/Private/Components/SceneCaptureComponent.cpp
+++ b/Engine/Source/Runtime/Engine/Private/Components/SceneCaptureComponent.cpp
@@ -180,6 +180,7 @@ USceneCaptureComponent::USceneCaptureComponent(const FObjectInitializer& ObjectI
 	MaxViewDistanceOverride = -1;
 	CaptureSortPriority = 0;
 	bUseRayTracingIfEnabled = 0;
+	SkylightScale = FLinearColor::White;
 
 	// Disable features that are not desired when capturing the scene
 	ShowFlags.SetMotionBlur(0); // motion blur doesn't work correctly with scene captures.
