Added virtual Merge function to UAssetUserData.

Per type, give the possibility two merge Asset User Data. Useful when merging two SKM.

#Synced-CL 51099946

diff --git a/Engine/Source/Runtime/Engine/Classes/Engine/AssetUserData.h b/Engine/Source/Runtime/Engine/Classes/Engine/AssetUserData.h
--- a/Engine/Source/Runtime/Engine/Classes/Engine/AssetUserData.h
+++ b/Engine/Source/Runtime/Engine/Classes/Engine/AssetUserData.h
@@ -30,4 +30,12 @@ class UAssetUserData
 	virtual void PostEditChangeOwner() {}
 
 	virtual void PostEditChangeOwner(const FPropertyChangedEvent& PropertyChangedEvent) { PRAGMA_DISABLE_DEPRECATION_WARNINGS PostEditChangeOwner(); PRAGMA_ENABLE_DEPRECATION_WARNINGS }
+	
+	/** Merge this Asset User Data with the given one. Used when merging Skeletal Meshes.
+	* 
+	* @return true if the merge has completed successfully. */
+	virtual bool Merge(const TNotNull<UAssetUserData*> Other)
+	{
+		return false;
+	}
 };
