Deprecate PIEPreviewDeviceProfileSelector
Remove all dependencies to the PIEPreviewDeviceProfileSelector from the Engine

#Synced-CL 51256620

diff --git a/Engine/Source/Developer/Android/AndroidDeviceDetection/AndroidDeviceDetection.Build.cs b/Engine/Source/Developer/Android/AndroidDeviceDetection/AndroidDeviceDetection.Build.cs
--- a/Engine/Source/Developer/Android/AndroidDeviceDetection/AndroidDeviceDetection.Build.cs
+++ b/Engine/Source/Developer/Android/AndroidDeviceDetection/AndroidDeviceDetection.Build.cs
@@ -38,7 +38,6 @@ public class AndroidDeviceDetection : ModuleRules
 			PrivateDependencyModuleNames.AddRange(
 				new string[]
 				{
-					"PIEPreviewDeviceProfileSelector",
 					"DesktopPlatform",
 					"SlateCore",
 					"Slate",
diff --git a/Engine/Source/Developer/Android/AndroidDeviceDetection/Private/AndroidDeviceDetectionModule.cpp b/Engine/Source/Developer/Android/AndroidDeviceDetection/Private/AndroidDeviceDetectionModule.cpp
--- a/Engine/Source/Developer/Android/AndroidDeviceDetection/Private/AndroidDeviceDetectionModule.cpp
+++ b/Engine/Source/Developer/Android/AndroidDeviceDetection/Private/AndroidDeviceDetectionModule.cpp
@@ -40,7 +40,6 @@
 #include "Internationalization/Regex.h"
 
 #if WITH_EDITOR
-#include "PIEPreviewDeviceProfileSelectorModule.h"
 #include "IDesktopPlatform.h"
 #include "DesktopPlatformModule.h"
 #include "Framework/Application/SlateApplication.h"
@@ -907,30 +906,10 @@ public:
 		// create and fire off our device detection thread
 		DetectionThreadRunnable = new FAndroidDeviceDetectionRunnable(DeviceMap, &DeviceMapLock, &ADBPathCheckLock);
 		DetectionThread = FRunnableThread::Create(DetectionThreadRunnable, TEXT("FAndroidDeviceDetectionRunnable"));
-
-#if WITH_EDITOR
-		// add some menu options just for Android
-		FPIEPreviewDeviceModule* PIEPreviewDeviceModule = FModuleManager::LoadModulePtr<FPIEPreviewDeviceModule>(TEXT("PIEPreviewDeviceProfileSelector"));
-		PIEPreviewDeviceModule->AddToDevicePreviewMenuDelegates.AddLambda([this](const FText& CategoryName, class FMenuBuilder& MenuBuilder)
-			{
-				if (CategoryName.CompareToCaseIgnored(FText::FromString(TEXT("Android"))) == 0)
-				{
-					CreatePIEPreviewMenu(MenuBuilder);
-				}
-			});
-#endif
 	}
 
 	virtual ~FAndroidDeviceDetection()
 	{
-#if WITH_EDITOR
-		FPIEPreviewDeviceModule* PIEPreviewDeviceModule = FModuleManager::GetModulePtr<FPIEPreviewDeviceModule>(TEXT("PIEPreviewDeviceProfileSelector"));
-		if (PIEPreviewDeviceModule != nullptr)
-		{
-			PIEPreviewDeviceModule->AddToDevicePreviewMenuDelegates.Remove(DelegateHandle);
-		}
-#endif
-
 		if (DetectionThreadRunnable && DetectionThread)
 		{
 			DetectionThreadRunnable->Stop();
diff --git a/Engine/Source/Editor/PIEPreviewDeviceProfileSelector/PIEPreviewDeviceProfileSelector.Build.cs b/Engine/Source/Editor/PIEPreviewDeviceProfileSelector/PIEPreviewDeviceProfileSelector.Build.cs
--- a/Engine/Source/Editor/PIEPreviewDeviceProfileSelector/PIEPreviewDeviceProfileSelector.Build.cs
+++ b/Engine/Source/Editor/PIEPreviewDeviceProfileSelector/PIEPreviewDeviceProfileSelector.Build.cs
@@ -33,14 +33,6 @@ namespace UnrealBuildTool.Rules
 					"UnrealEd",
                 }
 			);
-
-			CircularlyReferencedDependentModules.AddRange(
-				new string[]
-				{
-					"MaterialShaderQualitySettings",
-					"UnrealEd"
-				}
-			);
 		}
 	}
 }
diff --git a/Engine/Source/Editor/PIEPreviewDeviceProfileSelector/Private/PIEPreviewDevice.cpp b/Engine/Source/Editor/PIEPreviewDeviceProfileSelector/Private/PIEPreviewDevice.cpp
--- a/Engine/Source/Editor/PIEPreviewDeviceProfileSelector/Private/PIEPreviewDevice.cpp
+++ b/Engine/Source/Editor/PIEPreviewDeviceProfileSelector/Private/PIEPreviewDevice.cpp
@@ -18,6 +18,8 @@
 #include "SEditorViewport.h"
 #include "Widgets/SWindow.h"
 
+PRAGMA_DISABLE_DEPRECATION_WARNINGS
+
 static const int32 g_JSON_VALUE_NOT_SET = 0;
 
 FPIEPreviewDevice::FPIEPreviewDevice()
@@ -515,3 +517,5 @@ int32 FPIEPreviewDevice::GetWindowClientHeight() const
 {
 	return WindowHeight - FMath::TruncToInt32((float)WindowTitleBarSize * DPIScaleFactor);
 }
+
+PRAGMA_ENABLE_DEPRECATION_WARNINGS
\ No newline at end of file
diff --git a/Engine/Source/Editor/PIEPreviewDeviceProfileSelector/Private/PIEPreviewDevice.h b/Engine/Source/Editor/PIEPreviewDeviceProfileSelector/Private/PIEPreviewDevice.h
--- a/Engine/Source/Editor/PIEPreviewDeviceProfileSelector/Private/PIEPreviewDevice.h
+++ b/Engine/Source/Editor/PIEPreviewDeviceProfileSelector/Private/PIEPreviewDevice.h
@@ -11,7 +11,8 @@ namespace ERHIFeatureLevel { enum Type : int; }
 /**
  * FPIEPreviewDevice class stores specific device settings (FPIEPreviewDeviceSpecifications) and, on request, applies them to the Unreal runtime system
 */
-class FPIEPreviewDevice
+PRAGMA_DISABLE_DEPRECATION_WARNINGS
+class UE_DEPRECATED(5.8, "PIEPreviewDeviceProfileSelector is deprecated and will be removed. Please use the new Preview Json System to preview Devices") FPIEPreviewDevice
 {
 private:
 	TSharedPtr<FPIEPreviewDeviceSpecifications> DeviceSpecs;
@@ -203,4 +204,5 @@ FORCEINLINE void FPIEPreviewDevice::SetIgnoreMobileContentScaleFactor(bool bIgno
 FORCEINLINE bool FPIEPreviewDevice::GetIgnoreMobileContentScaleFactor() const
 {
 	return bIgnoreContentScaleFactor;
-}
\ No newline at end of file
+}
+PRAGMA_ENABLE_DEPRECATION_WARNINGS
\ No newline at end of file
diff --git a/Engine/Source/Editor/PIEPreviewDeviceProfileSelector/Private/PIEPreviewDeviceEnumeration.cpp b/Engine/Source/Editor/PIEPreviewDeviceProfileSelector/Private/PIEPreviewDeviceEnumeration.cpp
--- a/Engine/Source/Editor/PIEPreviewDeviceProfileSelector/Private/PIEPreviewDeviceEnumeration.cpp
+++ b/Engine/Source/Editor/PIEPreviewDeviceProfileSelector/Private/PIEPreviewDeviceEnumeration.cpp
@@ -7,6 +7,8 @@
 #include "Internationalization/Internationalization.h"
 #include "Misc/Paths.h"
 
+PRAGMA_DISABLE_DEPRECATION_WARNINGS
+
 #define LOCTEXT_NAMESPACE "PIEPreviewDevice"
 
 void FPIEPreviewDeviceContainer::EnumerateDeviceSpecifications(const FString& RootDir)
@@ -95,3 +97,4 @@ FText FPIEPreviewDeviceContainerCategory::GetCategoryToolTip()
 }
 
 #undef LOCTEXT_NAMESPACE
+PRAGMA_ENABLE_DEPRECATION_WARNINGS
\ No newline at end of file
diff --git a/Engine/Source/Editor/PIEPreviewDeviceProfileSelector/Private/PIEPreviewDeviceProfileSelectorModule.cpp b/Engine/Source/Editor/PIEPreviewDeviceProfileSelector/Private/PIEPreviewDeviceProfileSelectorModule.cpp
--- a/Engine/Source/Editor/PIEPreviewDeviceProfileSelector/Private/PIEPreviewDeviceProfileSelectorModule.cpp
+++ b/Engine/Source/Editor/PIEPreviewDeviceProfileSelector/Private/PIEPreviewDeviceProfileSelectorModule.cpp
@@ -25,6 +25,8 @@
 #include "PIEPreviewWindowCoreStyle.h"
 #include "Misc/CommandLine.h"
 
+PRAGMA_DISABLE_DEPRECATION_WARNINGS
+
 DECLARE_LOG_CATEGORY_EXTERN(LogPIEPreviewDevice, Log, All); 
 DEFINE_LOG_CATEGORY(LogPIEPreviewDevice);
 IMPLEMENT_MODULE(FPIEPreviewDeviceModule, PIEPreviewDeviceProfileSelector);
@@ -311,4 +313,5 @@ void FPIEPreviewDeviceModule::SetGameLayerManagerWidget(TSharedPtr<class SGameLa
 	{
 		WindowPtr->SetGameLayerManagerWidget(GameLayerManager);
 	}
-}
\ No newline at end of file
+}
+PRAGMA_ENABLE_DEPRECATION_WARNINGS
\ No newline at end of file
diff --git a/Engine/Source/Editor/PIEPreviewDeviceProfileSelector/Private/PIEPreviewSettings.h b/Engine/Source/Editor/PIEPreviewDeviceProfileSelector/Private/PIEPreviewSettings.h
--- a/Engine/Source/Editor/PIEPreviewDeviceProfileSelector/Private/PIEPreviewSettings.h
+++ b/Engine/Source/Editor/PIEPreviewDeviceProfileSelector/Private/PIEPreviewSettings.h
@@ -13,8 +13,9 @@
 
 #include "PIEPreviewSettings.generated.h"
 
+PRAGMA_DISABLE_DEPRECATION_WARNINGS
 UCLASS(MinimalAPI, hidecategories = Object, config = PIEPreviewSettings)
-class UPIEPreviewSettings : public UObject
+class UE_DEPRECATED(5.8, "PIEPreviewDeviceProfileSelector is deprecated and will be removed. Please use the new Preview Json System to preview Devices") UPIEPreviewSettings : public UObject
 {
 	GENERATED_UCLASS_BODY()
 
@@ -28,3 +29,4 @@ public:
 	UPROPERTY(config)
 	float WindowScalingFactor;
 };
+PRAGMA_ENABLE_DEPRECATION_WARNINGS
\ No newline at end of file
diff --git a/Engine/Source/Editor/PIEPreviewDeviceProfileSelector/Private/PIEPreviewWindow.cpp b/Engine/Source/Editor/PIEPreviewDeviceProfileSelector/Private/PIEPreviewWindow.cpp
--- a/Engine/Source/Editor/PIEPreviewDeviceProfileSelector/Private/PIEPreviewWindow.cpp
+++ b/Engine/Source/Editor/PIEPreviewDeviceProfileSelector/Private/PIEPreviewWindow.cpp
@@ -23,6 +23,7 @@
 #include "PIEPreviewDeviceSpecification.h"
 #include "PIEPreviewWindowTitleBar.h"
 
+PRAGMA_DISABLE_DEPRECATION_WARNINGS
 //***********************************************************************************
 //SPIEPreviewWindow Implementation
 //***********************************************************************************
@@ -642,4 +643,5 @@ void SPIEPreviewWindow::UpdateGameLayerManagerDefaultViewport()
 	}
 }
 
+PRAGMA_ENABLE_DEPRECATION_WARNINGS
 #endif
\ No newline at end of file
diff --git a/Engine/Source/Editor/PIEPreviewDeviceProfileSelector/Private/PIEPreviewWindow.h b/Engine/Source/Editor/PIEPreviewDeviceProfileSelector/Private/PIEPreviewWindow.h
--- a/Engine/Source/Editor/PIEPreviewDeviceProfileSelector/Private/PIEPreviewWindow.h
+++ b/Engine/Source/Editor/PIEPreviewDeviceProfileSelector/Private/PIEPreviewWindow.h
@@ -7,9 +7,10 @@
 #include "Widgets/SWindow.h"
 #include "PIEPreviewDevice.h"
 
+PRAGMA_DISABLE_DEPRECATION_WARNINGS
 class SPIEPreviewWindowTitleBar;
 
-class SPIEPreviewWindow
+class UE_DEPRECATED(5.8, "PIEPreviewDeviceProfileSelector is deprecated and will be removed. Please use the new Preview Json System to preview Devices") SPIEPreviewWindow
 	: public SWindow
 {
 public:
@@ -178,5 +179,5 @@ FORCEINLINE float SPIEPreviewWindow::GetScaleToDeviceSizeFactor() const
 {
 	return 0.0f;
 }
-
+PRAGMA_ENABLE_DEPRECATION_WARNINGS
 #endif
diff --git a/Engine/Source/Editor/PIEPreviewDeviceProfileSelector/Private/PIEPreviewWindowCoreStyle.cpp b/Engine/Source/Editor/PIEPreviewDeviceProfileSelector/Private/PIEPreviewWindowCoreStyle.cpp
--- a/Engine/Source/Editor/PIEPreviewDeviceProfileSelector/Private/PIEPreviewWindowCoreStyle.cpp
+++ b/Engine/Source/Editor/PIEPreviewDeviceProfileSelector/Private/PIEPreviewWindowCoreStyle.cpp
@@ -15,6 +15,8 @@
 
 #if WITH_EDITOR
 
+PRAGMA_DISABLE_DEPRECATION_WARNINGS
+
 TSharedPtr< ISlateStyle > FPIEPreviewWindowCoreStyle::Instance = nullptr;
 
 #define RootToContentDir Style->RootToContentDir
@@ -58,4 +60,6 @@ void FPIEPreviewWindowCoreStyle::InitializePIECoreStyle()
 
 #undef RootToContentDir
 
+PRAGMA_ENABLE_DEPRECATION_WARNINGS
+
 #endif
\ No newline at end of file
diff --git a/Engine/Source/Editor/PIEPreviewDeviceProfileSelector/Private/PIEPreviewWindowCoreStyle.h b/Engine/Source/Editor/PIEPreviewDeviceProfileSelector/Private/PIEPreviewWindowCoreStyle.h
--- a/Engine/Source/Editor/PIEPreviewDeviceProfileSelector/Private/PIEPreviewWindowCoreStyle.h
+++ b/Engine/Source/Editor/PIEPreviewDeviceProfileSelector/Private/PIEPreviewWindowCoreStyle.h
@@ -6,7 +6,7 @@
 
 #include "Styling/SlateStyle.h"
 
-class FPIEPreviewWindowCoreStyle
+class UE_DEPRECATED(5.8, "PIEPreviewDeviceProfileSelector is deprecated and will be removed. Please use the new Preview Json System to preview Devices") FPIEPreviewWindowCoreStyle
 {
 
 public:
@@ -15,12 +15,13 @@ public:
 	
 	static PIEPREVIEWDEVICEPROFILESELECTOR_API void InitializePIECoreStyle();
 	
+	PRAGMA_DISABLE_DEPRECATION_WARNINGS
 	/** @return the singleton instance. */
 	static const ISlateStyle& Get()
 	{
 		return *(Instance.Get());
 	}
-
+	PRAGMA_ENABLE_DEPRECATION_WARNINGS
 private:
 		
 	static PIEPREVIEWDEVICEPROFILESELECTOR_API TSharedPtr< class ISlateStyle > Instance;
diff --git a/Engine/Source/Editor/PIEPreviewDeviceProfileSelector/Private/PIEPreviewWindowStyle.cpp b/Engine/Source/Editor/PIEPreviewDeviceProfileSelector/Private/PIEPreviewWindowStyle.cpp
--- a/Engine/Source/Editor/PIEPreviewDeviceProfileSelector/Private/PIEPreviewWindowStyle.cpp
+++ b/Engine/Source/Editor/PIEPreviewDeviceProfileSelector/Private/PIEPreviewWindowStyle.cpp
@@ -5,7 +5,7 @@
 #include UE_INLINE_GENERATED_CPP_BY_NAME(PIEPreviewWindowStyle)
 
 #if WITH_EDITOR
-
+PRAGMA_DISABLE_DEPRECATION_WARNINGS
 FPIEPreviewWindowStyle::FPIEPreviewWindowStyle()
 {
 }
@@ -25,4 +25,5 @@ const FPIEPreviewWindowStyle& FPIEPreviewWindowStyle::GetDefault()
 	static FPIEPreviewWindowStyle Default;
 	return Default;
 }
+PRAGMA_ENABLE_DEPRECATION_WARNINGS
 #endif
diff --git a/Engine/Source/Editor/PIEPreviewDeviceProfileSelector/Private/PIEPreviewWindowStyle.h b/Engine/Source/Editor/PIEPreviewDeviceProfileSelector/Private/PIEPreviewWindowStyle.h
--- a/Engine/Source/Editor/PIEPreviewDeviceProfileSelector/Private/PIEPreviewWindowStyle.h
+++ b/Engine/Source/Editor/PIEPreviewDeviceProfileSelector/Private/PIEPreviewWindowStyle.h
@@ -10,7 +10,7 @@
 * Represents the appearance of an SPIEWindow
 */
 USTRUCT()
-struct FPIEPreviewWindowStyle : public FSlateWidgetStyle
+struct UE_DEPRECATED(5.8, "PIEPreviewDeviceProfileSelector is deprecated and will be removed. Please use the new Preview Json System to preview Devices") FPIEPreviewWindowStyle : public FSlateWidgetStyle
 {
 	GENERATED_USTRUCT_BODY()
 
diff --git a/Engine/Source/Editor/PIEPreviewDeviceProfileSelector/Private/PIEPreviewWindowTitleBar.cpp b/Engine/Source/Editor/PIEPreviewDeviceProfileSelector/Private/PIEPreviewWindowTitleBar.cpp
--- a/Engine/Source/Editor/PIEPreviewDeviceProfileSelector/Private/PIEPreviewWindowTitleBar.cpp
+++ b/Engine/Source/Editor/PIEPreviewDeviceProfileSelector/Private/PIEPreviewWindowTitleBar.cpp
@@ -7,7 +7,7 @@
 #include "SViewportToolBar.h"
 
 #if WITH_EDITOR
-
+PRAGMA_DISABLE_DEPRECATION_WARNINGS
 /** Toolbar class used to add some menus to configure various device display settings */
 class SPIEToolbar : public SViewportToolBar
 {
@@ -197,5 +197,5 @@ const FSlateBrush* SPIEPreviewWindowTitleBar::GetScreenRotationButtonImage() con
 		return &FPIEPreviewWindowCoreStyle::Get().GetWidgetStyle<FPIEPreviewWindowStyle>("PIEWindow").ScreenRotationButtonStyle.Normal;
 	}
 }
-
+PRAGMA_ENABLE_DEPRECATION_WARNINGS
 #endif
diff --git a/Engine/Source/Editor/PIEPreviewDeviceProfileSelector/Private/PIEPreviewWindowTitleBar.h b/Engine/Source/Editor/PIEPreviewDeviceProfileSelector/Private/PIEPreviewWindowTitleBar.h
--- a/Engine/Source/Editor/PIEPreviewDeviceProfileSelector/Private/PIEPreviewWindowTitleBar.h
+++ b/Engine/Source/Editor/PIEPreviewDeviceProfileSelector/Private/PIEPreviewWindowTitleBar.h
@@ -12,7 +12,8 @@
 /**
 * Implements a window PIE title bar widget.
 */
-class SPIEPreviewWindowTitleBar
+PRAGMA_DISABLE_DEPRECATION_WARNINGS
+class UE_DEPRECATED(5.8, "PIEPreviewDeviceProfileSelector is deprecated and will be removed. Please use the new Preview Json System to preview Devices") SPIEPreviewWindowTitleBar
 	: public SWindowTitleBar
 {
 private:
@@ -40,5 +41,5 @@ FORCEINLINE TSharedPtr<class SPIEPreviewWindow> SPIEPreviewWindowTitleBar::GetOw
 
 	return PIEWindow;
 }
-
+PRAGMA_ENABLE_DEPRECATION_WARNINGS
 #endif
diff --git a/Engine/Source/Editor/PIEPreviewDeviceProfileSelector/Public/PIEPreviewDeviceEnumeration.h b/Engine/Source/Editor/PIEPreviewDeviceProfileSelector/Public/PIEPreviewDeviceEnumeration.h
--- a/Engine/Source/Editor/PIEPreviewDeviceProfileSelector/Public/PIEPreviewDeviceEnumeration.h
+++ b/Engine/Source/Editor/PIEPreviewDeviceProfileSelector/Public/PIEPreviewDeviceEnumeration.h
@@ -14,8 +14,8 @@
 /**
 * 
 */
-
-class FPIEPreviewDeviceContainerCategory
+PRAGMA_DISABLE_DEPRECATION_WARNINGS
+class UE_DEPRECATED(5.8, "PIEPreviewDeviceProfileSelector is deprecated and will be removed. Please use the new Preview Json System to preview Devices") FPIEPreviewDeviceContainerCategory
 {
 public:
 	FPIEPreviewDeviceContainerCategory(const FString& InSubDirectoryPath, FText DisplayName) : 
@@ -40,7 +40,7 @@ protected:
 	friend class FPIEPreviewDeviceContainer;
 };
 
-class FPIEPreviewDeviceContainer
+class UE_DEPRECATED(5.8, "PIEPreviewDeviceProfileSelector is deprecated and will be removed. Please use the new Preview Json System to preview Devices") FPIEPreviewDeviceContainer
 {
 public:
 	// Recursively iterate through 'RootDir' for all device json files.
@@ -67,5 +67,5 @@ private:
 	UE_API void EnumerateDeviceSpecifications(TSharedPtr<FPIEPreviewDeviceContainerCategory> SubCategory);
 	UE_API void EnumerateDeviceLocalizedNames();
 };
-
+PRAGMA_ENABLE_DEPRECATION_WARNINGS
 #undef UE_API
diff --git a/Engine/Source/Editor/PIEPreviewDeviceProfileSelector/Public/PIEPreviewDeviceProfileSelectorModule.h b/Engine/Source/Editor/PIEPreviewDeviceProfileSelector/Public/PIEPreviewDeviceProfileSelectorModule.h
--- a/Engine/Source/Editor/PIEPreviewDeviceProfileSelector/Public/PIEPreviewDeviceProfileSelectorModule.h
+++ b/Engine/Source/Editor/PIEPreviewDeviceProfileSelector/Public/PIEPreviewDeviceProfileSelectorModule.h
@@ -12,8 +12,8 @@
 /**
 * Implements the Preview Device Profile Selector module.
 */
-
-class FPIEPreviewDeviceModule
+PRAGMA_DISABLE_DEPRECATION_WARNINGS
+class UE_DEPRECATED(5.8, "PIEPreviewDeviceProfileSelector is deprecated and will be removed. Please use the new Preview Json System to preview Devices") FPIEPreviewDeviceModule
 	: public IPIEPreviewDeviceModule
 {
 public:
@@ -107,3 +107,4 @@ private:
 
 	TSharedPtr<class SGameLayerManager> GameLayerManagerWidget;
 };
+PRAGMA_ENABLE_DEPRECATION_WARNINGS
\ No newline at end of file
diff --git a/Engine/Source/Editor/PIEPreviewDeviceSpecification/PIEPreviewDeviceSpecification.Build.cs b/Engine/Source/Editor/PIEPreviewDeviceSpecification/PIEPreviewDeviceSpecification.Build.cs
--- a/Engine/Source/Editor/PIEPreviewDeviceSpecification/PIEPreviewDeviceSpecification.Build.cs
+++ b/Engine/Source/Editor/PIEPreviewDeviceSpecification/PIEPreviewDeviceSpecification.Build.cs
@@ -7,14 +7,12 @@ namespace UnrealBuildTool.Rules
 	{
         public PIEPreviewDeviceSpecification(ReadOnlyTargetRules Target) : base(Target)
 		{
-			PrivateIncludePathModuleNames.Add("Engine");
-
 			PublicDependencyModuleNames.AddRange(
 				new string[]
 				{
 					"Core",
-				    "CoreUObject",
-                }
+					"CoreUObject",
+				}
 				);
 		}
 	}
diff --git a/Engine/Source/Editor/PIEPreviewDeviceSpecification/Public/PIEPreviewDeviceSpecification.h b/Engine/Source/Editor/PIEPreviewDeviceSpecification/Public/PIEPreviewDeviceSpecification.h
--- a/Engine/Source/Editor/PIEPreviewDeviceSpecification/Public/PIEPreviewDeviceSpecification.h
+++ b/Engine/Source/Editor/PIEPreviewDeviceSpecification/Public/PIEPreviewDeviceSpecification.h
@@ -5,10 +5,6 @@
 #include "Containers/Array.h"
 #include "Containers/UnrealString.h"
 #include "CoreMinimal.h"
-#include "HAL/Platform.h"
-#include "UObject/Object.h"
-#include "UObject/ObjectMacros.h"
-#include "UObject/UObjectGlobals.h"
 
 #include "PIEPreviewDeviceSpecification.generated.h"
 
@@ -29,7 +25,7 @@ enum class EPIEPreviewDeviceType : uint8
 };
 
 UCLASS(MinimalAPI)
-class UPIEPreviewDeviceSpecification : public UObject
+class UE_DEPRECATED(5.8, "PIEPreviewDeviceProfileSelector is deprecated and will be removed. Please use the new Preview Json System to preview Devices") UPIEPreviewDeviceSpecification : public UObject
 {
 	GENERATED_UCLASS_BODY()
 public:
@@ -152,7 +148,7 @@ public:
 
 
 USTRUCT()
-struct FPIEPreviewDeviceBezelViewportRect
+struct UE_DEPRECATED(5.8, "PIEPreviewDeviceProfileSelector is deprecated and will be removed. Please use the new Preview Json System to preview Devices") FPIEPreviewDeviceBezelViewportRect
 {
 public:
 	GENERATED_USTRUCT_BODY()
@@ -168,14 +164,17 @@ public:
 }; 
 
 USTRUCT()
-struct FPIEBezelProperties
+struct UE_DEPRECATED(5.8, "PIEPreviewDeviceProfileSelector is deprecated and will be removed. Please use the new Preview Json System to preview Devices") FPIEBezelProperties
 {
 	GENERATED_USTRUCT_BODY()
 
 	UPROPERTY()
 	FString DeviceBezelFile;
-	UPROPERTY()
+
+	PRAGMA_DISABLE_DEPRECATION_WARNINGS
+	UPROPERTY(meta = (DeprecatedProperty))
 	FPIEPreviewDeviceBezelViewportRect BezelViewportRect;
+	PRAGMA_ENABLE_DEPRECATION_WARNINGS
 };
 
 USTRUCT()
@@ -209,9 +208,11 @@ struct FPIEPreviewDeviceSpecifications
 
 	UPROPERTY()
 	TArray<float> ScaleFactors;
-	
-	UPROPERTY()
+
+	PRAGMA_DISABLE_DEPRECATION_WARNINGS
+	UPROPERTY(meta = (DeprecatedProperty))
 	FPIEBezelProperties BezelProperties;
+	PRAGMA_ENABLE_DEPRECATION_WARNINGS
 
 	UPROPERTY()
 	FPIEAndroidDeviceProperties AndroidProperties;
diff --git a/Engine/Source/Editor/UnrealEd/Private/Kismet2/DebuggerCommands.cpp b/Engine/Source/Editor/UnrealEd/Private/Kismet2/DebuggerCommands.cpp
--- a/Engine/Source/Editor/UnrealEd/Private/Kismet2/DebuggerCommands.cpp
+++ b/Engine/Source/Editor/UnrealEd/Private/Kismet2/DebuggerCommands.cpp
@@ -61,7 +61,6 @@
 #include "Interfaces/IProjectManager.h"
 
 #include "InstalledPlatformInfo.h"
-#include "PIEPreviewDeviceProfileSelectorModule.h"
 #include "IDesktopPlatform.h"
 #include "DesktopPlatformModule.h"
 #include "CookerSettings.h"
diff --git a/Engine/Source/Editor/UnrealEd/UnrealEd.Build.cs b/Engine/Source/Editor/UnrealEd/UnrealEd.Build.cs
--- a/Engine/Source/Editor/UnrealEd/UnrealEd.Build.cs
+++ b/Engine/Source/Editor/UnrealEd/UnrealEd.Build.cs
@@ -167,7 +167,6 @@ public class UnrealEd : ModuleRules
 				"ClothingSystemEditor",
 				"ClothingSystemRuntimeInterface",
 				"ClothingSystemRuntimeCommon",
-				"PIEPreviewDeviceProfileSelector",
 				"PakFileUtilities",
 				"TimeManagement",
 				"ScriptDisassembler",
diff --git a/Engine/Source/Programs/ChaosVisualDebugger/ChaosVisualDebugger.Build.cs b/Engine/Source/Programs/ChaosVisualDebugger/ChaosVisualDebugger.Build.cs
--- a/Engine/Source/Programs/ChaosVisualDebugger/ChaosVisualDebugger.Build.cs
+++ b/Engine/Source/Programs/ChaosVisualDebugger/ChaosVisualDebugger.Build.cs
@@ -23,7 +23,6 @@ public class ChaosVisualDebugger : ModuleRules
 				"MoviePlayerProxy",
 				"Projects",
 				"PreLoadScreen",
-				"PIEPreviewDeviceProfileSelector",
 				"RenderCore",
 				"RHI",
 				"Slate",
diff --git a/Engine/Source/Programs/LiveLinkHub/Source/LiveLinkHubLauncher/LiveLinkHubLauncher.Build.cs b/Engine/Source/Programs/LiveLinkHub/Source/LiveLinkHubLauncher/LiveLinkHubLauncher.Build.cs
--- a/Engine/Source/Programs/LiveLinkHub/Source/LiveLinkHubLauncher/LiveLinkHubLauncher.Build.cs
+++ b/Engine/Source/Programs/LiveLinkHub/Source/LiveLinkHubLauncher/LiveLinkHubLauncher.Build.cs
@@ -38,7 +38,6 @@ public class LiveLinkHubLauncher : ModuleRules
 				"ProfileVisualizer",
 				"Projects",
 				"PreLoadScreen",
-				"PIEPreviewDeviceProfileSelector",
 				"RenderCore",
 				"RHI",
 				"Slate",
diff --git a/Engine/Source/Programs/UnrealBuildTool/Configuration/UEBuildModuleCPP.cs b/Engine/Source/Programs/UnrealBuildTool/Configuration/UEBuildModuleCPP.cs
--- a/Engine/Source/Programs/UnrealBuildTool/Configuration/UEBuildModuleCPP.cs
+++ b/Engine/Source/Programs/UnrealBuildTool/Configuration/UEBuildModuleCPP.cs
@@ -259,8 +259,6 @@ namespace UnrealBuildTool
 			new("MovieSceneTools", "Sequencer"),
 			new("NavigationSystem", "UnrealEd"),
 			new("PacketHandler", "ReliabilityHandlerComponent"),
-			new("PIEPreviewDeviceProfileSelector", "MaterialShaderQualitySettings"),
-			new("PIEPreviewDeviceProfileSelector", "UnrealEd"),
 			new("PixelInspectorModule", "UnrealEd"),
 			new("Sequencer", "MovieSceneCaptureDialog"),
 			new("Sequencer", "MovieSceneTools"),
diff --git a/Engine/Source/Runtime/Engine/Engine.Build.cs b/Engine/Source/Runtime/Engine/Engine.Build.cs
--- a/Engine/Source/Runtime/Engine/Engine.Build.cs
+++ b/Engine/Source/Runtime/Engine/Engine.Build.cs
@@ -418,7 +418,6 @@ public class Engine : ModuleRules
 					"HierarchicalLODUtilities",
 					"MeshBuilder",
 					"NaniteBuilder",
-					"PIEPreviewDeviceProfileSelector",
 				}
 			);
 
diff --git a/Engine/Source/Runtime/Engine/Private/DeviceProfiles/DeviceProfileManager.cpp b/Engine/Source/Runtime/Engine/Private/DeviceProfiles/DeviceProfileManager.cpp
--- a/Engine/Source/Runtime/Engine/Private/DeviceProfiles/DeviceProfileManager.cpp
+++ b/Engine/Source/Runtime/Engine/Private/DeviceProfiles/DeviceProfileManager.cpp
@@ -16,10 +16,9 @@
 #if WITH_EDITOR
 #include "Interfaces/ITargetPlatform.h"
 #include "Interfaces/ITargetPlatformManagerModule.h"
-#include "PIEPreviewDeviceProfileSelectorModule.h"
-#else
-#include "IDeviceProfileSelectorModule.h"
 #endif
+#include "IDeviceProfileSelectorModule.h"
+
 #include "DeviceProfiles/DeviceProfileFragment.h"
 #include "GenericPlatform/GenericPlatformCrashContext.h"
 
@@ -1560,7 +1559,7 @@ IDeviceProfileSelectorModule* UDeviceProfileManager::GetPreviewDeviceProfileSele
 	if (PreviewConfigSystemIn->GetString(TEXT("DeviceProfileManager"), TEXT("PreviewDeviceProfileSelectionModule"), PreviewDeviceProfileSelectionModuleName, GEngineIni))
 	{
 		// this should only be specified when previewing.
-		if (IPIEPreviewDeviceModule* DPSelectorModule = FModuleManager::LoadModulePtr<IPIEPreviewDeviceModule>(*PreviewDeviceProfileSelectionModuleName))
+		if (IDeviceProfileSelectorModule* DPSelectorModule = FModuleManager::LoadModulePtr<IDeviceProfileSelectorModule>(*PreviewDeviceProfileSelectionModuleName))
 		{
 			return DPSelectorModule;
 		}
@@ -1607,20 +1606,6 @@ const FString UDeviceProfileManager::GetPlatformDeviceProfileName()
 		ActiveProfileName = DPSelectorModule->GetRuntimeDeviceProfileName();
 	}
 
-#if WITH_EDITOR
-	if (FPIEPreviewDeviceModule::IsRequestingPreviewDevice())
-	{
-		IDeviceProfileSelectorModule* PIEPreviewDeviceProfileSelectorModule = FModuleManager::LoadModulePtr<IDeviceProfileSelectorModule>("PIEPreviewDeviceProfileSelector");
-		if (PIEPreviewDeviceProfileSelectorModule)
-		{
-			FString PIEProfileName = PIEPreviewDeviceProfileSelectorModule->GetRuntimeDeviceProfileName();
-			if (!PIEProfileName.IsEmpty())
-			{
-				ActiveProfileName = PIEProfileName;
-			}
-		}
-	}
-#endif
 	return ActiveProfileName;
 }
 
diff --git a/Engine/Source/Runtime/Engine/Private/GameEngine.cpp b/Engine/Source/Runtime/Engine/Private/GameEngine.cpp
--- a/Engine/Source/Runtime/Engine/Private/GameEngine.cpp
+++ b/Engine/Source/Runtime/Engine/Private/GameEngine.cpp
@@ -73,10 +73,6 @@
 #include "Windows/WindowsPlatformCrashContext.h"
 #endif
 
-#if WITH_EDITOR
-#include "PIEPreviewDeviceProfileSelectorModule.h"
-#endif
-
 #if !UE_SERVER
 	#include "IMediaModule.h"
 #endif
@@ -226,17 +222,6 @@ void UGameEngine::CreateGameViewportWidget( UGameViewportClient* GameViewportCli
 			ViewportOverlayWidgetRef
 		];
 
-	// when we're running in a "device simulation" window, render the scene to an intermediate texture
-	// in the mobile device "emulation" case this is needed to properly position the viewport (as a widget) inside its bezel
-#if WITH_EDITOR
-	auto PIEPreviewDeviceModule = FModuleManager::LoadModulePtr<IPIEPreviewDeviceModule>("PIEPreviewDeviceProfileSelector");
-	if (PIEPreviewDeviceModule && FPIEPreviewDeviceModule::IsRequestingPreviewDevice())
-	{
-		bRenderDirectlyToWindow = false;
-		PIEPreviewDeviceModule->SetGameLayerManagerWidget(GameLayerManagerRef);
-	}
-#endif
-
 	const bool bStereoAllowed = bRenderDirectlyToWindow;
 
 	TSharedRef<SViewport> GameViewportWidgetRef = 
@@ -504,14 +489,6 @@ TSharedRef<SWindow> UGameEngine::CreateGameWindow()
 	int32 ResY = GSystemResolution.ResY;
 	EWindowMode::Type WindowMode = GSystemResolution.WindowMode;
 	int32 MonitorIndex = GSystemResolution.MonitorIndex;
-	
-#if WITH_EDITOR
-	/**************************************************************/
-	/*****PIE Window Gets its Size from SETRES console variable****/
-	/**************************************************************/
-	if (!FPIEPreviewDeviceModule::IsRequestingPreviewDevice())
-#endif
-	
 	{
 		ConditionallyOverrideSettings(ResX, ResY, WindowMode, MonitorIndex);
 
@@ -731,13 +708,6 @@ TSharedRef<SWindow> UGameEngine::CreateGameWindow()
 	.HasCloseButton(bAllowClose)
 	.SupportsMinimize(bAllowMinimize)
 	.SupportsMaximize(bAllowMaximize);
-#if WITH_EDITOR
-	auto PIEPreviewDeviceModule = FModuleManager::LoadModulePtr<IPIEPreviewDeviceModule>("PIEPreviewDeviceProfileSelector");
-	if (PIEPreviewDeviceModule && FPIEPreviewDeviceModule::IsRequestingPreviewDevice())
-	{
-		Window = PIEPreviewDeviceModule->CreatePIEPreviewDeviceWindow(FVector2D(ResX, ResY), WindowTitle, AutoCenterType, FVector2D(WinX, WinY), MaxWindowWidth, MaxWindowHeight);
-	}
-#endif
 	Window->SetAllowFastUpdate(true);
 
 	const bool bShowImmediately = false;
@@ -766,13 +736,6 @@ TSharedRef<SWindow> UGameEngine::CreateGameWindow()
 		Window->ShowWindow();
 	}
 
-#if WITH_EDITOR
-	if (PIEPreviewDeviceModule && FPIEPreviewDeviceModule::IsRequestingPreviewDevice())
-	{
-		PIEPreviewDeviceModule->OnWindowReady(Window);
-	}
-#endif
-
 	// Tick now to force a redraw of the window and ensure correct fullscreen application
 	FSlateApplication::Get().Tick();
 
diff --git a/Engine/Source/Runtime/Engine/Private/World.cpp b/Engine/Source/Runtime/Engine/Private/World.cpp
--- a/Engine/Source/Runtime/Engine/Private/World.cpp
+++ b/Engine/Source/Runtime/Engine/Private/World.cpp
@@ -96,7 +96,6 @@
 	#include "HierarchicalLODUtilitiesModule.h"
 	#include "ObjectTools.h"
 	#include "Engine/LODActor.h"
-	#include "PIEPreviewDeviceProfileSelectorModule.h"
 	#include "StaticMeshCompiler.h"
 	#include "WorldPartition/DataLayer/WorldDataLayers.h"
 	#include "PieFixupSerializer.h"
@@ -9283,8 +9282,7 @@ bool UWorld::IsPlayInPreview() const
 bool UWorld::IsPlayInMobilePreview() const
 {
 #if WITH_EDITOR
-	if (FPIEPreviewDeviceModule::IsRequestingPreviewDevice()
-		|| FParse::Param(FCommandLine::Get(), TEXT("featureleveles31")))
+	if (FParse::Param(FCommandLine::Get(), TEXT("featureleveles31")))
 	{
 		return true;
 	}
diff --git a/Engine/Source/Runtime/Launch/Launch.Build.cs b/Engine/Source/Runtime/Launch/Launch.Build.cs
--- a/Engine/Source/Runtime/Launch/Launch.Build.cs
+++ b/Engine/Source/Runtime/Launch/Launch.Build.cs
@@ -207,7 +207,6 @@ public class Launch : ModuleRules
 					"EditorFramework",
 					"UnrealEd",
 					"DesktopPlatform",
-					"PIEPreviewDeviceProfileSelector",
 			});
 
 			// ExtraModules that are loaded when WITH_EDITOR=1 is true
diff --git a/Engine/Source/Runtime/Launch/Private/LaunchEngineLoop.cpp b/Engine/Source/Runtime/Launch/Private/LaunchEngineLoop.cpp
--- a/Engine/Source/Runtime/Launch/Private/LaunchEngineLoop.cpp
+++ b/Engine/Source/Runtime/Launch/Private/LaunchEngineLoop.cpp
@@ -136,7 +136,6 @@
 	#include "UObject/StrongObjectPtr.h"
 	#include "Editor/UnrealEdEngine.h"
 	#include "Settings/EditorExperimentalSettings.h"
-	#include "PIEPreviewDeviceProfileSelectorModule.h"
 #endif// WITH_ENGINE
 	#include "Serialization/BulkDataRegistry.h"
 	#include "Virtualization/VirtualizationSystem.h"
@@ -3071,10 +3070,6 @@ int32 FEngineLoop::PreInitPreStartupScreen(const TCHAR* CmdLine)
 	InitEngineTextLocalization();
 
 	bool bForceEnableHighDPI = false;
-#if WITH_EDITOR
-	bForceEnableHighDPI = FPIEPreviewDeviceModule::IsRequestingPreviewDevice();
-#endif
-
 	// This must be called before any window (including the splash screen is created
 	FSlateApplication::InitHighDPI(bForceEnableHighDPI);
 
@@ -3734,16 +3729,6 @@ int32 FEngineLoop::PreInitPostStartupScreen(const TCHAR* CmdLine)
 
 		FEmbeddedCommunication::ForceTick(6);
 
-#if WITH_EDITOR
-		if(FPIEPreviewDeviceModule::IsRequestingPreviewDevice())
-		{
-			auto PIEPreviewDeviceModule = FModuleManager::LoadModulePtr<IPIEPreviewDeviceModule>("PIEPreviewDeviceProfileSelector");
-			if (PIEPreviewDeviceModule)
-			{
-				PIEPreviewDeviceModule->ApplyPreviewDeviceState();
-			}
-		}
-#endif
 #if USE_LOCALIZED_PACKAGE_CACHE
 		{
 			SCOPED_BOOT_TIMING("FPackageLocalizationManager::Get().PerformLazyInitialization()");
@@ -6900,17 +6885,6 @@ bool FEngineLoop::AppInit( )
 		FCoreDelegates::OnInit.Broadcast();
 	}
 
-#if WITH_EDITOR && WITH_ENGINE
-	if (FPIEPreviewDeviceModule::IsRequestingPreviewDevice())
-	{
-		FPIEPreviewDeviceModule* PIEPreviewDeviceProfileSelectorModule = FModuleManager::LoadModulePtr<FPIEPreviewDeviceModule>("PIEPreviewDeviceProfileSelector");
-		if (PIEPreviewDeviceProfileSelectorModule)
-		{
-			Scalability::ChangeScalabilityPreviewPlatform(PIEPreviewDeviceProfileSelectorModule->GetPreviewPlatformName(), GShaderPlatformForFeatureLevel[ERHIFeatureLevel::ES3_1]);
-		}
-	}
-#endif
-
 	FEmbeddedCommunication::ForceTick(19);
 
 	return true;
