Removed string copy when splitting to make a FName

#rnx

#Synced-CL 51168648

diff --git a/Engine/Source/Runtime/Engine/Private/WorldPartition/ActorDescContainer.cpp b/Engine/Source/Runtime/Engine/Private/WorldPartition/ActorDescContainer.cpp
--- a/Engine/Source/Runtime/Engine/Private/WorldPartition/ActorDescContainer.cpp
+++ b/Engine/Source/Runtime/Engine/Private/WorldPartition/ActorDescContainer.cpp
@@ -13,6 +13,7 @@
 #include "AssetRegistry/ARFilter.h"
 #include "AssetRegistry/AssetRegistryModule.h"
 #include "Modules/ModuleManager.h"
+#include "String/Split.h"
 #include "WorldPartition/WorldPartition.h"
 #include "WorldPartition/WorldPartitionActorDescUtils.h"
 #include "WorldPartition/WorldPartitionClassDescRegistry.h"
@@ -459,14 +460,14 @@ bool UActorDescContainer::ShouldHandleActorEvent(const AActor* Actor)
 
 const FWorldPartitionActorDesc* UActorDescContainer::GetActorDescByPath(const FString& ActorPath) const
 {
-	FString ActorName;
-	FString ActorContext;
-	if (!ActorPath.Split(TEXT("."), &ActorContext, &ActorName, ESearchCase::CaseSensitive, ESearchDir::FromEnd))
+	FStringView ActorName;
+	FStringView ActorContext;
+	if (!UE::String::SplitLastChar(FStringView(ActorPath), TEXT('.'), ActorContext, ActorName, ESearchCase::CaseSensitive))
 	{
 		ActorName = ActorPath;
 	}
 
-	return GetActorDescByName(FName(*ActorName));
+	return GetActorDescByName(FName(ActorName));
 }
 
 const FWorldPartitionActorDesc* UActorDescContainer::GetActorDescByPath(const FSoftObjectPath& ActorPath) const
