[Misc/ALT]
- Fix some checkSlow during multithreaded serialization

#Synced-CL 51209060

diff --git a/Engine/Plugins/Compositing/OpenColorIO/Source/OpenColorIO/Public/OpenColorIOShared.h b/Engine/Plugins/Compositing/OpenColorIO/Source/OpenColorIO/Public/OpenColorIOShared.h
--- a/Engine/Plugins/Compositing/OpenColorIO/Source/OpenColorIO/Public/OpenColorIOShared.h
+++ b/Engine/Plugins/Compositing/OpenColorIO/Source/OpenColorIO/Public/OpenColorIOShared.h
@@ -393,7 +393,7 @@ public:
 
 	FOpenColorIOShaderMap* GetGameThreadShaderMap() const
 	{
-		checkSlow(IsInGameThread() || IsInAsyncLoadingThread());
+		checkSlow(IsInGameThread() || IsInAsyncLoadingThread() || IsInParallelLoadingThread());
 		return GameThreadShaderMap;
 	}
 
@@ -406,7 +406,7 @@ public:
 	/** Note: SetRenderingThreadShaderMap must also be called with the same value, but from the rendering thread. */
 	void SetGameThreadShaderMap(FOpenColorIOShaderMap* InShaderMap)
 	{
-		checkSlow(IsInGameThread() || IsInAsyncLoadingThread());
+		checkSlow(IsInGameThread() || IsInAsyncLoadingThread() || IsInParallelLoadingThread());
 		GameThreadShaderMap = InShaderMap;
 		if (LIKELY(GameThreadShaderMap))
 		{
@@ -426,7 +426,7 @@ public:
 
 	void SetInlineShaderMap(FOpenColorIOShaderMap* InShaderMap)
 	{
-		checkSlow(IsInGameThread() || IsInAsyncLoadingThread());
+		checkSlow(IsInGameThread() || IsInAsyncLoadingThread() || IsInParallelLoadingThread());
 		GameThreadShaderMap = InShaderMap;
 		if (LIKELY(GameThreadShaderMap))
 		{
diff --git a/Engine/Plugins/Experimental/NNERuntimeIREE/Source/NNERuntimeIREEShader/Internal/NNERuntimeIREEShaderShared.h b/Engine/Plugins/Experimental/NNERuntimeIREE/Source/NNERuntimeIREEShader/Internal/NNERuntimeIREEShaderShared.h
--- a/Engine/Plugins/Experimental/NNERuntimeIREE/Source/NNERuntimeIREEShader/Internal/NNERuntimeIREEShaderShared.h
+++ b/Engine/Plugins/Experimental/NNERuntimeIREE/Source/NNERuntimeIREEShader/Internal/NNERuntimeIREEShaderShared.h
@@ -298,7 +298,7 @@ public:
 
 	FNNERuntimeIREEShaderMap* GetGameThreadShaderMap() const
 	{
-		checkSlow(IsInGameThread() || IsInAsyncLoadingThread());
+		checkSlow(IsInGameThread() || IsInAsyncLoadingThread() || IsInParallelLoadingThread());
 		return GameThreadShaderMap;
 	}
 
@@ -311,7 +311,7 @@ public:
 	/** Note: SetRenderingThreadShaderMap must also be called with the same value, but from the rendering thread. */
 	void SetGameThreadShaderMap(FNNERuntimeIREEShaderMap* InShaderMap)
 	{
-		checkSlow(IsInGameThread() || IsInAsyncLoadingThread());
+		checkSlow(IsInGameThread() || IsInAsyncLoadingThread() || IsInParallelLoadingThread());
 		GameThreadShaderMap = InShaderMap;
 		if (LIKELY(GameThreadShaderMap))
 		{
diff --git a/Engine/Plugins/FX/Niagara/Source/NiagaraShader/Public/NiagaraShared.h b/Engine/Plugins/FX/Niagara/Source/NiagaraShader/Public/NiagaraShared.h
--- a/Engine/Plugins/FX/Niagara/Source/NiagaraShader/Public/NiagaraShared.h
+++ b/Engine/Plugins/FX/Niagara/Source/NiagaraShader/Public/NiagaraShared.h
@@ -792,14 +792,14 @@ public:
 
 	class FNiagaraShaderMap* GetGameThreadShaderMap() const
 	{
-		checkSlow(IsInGameThread() || IsInAsyncLoadingThread());
+		checkSlow(IsInGameThread() || IsInAsyncLoadingThread() || IsInParallelLoadingThread());
 		return GameThreadShaderMap;
 	}
 
 	/** Note: SetRenderingThreadShaderMap must also be called with the same value, but from the rendering thread. */
 	void SetGameThreadShaderMap(FNiagaraShaderMap* InShaderMap)
 	{
-		checkSlow(IsInGameThread() || IsInAsyncLoadingThread());
+		checkSlow(IsInGameThread() || IsInAsyncLoadingThread() || IsInParallelLoadingThread());
 		GameThreadShaderMap = InShaderMap;
 		if (LIKELY(GameThreadShaderMap))
 		{
@@ -823,7 +823,7 @@ public:
 
 	void SetShaderMap(FNiagaraShaderMap* InShaderMap)
 	{
-		checkSlow(IsInGameThread() || IsInAsyncLoadingThread());
+		checkSlow(IsInGameThread() || IsInAsyncLoadingThread() || IsInParallelLoadingThread());
 		GameThreadShaderMap = InShaderMap;
 		if (LIKELY(GameThreadShaderMap))
 		{
diff --git a/Engine/Plugins/Runtime/ComputeFramework/Source/ComputeFramework/Private/ComputeFramework/ComputeKernelShared.h b/Engine/Plugins/Runtime/ComputeFramework/Source/ComputeFramework/Private/ComputeFramework/ComputeKernelShared.h
--- a/Engine/Plugins/Runtime/ComputeFramework/Source/ComputeFramework/Private/ComputeFramework/ComputeKernelShared.h
+++ b/Engine/Plugins/Runtime/ComputeFramework/Source/ComputeFramework/Private/ComputeFramework/ComputeKernelShared.h
@@ -372,7 +372,7 @@ public:
 
 	FComputeKernelShaderMap* GetGameThreadShaderMap() const
 	{
-		checkSlow(IsInGameThread() || IsInAsyncLoadingThread());
+		checkSlow(IsInGameThread() || IsInAsyncLoadingThread() || IsInParallelLoadingThread());
 		return GameThreadShaderMap;
 	}
 
@@ -385,7 +385,7 @@ public:
 	/** Note: SetRenderingThreadShaderMap must also be called with the same value, but from the rendering thread. */
 	void SetGameThreadShaderMap(FComputeKernelShaderMap* InShaderMap)
 	{
-		checkSlow(IsInGameThread() || IsInAsyncLoadingThread());
+		checkSlow(IsInGameThread() || IsInAsyncLoadingThread() || IsInParallelLoadingThread());
 		GameThreadShaderMap = InShaderMap;
 		if (LIKELY(GameThreadShaderMap))
 		{
@@ -403,7 +403,7 @@ public:
 
 	void SetInlineShaderMap(FComputeKernelShaderMap* InShaderMap)
 	{
-		checkSlow(IsInGameThread() || IsInAsyncLoadingThread());
+		checkSlow(IsInGameThread() || IsInAsyncLoadingThread() || IsInParallelLoadingThread());
 		GameThreadShaderMap = InShaderMap;
 		if (LIKELY(GameThreadShaderMap))
 		{
diff --git a/Engine/Source/Runtime/Engine/Private/Materials/MaterialShared.cpp b/Engine/Source/Runtime/Engine/Private/Materials/MaterialShared.cpp
--- a/Engine/Source/Runtime/Engine/Private/Materials/MaterialShared.cpp
+++ b/Engine/Source/Runtime/Engine/Private/Materials/MaterialShared.cpp
@@ -951,7 +951,7 @@ const FMaterialShaderMap* FMaterial::GetShaderMapToUse() const
 { 
 	const FMaterialShaderMap* ShaderMapToUse = NULL;
 
-	if (IsInGameThread() || IsInParallelGameThread())
+	if (!IsInAnyRenderingThread())
 	{
 		// If we are accessing uniform texture expressions on the game thread, use results from a shader map whose compile is in flight that matches this material
 		// This allows querying what textures a material uses even when it is being asynchronously compiled
@@ -1255,7 +1255,7 @@ ESubstrateBsdfFeature FMaterial::MaterialGetSubstrateMaterialBsdfFeatures_Render
 
 void FMaterial::SetGameThreadShaderMap(FMaterialShaderMap* InMaterialShaderMap)
 {
-	checkSlow(IsInGameThread() || IsInAsyncLoadingThread());
+	checkSlow(IsInGameThread() || IsInAsyncLoadingThread() || IsInParallelLoadingThread());
 
 	bool bAssumeShaderMapIsComplete = false;
 #if UE_BUILD_SHIPPING || UE_BUILD_TEST
@@ -1280,7 +1280,7 @@ void FMaterial::SetGameThreadShaderMap(FMaterialShaderMap* InMaterialShaderMap)
 
 void FMaterial::UpdateInlineShaderMapIsComplete()
 {
-	checkSlow(IsInGameThread() || IsInAsyncLoadingThread());
+	checkSlow(IsInGameThread() || IsInAsyncLoadingThread() || IsInParallelLoadingThread());
 	check(bContainsInlineShaders);
 	// We expect inline shader maps to be complete, so we want to log missing shaders here
 	const bool bSilent = false;
@@ -1301,7 +1301,7 @@ void FMaterial::UpdateInlineShaderMapIsComplete()
 
 void FMaterial::SetInlineShaderMap(FMaterialShaderMap* InMaterialShaderMap)
 {
-	checkSlow(IsInGameThread() || IsInAsyncLoadingThread());
+	checkSlow(IsInGameThread() || IsInAsyncLoadingThread() || IsInParallelLoadingThread());
 	check(InMaterialShaderMap);
 
 	GameThreadShaderMap = InMaterialShaderMap;
