Fix issue that Instanced static mesh component didn't respect SetActorEnableCollision result

#changelist validated

#Synced-CL 51262360

diff --git a/Engine/Source/Runtime/Engine/Classes/Components/InstancedStaticMeshComponent.h b/Engine/Source/Runtime/Engine/Classes/Components/InstancedStaticMeshComponent.h
--- a/Engine/Source/Runtime/Engine/Classes/Components/InstancedStaticMeshComponent.h
+++ b/Engine/Source/Runtime/Engine/Classes/Components/InstancedStaticMeshComponent.h
@@ -449,6 +449,7 @@ public:
 	ENGINE_API virtual void PostLoad() override;
 	ENGINE_API virtual void OnRegister() override;
 	ENGINE_API virtual void OnUnregister() override;
+	ENGINE_API virtual void OnActorEnableCollisionChanged() override;
 	
 	/** Sets to use RemoveAtSwap on instance removal. This is an optimization, but will change the resultant instance reordering. */
 	void SetRemoveSwap() { bSupportRemoveAtSwap = true; }
diff --git a/Engine/Source/Runtime/Engine/Private/InstancedStaticMesh.cpp b/Engine/Source/Runtime/Engine/Private/InstancedStaticMesh.cpp
--- a/Engine/Source/Runtime/Engine/Private/InstancedStaticMesh.cpp
+++ b/Engine/Source/Runtime/Engine/Private/InstancedStaticMesh.cpp
@@ -5132,6 +5132,19 @@ void UInstancedStaticMeshComponent::OnUnregister()
 	Super::OnUnregister();
 }
 
+void UInstancedStaticMeshComponent::OnActorEnableCollisionChanged()
+{
+	Super::OnActorEnableCollisionChanged();
+	
+	for (FBodyInstance* InstanceBody : InstanceBodies)
+	{
+		if (InstanceBody)
+		{
+			InstanceBody->UpdatePhysicsFilterData();
+		}
+	}
+}
+
 #if WITH_EDITOR
 
 bool UInstancedStaticMeshComponent::CanEditChange(const FProperty* InProperty) const
