[SkeletalMesh] Fix assertion of expected position precision by using `IsNearlyEqualByULP()` when comparing floating-point values.

#rnx

#Synced-CL 51246683

diff --git a/Engine/Source/Runtime/Engine/Private/SkeletalMeshLODRenderData.cpp b/Engine/Source/Runtime/Engine/Private/SkeletalMeshLODRenderData.cpp
--- a/Engine/Source/Runtime/Engine/Private/SkeletalMeshLODRenderData.cpp
+++ b/Engine/Source/Runtime/Engine/Private/SkeletalMeshLODRenderData.cpp
@@ -863,8 +863,9 @@ void FSkeletalMeshLODRenderData::SerializeStreamedData(FArchive& Ar, USkinnedAss
 
 #if DO_CHECK
 				const FSkeletalMeshLODInfo* SkeletalMeshLODInfo = Owner->GetLODInfo(LODIdx);
-				float PositionPrecision = UE::MorphTargetVertexCodec::ComputePositionPrecision(SkeletalMeshLODInfo->MorphTargetPositionErrorTolerance);
-				checkf(PositionPrecision == MorphTargetVertexInfoBuffers.GetPositionPrecision(), TEXT("Morph target render data position tollerance for the skeleton %s does not match the expected tolerance in the LOD info."), *Owner->GetName());
+				const float ExpectedPositionPrecision = UE::MorphTargetVertexCodec::ComputePositionPrecision(SkeletalMeshLODInfo->MorphTargetPositionErrorTolerance);
+				checkf(FMath::IsNearlyEqualByULP(ExpectedPositionPrecision, MorphTargetVertexInfoBuffers.GetPositionPrecision()), TEXT("Morph target render data position tollerance in the LOD info (%f) for the skeleton %s does not match the expected tolerance (%f)."),
+					MorphTargetVertexInfoBuffers.GetPositionPrecision(), *Owner->GetName(), ExpectedPositionPrecision);
 #endif
 			}
 		}
