[Texture]
- Cleanup the GetResource logic

#Synced-CL 51263423

diff --git a/Engine/Source/Runtime/Engine/Private/Texture.cpp b/Engine/Source/Runtime/Engine/Private/Texture.cpp
--- a/Engine/Source/Runtime/Engine/Private/Texture.cpp
+++ b/Engine/Source/Runtime/Engine/Private/Texture.cpp
@@ -234,17 +234,12 @@ void UTexture::SetModernSettingsForNewOrChangedTexture()
 
 const FTextureResource* UTexture::GetResource() const
 {
-	if (IsInParallelGameThread() || IsInGameThread() || IsInSlateThread() || IsInAsyncLoadingThread())
-	{
-		return PrivateResource;
-	}
-	else if (IsInParallelRenderingThread() || IsInRHIThread())
+	if (IsInAnyRenderingThread())
 	{
 		return PrivateResourceRenderThread;
 	}
 
-	ensureMsgf(false, TEXT("Attempted to access a texture resource from an unknown thread."));
-	return nullptr;
+	return PrivateResource;
 }
 
 FTextureResource* UTexture::GetResource()
