Wrap analytic properties in a singleton to prevent undefined static initialization ordering.
-Fixes an issue causing SessionSummary to fail in merge module builds.

#rnx

diff --git a/Engine/Source/Runtime/Engine/Private/Analytics/EngineAnalyticsSessionSummary.cpp b/Engine/Source/Runtime/Engine/Private/Analytics/EngineAnalyticsSessionSummary.cpp
--- a/Engine/Source/Runtime/Engine/Private/Analytics/EngineAnalyticsSessionSummary.cpp
+++ b/Engine/Source/Runtime/Engine/Private/Analytics/EngineAnalyticsSessionSummary.cpp
@@ -25,28 +25,59 @@ extern CORE_API bool GIsGPUCrashed;
 
 namespace EngineAnalyticsProperties
 {
-	// List of mutable properties.
-	static const TAnalyticsProperty<FDateTime> Timestamp         = TEXT("Timestamp");
-	static const TAnalyticsProperty<FDateTime> LastTickTimestamp = TEXT("LastTickTimestamp");
-	static const TAnalyticsProperty<int32> SessionDurationSecs   = TEXT("SessionDuration");
-	static const TAnalyticsProperty<uint64> EngineTickCount      = TEXT("EngineTickCount");
-	static const TAnalyticsProperty<uint64> SessionTickCount     = TEXT("SessionTickCount");
-	static const TAnalyticsProperty<FString> UserActivity        = TEXT("CurrentUserActivity");
-	static const TAnalyticsProperty<int32>  MonitorExitCode      = TEXT("MonitorExitCode");
-	static const TAnalyticsProperty<float> AverageFPS            = TEXT("AverageFPS");
-	static const TAnalyticsProperty<bool> IsTerminating			 = TEXT("IsTerminating");
-	static const TAnalyticsProperty<bool> IsHang				 = TEXT("IsHang");
-	static const TAnalyticsProperty<bool> WasShutdown            = TEXT("WasShutdown");
-	static const TAnalyticsProperty<bool> IsCrashing             = TEXT("IsCrash");
-	static const TAnalyticsProperty<bool> IsGpuCrash             = TEXT("GPUCrash");
-	static const TAnalyticsProperty<bool> IsVanilla              = TEXT("IsVanilla");
-	static const TAnalyticsProperty<bool> IsLowDriveSpace        = TEXT("IsLowDriveSpace");
-	static const TAnalyticsProperty<bool> WasDebuggerIgnored     = TEXT("WasDebuggerIgnored");
-	static const TAnalyticsProperty<bool> WasEverDebugged        = TEXT("WasDebugged");
-	static const TAnalyticsProperty<bool> IsCrcMissing           = TEXT("IsCrcMissing");
-	static const TAnalyticsProperty<bool> RunningOnBattery       = TEXT("RunningOnBattery");
-	static const TAnalyticsProperty<bool> IsUserLoggingOut       = FAnalyticsSessionSummaryManager::IsUserLoggingOutProperty;
-	static const TAnalyticsProperty<int32> ShutdownType          = FAnalyticsSessionSummaryManager::ShutdownTypeCodeProperty;
+	// List of mutable properties. These need to be in a singleton to prevent undefined static initialization ordering
+	class FProperties
+	{
+	public:
+		static const FProperties& Get() { static FProperties Instance; return Instance; }
+
+		const TAnalyticsProperty<FDateTime> Timestamp;
+		const TAnalyticsProperty<FDateTime> LastTickTimestamp;
+		const TAnalyticsProperty<int32> SessionDurationSecs;
+		const TAnalyticsProperty<uint64> EngineTickCount;
+		const TAnalyticsProperty<uint64> SessionTickCount;
+		const TAnalyticsProperty<FString> UserActivity;
+		const TAnalyticsProperty<int32>  MonitorExitCode;
+		const TAnalyticsProperty<float> AverageFPS;
+		const TAnalyticsProperty<bool> IsTerminating;
+		const TAnalyticsProperty<bool> IsHang;
+		const TAnalyticsProperty<bool> WasShutdown;
+		const TAnalyticsProperty<bool> IsCrashing;
+		const TAnalyticsProperty<bool> IsGpuCrash;
+		const TAnalyticsProperty<bool> IsVanilla;
+		const TAnalyticsProperty<bool> IsLowDriveSpace;
+		const TAnalyticsProperty<bool> WasDebuggerIgnored;
+		const TAnalyticsProperty<bool> WasEverDebugged;
+		const TAnalyticsProperty<bool> IsCrcMissing;
+		const TAnalyticsProperty<bool> RunningOnBattery;
+		const TAnalyticsProperty<bool> IsUserLoggingOut;
+		const TAnalyticsProperty<int32> ShutdownType;
+
+	private:
+		FProperties()
+			: Timestamp(TEXT("Timestamp"))
+			, LastTickTimestamp(TEXT("LastTickTimestamp"))
+			, SessionDurationSecs(TEXT("SessionDuration"))
+			, EngineTickCount(TEXT("EngineTickCount"))
+			, SessionTickCount(TEXT("SessionTickCount"))
+			, UserActivity(TEXT("CurrentUserActivity"))
+			, MonitorExitCode(TEXT("MonitorExitCode"))
+			, AverageFPS(TEXT("AverageFPS"))
+			, IsTerminating(TEXT("IsTerminating"))
+			, IsHang(TEXT("IsHang"))
+			, WasShutdown(TEXT("WasShutdown"))
+			, IsCrashing(TEXT("IsCrash"))
+			, IsGpuCrash(TEXT("GPUCrash"))
+			, IsVanilla(TEXT("IsVanilla"))
+			, IsLowDriveSpace(TEXT("IsLowDriveSpace"))
+			, WasDebuggerIgnored(TEXT("WasDebuggerIgnored"))
+			, WasEverDebugged(TEXT("WasDebugged"))
+			, IsCrcMissing(TEXT("IsCrcMissing"))
+			, RunningOnBattery(TEXT("RunningOnBattery"))
+			, IsUserLoggingOut(FAnalyticsSessionSummaryManager::IsUserLoggingOutProperty)
+			, ShutdownType(FAnalyticsSessionSummaryManager::ShutdownTypeCodeProperty)
+		{}
+	};
 
 	static FString UnknownUserActivity("Unknown");
 
@@ -195,25 +226,25 @@ FEngineAnalyticsSessionSummary::FEngineAnalyticsSessionSummary(TSharedPtr<IAnaly
 	FDateTime CurrTimeUtc = FDateTime::UtcNow();
 
 	// Create the mutable metrics.
-	EngineAnalyticsProperties::Timestamp.Set(GetStore(), CurrTimeUtc);
-	EngineAnalyticsProperties::LastTickTimestamp.Set(GetStore(), CurrTimeUtc);
-	EngineAnalyticsProperties::SessionDurationSecs.Set(GetStore(), FPlatformTime::Seconds() - SessionStartTimeSecs);
-	EngineAnalyticsProperties::EngineTickCount.Set(GetStore(), EngineAnalyticsProperties::GetEngineTickCount());
-	EngineAnalyticsProperties::SessionTickCount.Set(GetStore(), CurrSessionTickCount);
-	EngineAnalyticsProperties::UserActivity.Set(GetStore(), EngineAnalyticsProperties::GetUserActivityString(FUserActivityTracking::GetUserActivity()));
-	EngineAnalyticsProperties::AverageFPS.Set(GetStore(), EngineAnalyticsProperties::GetAverageFPS());
-	EngineAnalyticsProperties::IsTerminating.Set(GetStore(),  false);
-	EngineAnalyticsProperties::IsHang.Set(GetStore(), false);
-	EngineAnalyticsProperties::WasShutdown.Set(GetStore(), false);
-	EngineAnalyticsProperties::IsCrashing.Set(GetStore(), false);
-	EngineAnalyticsProperties::IsGpuCrash.Set(GetStore(), false);
-	EngineAnalyticsProperties::IsUserLoggingOut.Set(GetStore(), false);
-	EngineAnalyticsProperties::IsVanilla.Set(GetStore(), GEngine && GEngine->IsVanillaProduct());
-	EngineAnalyticsProperties::IsLowDriveSpace.Set(GetStore(), false);
-	EngineAnalyticsProperties::WasDebuggerIgnored.Set(GetStore(), false);
-	EngineAnalyticsProperties::WasEverDebugged.Set(GetStore(), false);
-	EngineAnalyticsProperties::ShutdownType.Set(GetStore(), (int32)EAnalyticsSessionShutdownType::Abnormal);
-	EngineAnalyticsProperties::RunningOnBattery.Set(GetStore(), FPlatformMisc::IsRunningOnBattery());
+	EngineAnalyticsProperties::FProperties::Get().Timestamp.Set(GetStore(), CurrTimeUtc);
+	EngineAnalyticsProperties::FProperties::Get().LastTickTimestamp.Set(GetStore(), CurrTimeUtc);
+	EngineAnalyticsProperties::FProperties::Get().SessionDurationSecs.Set(GetStore(), FPlatformTime::Seconds() - SessionStartTimeSecs);
+	EngineAnalyticsProperties::FProperties::Get().EngineTickCount.Set(GetStore(), EngineAnalyticsProperties::GetEngineTickCount());
+	EngineAnalyticsProperties::FProperties::Get().SessionTickCount.Set(GetStore(), CurrSessionTickCount);
+	EngineAnalyticsProperties::FProperties::Get().UserActivity.Set(GetStore(), EngineAnalyticsProperties::GetUserActivityString(FUserActivityTracking::GetUserActivity()));
+	EngineAnalyticsProperties::FProperties::Get().AverageFPS.Set(GetStore(), EngineAnalyticsProperties::GetAverageFPS());
+	EngineAnalyticsProperties::FProperties::Get().IsTerminating.Set(GetStore(),  false);
+	EngineAnalyticsProperties::FProperties::Get().IsHang.Set(GetStore(), false);
+	EngineAnalyticsProperties::FProperties::Get().WasShutdown.Set(GetStore(), false);
+	EngineAnalyticsProperties::FProperties::Get().IsCrashing.Set(GetStore(), false);
+	EngineAnalyticsProperties::FProperties::Get().IsGpuCrash.Set(GetStore(), false);
+	EngineAnalyticsProperties::FProperties::Get().IsUserLoggingOut.Set(GetStore(), false);
+	EngineAnalyticsProperties::FProperties::Get().IsVanilla.Set(GetStore(), GEngine && GEngine->IsVanillaProduct());
+	EngineAnalyticsProperties::FProperties::Get().IsLowDriveSpace.Set(GetStore(), false);
+	EngineAnalyticsProperties::FProperties::Get().WasDebuggerIgnored.Set(GetStore(), false);
+	EngineAnalyticsProperties::FProperties::Get().WasEverDebugged.Set(GetStore(), false);
+	EngineAnalyticsProperties::FProperties::Get().ShutdownType.Set(GetStore(), (int32)EAnalyticsSessionShutdownType::Abnormal);
+	EngineAnalyticsProperties::FProperties::Get().RunningOnBattery.Set(GetStore(), FPlatformMisc::IsRunningOnBattery());
 
 	Store->Flush();
 
@@ -240,8 +271,8 @@ void FEngineAnalyticsSessionSummary::Shutdown()
 		FCoreDelegates::OnUserLoginChangedEvent.RemoveAll(this);
 		FUserActivityTracking::OnActivityChanged.RemoveAll(this);
 
-		EngineAnalyticsProperties::WasShutdown.Set(GetStore(), true);
-		EngineAnalyticsProperties::ShutdownType.Set(GetStore(), (int32)EAnalyticsSessionShutdownType::Shutdown, EngineAnalyticsProperties::ShouldOverwriteShutdownType);
+		EngineAnalyticsProperties::FProperties::Get().WasShutdown.Set(GetStore(), true);
+		EngineAnalyticsProperties::FProperties::Get().ShutdownType.Set(GetStore(), (int32)EAnalyticsSessionShutdownType::Shutdown, EngineAnalyticsProperties::ShouldOverwriteShutdownType);
 		UpdateSessionProgress(/*bFlushAsync*/false, /*FlushTimeout*/FTimespan::MaxValue()); // Forces a sunchronous flush.
 		bShutdown = true;
 	}
@@ -255,13 +286,13 @@ void FEngineAnalyticsSessionSummary::Tick(float DeltaTime)
 	}
 
 	// The number of times this instance was ticked.
-	EngineAnalyticsProperties::SessionTickCount.Update(GetStore(), [](uint64& Value) { ++Value; return true; });
+	EngineAnalyticsProperties::FProperties::Get().SessionTickCount.Update(GetStore(), [](uint64& Value) { ++Value; return true; });
 
 	// The average FPS.
-	EngineAnalyticsProperties::AverageFPS.Set(GetStore(), EngineAnalyticsProperties::GetAverageFPS());
+	EngineAnalyticsProperties::FProperties::Get().AverageFPS.Set(GetStore(), EngineAnalyticsProperties::GetAverageFPS());
 
 	// This tick timestamp. Not allowed to go back.
-	EngineAnalyticsProperties::LastTickTimestamp.Set(GetStore(), FDateTime::UtcNow(), [](const FDateTime* Actual, const FDateTime& Proposed){ return Proposed >*Actual; });
+	EngineAnalyticsProperties::FProperties::Get().LastTickTimestamp.Set(GetStore(), FDateTime::UtcNow(), [](const FDateTime* Actual, const FDateTime& Proposed){ return Proposed >*Actual; });
 
 	// Detect if CRC state changed since the last update.
 	bool bFlush = UpdateExternalProcessReporterState(/*bQuickCheck*/true);
@@ -279,9 +310,9 @@ void FEngineAnalyticsSessionSummary::Tick(float DeltaTime)
 
 void FEngineAnalyticsSessionSummary::OnCrashing()
 {
-	EngineAnalyticsProperties::IsCrashing.Set(GetStore(), true);
-	EngineAnalyticsProperties::IsGpuCrash.Set(GetStore(), GIsGPUCrashed);
-	EngineAnalyticsProperties::ShutdownType.Set(GetStore(), (int32)EAnalyticsSessionShutdownType::Crashed);
+	EngineAnalyticsProperties::FProperties::Get().IsCrashing.Set(GetStore(), true);
+	EngineAnalyticsProperties::FProperties::Get().IsGpuCrash.Set(GetStore(), GIsGPUCrashed);
+	EngineAnalyticsProperties::FProperties::Get().ShutdownType.Set(GetStore(), (int32)EAnalyticsSessionShutdownType::Crashed);
 
 	// Flush the store if possible but don't wait forever to avoid deadlocking if the crash affects the store ability to complete its background task.
 	UpdateSessionProgress(/*bFlushAsync*/false, /*FlushTimeout*/FTimespan::FromMilliseconds(100), /*bCrashing*/true);
@@ -289,15 +320,15 @@ void FEngineAnalyticsSessionSummary::OnCrashing()
 
 void FEngineAnalyticsSessionSummary::OnHung()
 {
-	EngineAnalyticsProperties::IsHang.Set(GetStore(), true);
-	EngineAnalyticsProperties::ShutdownType.Set(GetStore(), (int32)EAnalyticsSessionShutdownType::Hang, EngineAnalyticsProperties::ShouldOverwriteShutdownType);
+	EngineAnalyticsProperties::FProperties::Get().IsHang.Set(GetStore(), true);
+	EngineAnalyticsProperties::FProperties::Get().ShutdownType.Set(GetStore(), (int32)EAnalyticsSessionShutdownType::Hang, EngineAnalyticsProperties::ShouldOverwriteShutdownType);
 	UpdateSessionProgress(/*bFlushAsync*/false, /*FlushTimeout*/FTimespan::FromMilliseconds(100));
 }
 
 void FEngineAnalyticsSessionSummary::OnTerminate()
 {
-	EngineAnalyticsProperties::IsTerminating.Set(GetStore(), true);
-	EngineAnalyticsProperties::ShutdownType.Set(GetStore(), (int32)EAnalyticsSessionShutdownType::Terminated, EngineAnalyticsProperties::ShouldOverwriteShutdownType);
+	EngineAnalyticsProperties::FProperties::Get().IsTerminating.Set(GetStore(), true);
+	EngineAnalyticsProperties::FProperties::Get().ShutdownType.Set(GetStore(), (int32)EAnalyticsSessionShutdownType::Terminated, EngineAnalyticsProperties::ShouldOverwriteShutdownType);
 	UpdateSessionProgress(/*bFlushAsync*/false, /*FlushTimeout*/FTimespan::FromMilliseconds(100));
 }
 
@@ -305,26 +336,26 @@ void FEngineAnalyticsSessionSummary::OnUserLoginChanged(bool bLoggingIn, int32,
 {
 	if (!bLoggingIn)
 	{
-		EngineAnalyticsProperties::IsUserLoggingOut.Set(GetStore(), true);
+		EngineAnalyticsProperties::FProperties::Get().IsUserLoggingOut.Set(GetStore(), true);
 		UpdateSessionProgress(/*bFlushAsync*/false, /*FlushTimeout*/FTimespan::FromMilliseconds(100));
 	}
 }
 
 void FEngineAnalyticsSessionSummary::OnVanillaStateChanged(bool bIsVanilla)
 {
-	EngineAnalyticsProperties::IsVanilla.Set(GetStore(), bIsVanilla);
+	EngineAnalyticsProperties::FProperties::Get().IsVanilla.Set(GetStore(), bIsVanilla);
 	UpdateSessionProgress();
 }
 
 void FEngineAnalyticsSessionSummary::LowDriveSpaceDetected()
 {
-	EngineAnalyticsProperties::IsLowDriveSpace.Set(GetStore(), true);
+	EngineAnalyticsProperties::FProperties::Get().IsLowDriveSpace.Set(GetStore(), true);
 	UpdateSessionProgress();
 }
 
 void FEngineAnalyticsSessionSummary::OnUserActivity(const FUserActivity& InUserActivity)
 {
-	EngineAnalyticsProperties::UserActivity.Set(GetStore(), EngineAnalyticsProperties::GetUserActivityString(InUserActivity));
+	EngineAnalyticsProperties::FProperties::Get().UserActivity.Set(GetStore(), EngineAnalyticsProperties::GetUserActivityString(InUserActivity));
 	UpdateSessionProgress();
 }
 
@@ -334,15 +365,15 @@ void FEngineAnalyticsSessionSummary::UpdateSessionProgress(bool bFlushAsync, con
 	UpdateSessionProgressInternal(bCrashing);
 
 	// Current session duration in seconds.
-	EngineAnalyticsProperties::SessionDurationSecs.Set(GetStore(), FMath::FloorToInt(static_cast<float>(FPlatformTime::Seconds() - SessionStartTimeSecs)),
+	EngineAnalyticsProperties::FProperties::Get().SessionDurationSecs.Set(GetStore(), FMath::FloorToInt(static_cast<float>(FPlatformTime::Seconds() - SessionStartTimeSecs)),
 		[](const int32* Actual, const int32& Proposed) { return Proposed > *Actual; });
 
 	// Wall timestamp.
-	EngineAnalyticsProperties::Timestamp.Set(GetStore(), FDateTime::UtcNow(),
+	EngineAnalyticsProperties::FProperties::Get().Timestamp.Set(GetStore(), FDateTime::UtcNow(),
 		[](const FDateTime* Actual, const FDateTime& Proposed) { return Proposed > *Actual; });
 
 	// The number of times the engine ticked.
-	EngineAnalyticsProperties::EngineTickCount.Set(GetStore(), EngineAnalyticsProperties::GetEngineTickCount(),
+	EngineAnalyticsProperties::FProperties::Get().EngineTickCount.Set(GetStore(), EngineAnalyticsProperties::GetEngineTickCount(),
 		[](const uint64* Actual, const uint64& Proposed) { return Proposed > *Actual; });
 
 	// Flush (or try to flush) the store.
@@ -362,7 +393,7 @@ bool FEngineAnalyticsSessionSummary::UpdateExternalProcessReporterState(bool bQu
 	else if (TOptional<int32> ExitCode = FGenericCrashContext::GetOutOfProcessCrashReporterExitCode())
 	{
 		CrcExitCode = MoveTemp(ExitCode);
-		EngineAnalyticsProperties::MonitorExitCode.Set(GetStore(), *CrcExitCode);
+		EngineAnalyticsProperties::FProperties::Get().MonitorExitCode.Set(GetStore(), *CrcExitCode);
 		UpdateSessionProgress();
 		return true;
 	}
@@ -374,7 +405,7 @@ bool FEngineAnalyticsSessionSummary::UpdateExternalProcessReporterState(bool bQu
 	{
 		// Set a rather unique, but known exit code as place holder, hoping that next update, the engine will report the real one.
 		CrcExitCode.Emplace(ECrashExitCodes::OutOfProcessReporterExitedUnexpectedly);
-		EngineAnalyticsProperties::MonitorExitCode.Set(GetStore(), *CrcExitCode);
+		EngineAnalyticsProperties::FProperties::Get().MonitorExitCode.Set(GetStore(), *CrcExitCode);
 		UpdateSessionProgress();
 		return true;
 	}
@@ -397,9 +428,9 @@ bool FEngineAnalyticsSessionSummary::UpdateDebuggerStates()
 		bStateChanged = true;
 		bDebuggerIgnored = true;
 		bWasEverDebugged = true;
-		EngineAnalyticsProperties::WasDebuggerIgnored.Set(GetStore(), true);
-		EngineAnalyticsProperties::WasEverDebugged.Set(GetStore(), true); // Assumes the session was debugged if the debugger was ignored.
-		EngineAnalyticsProperties::ShutdownType.Set(GetStore(), (int32)EAnalyticsSessionShutdownType::Debugged, EngineAnalyticsProperties::ShouldOverwriteShutdownType);
+		EngineAnalyticsProperties::FProperties::Get().WasDebuggerIgnored.Set(GetStore(), true);
+		EngineAnalyticsProperties::FProperties::Get().WasEverDebugged.Set(GetStore(), true); // Assumes the session was debugged if the debugger was ignored.
+		EngineAnalyticsProperties::FProperties::Get().ShutdownType.Set(GetStore(), (int32)EAnalyticsSessionShutdownType::Debugged, EngineAnalyticsProperties::ShouldOverwriteShutdownType);
 	}
 
 	// Check if the debugger is attached.
@@ -407,8 +438,8 @@ bool FEngineAnalyticsSessionSummary::UpdateDebuggerStates()
 	{
 		bStateChanged = true;
 		bWasEverDebugged = true;
-		EngineAnalyticsProperties::WasEverDebugged.Set(GetStore(), true);
-		EngineAnalyticsProperties::ShutdownType.Set(GetStore(), (int32)EAnalyticsSessionShutdownType::Debugged, EngineAnalyticsProperties::ShouldOverwriteShutdownType);
+		EngineAnalyticsProperties::FProperties::Get().WasEverDebugged.Set(GetStore(), true);
+		EngineAnalyticsProperties::FProperties::Get().ShutdownType.Set(GetStore(), (int32)EAnalyticsSessionShutdownType::Debugged, EngineAnalyticsProperties::ShouldOverwriteShutdownType);
 	}
 
 	return bStateChanged;
