Fix texture streaming for translucent materials in DVSM_OutputMaterialTextureScales.

Disable blending in FDebugViewModeInterface::SetDrawRenderState() for DVSM_OutputMaterialTextureScales to ensure consistent FMaterialTextureInfo::SamplingScale generation regardless of material blend mode.

Previously, blending was enabled for translucent materials, causing incorrect sampling scales, due to a regression in 45836355.

#Synced-CL 51001024

diff --git a/Engine/Source/Runtime/Engine/Private/DebugViewModeInterface.cpp b/Engine/Source/Runtime/Engine/Private/DebugViewModeInterface.cpp
--- a/Engine/Source/Runtime/Engine/Private/DebugViewModeInterface.cpp
+++ b/Engine/Source/Runtime/Engine/Private/DebugViewModeInterface.cpp
@@ -79,6 +79,7 @@ void FDebugViewModeInterface::SetDrawRenderState(EDebugViewShaderMode DebugViewM
 	}
 	else if (DebugViewMode == DVSM_OutputMaterialTextureScales)
 	{
+		DrawRenderState.BlendState = TStaticBlendState<>::GetRHI();
 		DrawRenderState.DepthStencilState = TStaticDepthStencilState<false, CF_Always>::GetRHI();
 	}
 	else if (DebugViewMode == DVSM_ShadowCasters)
