Restore logging in RecreateClothActors to help debug live cloth crashes.

diff --git a/Engine/Source/Runtime/Engine/Private/SkeletalMeshComponentPhysics.cpp b/Engine/Source/Runtime/Engine/Private/SkeletalMeshComponentPhysics.cpp
--- a/Engine/Source/Runtime/Engine/Private/SkeletalMeshComponentPhysics.cpp
+++ b/Engine/Source/Runtime/Engine/Private/SkeletalMeshComponentPhysics.cpp
@@ -3117,6 +3117,8 @@ void USkeletalMeshComponent::RecreateClothingActors()
 		const TArray<UClothingAssetBase*>& MeshClothingAssets = SkelMesh->GetMeshClothingAssets();
 		if (MeshClothingAssets.Num())
 		{
+			UE_LOGF(LogSkeletalMesh, Log, "USkeletalMeshComponent: Recreating Clothing Actors for '%ls' with '%ls'", *GetName(), *SkelMesh->GetName());
+
 			// Group assets per factory, as to only create one simulation per asset of the same type
 			TArray<UClothingAssetBase*> ClothingAssetsInUse;
 			SkelMesh->GetClothingAssetsInUse(ClothingAssetsInUse);
