[PostProcessing] Add experimental field to allow using a Chaos physics object as PostProcessVolume bounds.

#rnx

#Synced-CL 51001024

diff --git a/Engine/Source/Runtime/Engine/Classes/Components/PostProcessComponent.h b/Engine/Source/Runtime/Engine/Classes/Components/PostProcessComponent.h
--- a/Engine/Source/Runtime/Engine/Classes/Components/PostProcessComponent.h
+++ b/Engine/Source/Runtime/Engine/Classes/Components/PostProcessComponent.h
@@ -12,9 +12,11 @@
 #include "Interfaces/Interface_PostProcessVolume.h"
 #include "PostProcessComponent.generated.h"
 
+struct FVolumeBoundsMesh;
+
 /**
  *  PostProcessComponent. Enables Post process controls for blueprints.
- *	Will use a parent UShapeComponent to provide volume data if available.
+ *	Will use a parent UShapeComponent or VolumeBoundsMesh to provide volume data if available.
  */
 
 UCLASS(ClassGroup = Rendering, collapsecategories, hidecategories = (Object), editinlinenew, meta = (BlueprintSpawnableComponent), MinimalAPI)
@@ -45,7 +47,7 @@ class UPostProcessComponent : public USceneComponent, public IInterface_PostProc
 	UPROPERTY(EditAnywhere, BlueprintReadWrite, Category = PostProcessVolume)
 	uint32 bEnabled: 1;
 
-	/** set this to false to use the parent shape component as volume bounds. True affects the whole world regardless.		*/
+	/** set this to false to use the parent shape component or VolumeBoundsMesh as volume bounds. True affects the whole world regardless.		*/
 	UPROPERTY(EditAnywhere, BlueprintReadWrite, Category = PostProcessVolume)
 	uint32 bUnbound: 1;
  
@@ -53,18 +55,25 @@ class UPostProcessComponent : public USceneComponent, public IInterface_PostProc
 	UPROPERTY(TextExportTransient, NonPIEDuplicateTransient, NonTransactional)
 	FGuid VolumeGuid;
 
+	// If set, this mesh will be used as bounds for the volume. 
+	TPimplPtr<FVolumeBoundsMesh> VolumeBoundsMesh;
+	
+	TOptional<FBoxSphereBounds> GetPostProcessBounds() const;
+
 	//~ Begin IInterface_PostProcessVolume Interface
 	ENGINE_API virtual bool EncompassesPoint(FVector Point, float SphereRadius/*=0.f*/, float* OutDistanceToPoint) override;
 	virtual FPostProcessVolumeProperties GetProperties() const override
 	{
+		TOptional<FBoxSphereBounds> PPBounds = GetPostProcessBounds();
+
 		FPostProcessVolumeProperties Ret;
 		Ret.bIsEnabled = bEnabled != 0 && ShouldRender();
-		Ret.bIsUnbound = bUnbound != 0 || Cast<UShapeComponent>(GetAttachParent()) == nullptr;
+		Ret.bIsUnbound = bUnbound != 0 || !PPBounds.IsSet();
 		Ret.BlendRadius = BlendRadius;
 		Ret.BlendWeight = BlendWeight;
 		Ret.Priority = Priority;
 		Ret.Settings = &Settings;
-		Ret.Size = Ret.bIsUnbound ? DBL_MAX : Cast<UShapeComponent>(GetAttachParent())->CalcLocalBounds().GetBox().GetVolume();
+		Ret.Size = PPBounds.IsSet() ? PPBounds->GetBox().GetVolume() : DBL_MAX;
 		Ret.VolumeGuid = VolumeGuid;
 		return Ret;
 	}
diff --git a/Engine/Source/Runtime/Engine/Classes/Interfaces/Interface_PostProcessVolume.h b/Engine/Source/Runtime/Engine/Classes/Interfaces/Interface_PostProcessVolume.h
--- a/Engine/Source/Runtime/Engine/Classes/Interfaces/Interface_PostProcessVolume.h
+++ b/Engine/Source/Runtime/Engine/Classes/Interfaces/Interface_PostProcessVolume.h
@@ -39,6 +39,9 @@ class IInterface_PostProcessVolume
 {
 	GENERATED_IINTERFACE_BODY()
 
+	// Returns true if the sphere around Point with radius SphereRadius intersects with the volume
+	// OutDistanceToPoint is set to the distance to the nearest point on the volume,
+	// or may be set to FLT_MAX if Point is outside the blending range of the volume.
 	virtual bool EncompassesPoint(FVector Point, float SphereRadius/*=0.f*/, float* OutDistanceToPoint) = 0;
 	virtual FPostProcessVolumeProperties GetProperties() const = 0;
 #if DEBUG_POST_PROCESS_VOLUME_ENABLE
diff --git a/Engine/Source/Runtime/Engine/Internal/VolumeBoundsMesh.h b/Engine/Source/Runtime/Engine/Internal/VolumeBoundsMesh.h
new file mode 100644
--- /dev/null
+++ b/Engine/Source/Runtime/Engine/Internal/VolumeBoundsMesh.h
@@ -0,0 +1,15 @@
+// Copyright Epic Games, Inc. All Rights Reserved.
+
+#pragma once
+
+#include "CoreMinimal.h"
+#include "PhysicsEngine/PhysicsObjectExternalInterface.h"
+
+struct FVolumeBoundsMesh
+{
+	// Hack to make this compatible with both the actor framework and SceneGraph,
+	// and avoid holding onto stale FPhysicsObject*.
+	// This function should hold a weak ptr on the owner of the physics body, 
+	// and not cache the physics object handles directly, as those may be invalidated.
+	TFunction<TArray<Chaos::FPhysicsObjectHandle>()> PhysicsObjectsAccessor;
+};
diff --git a/Engine/Source/Runtime/Engine/Private/Components/PostProcessComponent.cpp b/Engine/Source/Runtime/Engine/Private/Components/PostProcessComponent.cpp
--- a/Engine/Source/Runtime/Engine/Private/Components/PostProcessComponent.cpp
+++ b/Engine/Source/Runtime/Engine/Private/Components/PostProcessComponent.cpp
@@ -3,6 +3,13 @@
 #include "Components/PostProcessComponent.h"
 #include "Components/SphereComponent.h"
 #include "Engine/World.h"
+#include "VolumeBoundsMesh.h"
+#include "Templates/PimplPtr.h"
+
+#include "PhysicsEngine/PhysicsObjectExternalInterface.h"
+#include "Chaos/ImplicitObjectScaled.h"
+#include "Chaos/Levelset.h"
+#include "Chaos/MLLevelset.h"
 
 #include UE_INLINE_GENERATED_CPP_BY_NAME(PostProcessComponent)
 
@@ -40,6 +47,34 @@ void UPostProcessComponent::Serialize(FArchive& Ar)
 	}
 }
 
+TOptional<FBoxSphereBounds> UPostProcessComponent::GetPostProcessBounds() const
+{
+	if (!bUnbound)
+	{
+		TArray<Chaos::FPhysicsObjectHandle> VolumeBoundsPhysicsObjects {};
+		if (VolumeBoundsMesh && VolumeBoundsMesh->PhysicsObjectsAccessor)
+		{
+			VolumeBoundsPhysicsObjects = VolumeBoundsMesh->PhysicsObjectsAccessor();
+		}
+
+		if (!VolumeBoundsPhysicsObjects.IsEmpty())
+		{
+			FLockedReadPhysicsObjectExternalInterface Interface = FPhysicsObjectExternalInterface::LockRead(VolumeBoundsPhysicsObjects);
+			FBox BoxBounds = Interface->GetWorldBounds(VolumeBoundsPhysicsObjects);
+			return FBoxSphereBounds(BoxBounds);
+		}
+		else if (const UPrimitiveComponent* Primitive = Cast<UPrimitiveComponent>(GetAttachParent()))
+		{
+			return Primitive->GetBounds();
+		}
+		else if (const UShapeComponent* Shape = Cast<UShapeComponent>(GetAttachParent()))
+		{
+			return Shape->GetBounds();
+		}
+	}
+	return NullOpt;
+}
+
 void UPostProcessComponent::PostInitProperties()
 {
 	Super::PostInitProperties();
@@ -52,10 +87,138 @@ void UPostProcessComponent::PostInitProperties()
 #endif
 }
 
+// Distance queries require manifold geometry to get meaningful results.
+// If a piece of geometry is not 'watertight', you cannot tell what's inside vs outside.
+// This function filters out triangle mesh geometry representations, which may be non-manifold.
+static bool IsGuaranteedManifoldGeometry(const Chaos::FImplicitObject* Geometry)
+{
+	if (!Geometry)
+	{
+		return false;
+	}
+	else if (const Chaos::FImplicitObjectInstanced* ObjInstanced = Geometry->AsA<Chaos::FImplicitObjectInstanced>())
+	{
+		return IsGuaranteedManifoldGeometry(ObjInstanced->GetInnerObject().Get());
+	}
+	else if (const Chaos::FImplicitObjectScaled* ObjScaled = Geometry->AsA<Chaos::FImplicitObjectScaled>())
+	{
+		return IsGuaranteedManifoldGeometry(ObjScaled->GetInnerObject().Get());
+	}
+	else if (const Chaos::TImplicitObjectTransformed<Chaos::FReal, 3>* ObjTransformed = Geometry->AsA<Chaos::TImplicitObjectTransformed<Chaos::FReal, 3>>())
+	{
+		return IsGuaranteedManifoldGeometry(ObjTransformed->GetGeometry());
+	}
+	else if (const Chaos::FImplicitObjectUnion* ObjUnion = Geometry->AsA<Chaos::FImplicitObjectUnion>())
+	{
+		for (const Chaos::FImplicitObject* Element : ObjUnion->GetObjects())
+		{
+			if (!IsGuaranteedManifoldGeometry(Element))
+			{
+				return false;
+			}
+		}
+		return true;
+	}
+	else if (Geometry->IsConvex()
+		|| Geometry->IsA<Chaos::FLevelSet>()
+		|| Geometry->IsA<Chaos::FMLLevelSet>())
+	{
+		return true;
+	}
+	return false;
+}
+
+template<typename TVisitor>
+void VisitGuaranteedManifoldGeometry(TArray<Chaos::FPhysicsObjectHandle>& PhysicsObjects, TVisitor Visitor)
+{
+	FLockedReadPhysicsObjectExternalInterface Interface = FPhysicsObjectExternalInterface::LockRead(PhysicsObjects);
+	Interface->VisitEveryShape(PhysicsObjects, [Visitor](const Chaos::FConstPhysicsObjectHandle Handle, Chaos::TThreadShapeInstance<Chaos::EThreadContext::External>* ShapeInstance) -> bool
+	{
+		const Chaos::FImplicitObjectRef Geometry = ShapeInstance->GetGeometry();
+		if (IsGuaranteedManifoldGeometry(Geometry))
+		{
+			Visitor(Handle, Geometry);
+		}
+		return false;
+	});
+}
+
+static FClosestPhysicsObjectResult FindApproximateDistanceTo(TArray<Chaos::FPhysicsObjectHandle>& PhysicsObjects, FVector WorldLocation, FLockedReadPhysicsObjectExternalInterface& Interface)
+{
+	FClosestPhysicsObjectResult AggregateResult {};
+	VisitGuaranteedManifoldGeometry(PhysicsObjects, [&AggregateResult, &Interface, WorldLocation](const Chaos::FConstPhysicsObjectHandle Object, const Chaos::FImplicitObjectRef Geometry)
+	{
+		const FTransform WorldTransform = Interface->GetTransform(Object);
+		const FVector LocalLocation = WorldTransform.InverseTransformPosition(WorldLocation);
+
+		FClosestPhysicsObjectResult Result;
+
+		Result.PhysicsObject = const_cast<Chaos::FPhysicsObjectHandle>(Object);
+
+		Chaos::FVec3 Normal;
+		Result.ClosestDistance = static_cast<double>(Geometry->PhiWithNormal(LocalLocation, Normal));
+		Result.ClosestLocation = WorldTransform.TransformPosition(LocalLocation - Result.ClosestDistance * Normal);
+
+		if (!AggregateResult || Result.ClosestDistance < AggregateResult.ClosestDistance)
+		{
+			AggregateResult = Result;
+		}
+	});
+	return AggregateResult;
+}
+
 bool UPostProcessComponent::EncompassesPoint(FVector Point, float SphereRadius/*=0.f*/, float* OutDistanceToPoint)
 {
-	UShapeComponent* ParentShape = Cast<UShapeComponent>(GetAttachParent());
-	if (ParentShape != nullptr)
+	const TOptional<FBoxSphereBounds> PPBounds = GetPostProcessBounds();
+	if (!PPBounds.IsSet()) // global volume encompasses all
+	{
+		if (OutDistanceToPoint)
+		{
+			*OutDistanceToPoint = 0;
+		}
+		return true;
+	}
+
+	const bool bCouldBeInBlendingRange = PPBounds->GetBox().ExpandBy(BlendRadius + SphereRadius).IsInsideOrOn(Point);
+	if (!bCouldBeInBlendingRange)
+	{
+		if (OutDistanceToPoint)
+		{
+			*OutDistanceToPoint = FLT_MAX;
+		}
+		return false;
+	}
+
+	TArray<Chaos::FPhysicsObjectHandle> ActiveVolumeBoundsPhysicsObjects;
+	if (VolumeBoundsMesh && VolumeBoundsMesh->PhysicsObjectsAccessor)
+	{
+		ActiveVolumeBoundsPhysicsObjects = VolumeBoundsMesh->PhysicsObjectsAccessor();
+	}
+	else if (UPrimitiveComponent* ParentPrimitive = Cast<UPrimitiveComponent>(GetAttachParent()))
+	{
+		ActiveVolumeBoundsPhysicsObjects = ParentPrimitive->GetAllPhysicsObjects();
+	}
+
+	if (!ActiveVolumeBoundsPhysicsObjects.IsEmpty())
+	{
+		FLockedReadPhysicsObjectExternalInterface Interface = FPhysicsObjectExternalInterface::LockRead(ActiveVolumeBoundsPhysicsObjects);
+		FClosestPhysicsObjectResult QueryResult = FindApproximateDistanceTo(ActiveVolumeBoundsPhysicsObjects, Point, Interface);
+		float Distance = FMath::Max(0.0f, static_cast<float>(QueryResult.ClosestDistance));
+		if (QueryResult.PhysicsObject == nullptr) // no valid physics object found
+		{
+			if (OutDistanceToPoint)
+			{
+				*OutDistanceToPoint = FLT_MAX;
+			}
+			return false;
+		}
+		if (OutDistanceToPoint)
+		{
+			*OutDistanceToPoint = Distance;
+		}
+		return Distance <= SphereRadius;
+	}
+	else if (UShapeComponent* ParentShape = Cast<UShapeComponent>(GetAttachParent()))
 	{
 		float Distance = -1.f;
 
