Additional Engine level MetaSound data interfaces:
- UE.Listener.Location
- UE.Source.Location
- UE.Source.Occlusion
- UE.World.Time

[RN] Additional Engine level MetaSound data interfaces:- UE.Listener.Location- UE.Source.Location- UE.Source.Occlusion- UE.World.Time

#Synced-CL 51036651

diff --git a/Engine/Source/Runtime/Engine/Classes/Sound/SoundAttenuation.h b/Engine/Source/Runtime/Engine/Classes/Sound/SoundAttenuation.h
--- a/Engine/Source/Runtime/Engine/Classes/Sound/SoundAttenuation.h
+++ b/Engine/Source/Runtime/Engine/Classes/Sound/SoundAttenuation.h
@@ -500,4 +500,58 @@ namespace Audio
 
 		ENGINE_API Audio::FParameterInterfacePtr GetInterface();
 	} // namespace EmitterInterface
+
+	namespace SourceLocationInterface
+	{
+		ENGINE_API const extern FName Name;
+
+		namespace Inputs
+		{
+			ENGINE_API const extern FName X;
+			ENGINE_API const extern FName Y;
+			ENGINE_API const extern FName Z;
+		} // namespace Inputs
+
+		ENGINE_API Audio::FParameterInterfacePtr GetInterface();
+	} // namespace SourceLocationInterface
+
+	namespace ListenerLocationInterface
+	{
+		ENGINE_API const extern FName Name;
+
+		namespace Inputs
+		{
+			ENGINE_API const extern FName X;
+			ENGINE_API const extern FName Y;
+			ENGINE_API const extern FName Z;
+		} // namespace Inputs
+
+		ENGINE_API Audio::FParameterInterfacePtr GetInterface();
+	} // namespace ListenerLocationInterface
+
+	namespace SourceOcclusionInterface
+	{
+		ENGINE_API const extern FName Name;
+
+		namespace Inputs
+		{
+			ENGINE_API const extern FName IsOccluded;
+		} // namespace Inputs
+
+		ENGINE_API Audio::FParameterInterfacePtr GetInterface();
+	} // namespace SourceOcclusionInterface
+
+	namespace WorldTimeInterface
+	{
+		ENGINE_API const extern FName Name;
+
+		namespace Inputs
+		{
+			ENGINE_API const extern FName WorldTimeWhenPlayed;
+			ENGINE_API const extern FName CurrentWorldTime;
+			ENGINE_API const extern FName ElapsedWorldTime;
+		} // namespace Inputs
+
+		ENGINE_API Audio::FParameterInterfacePtr GetInterface();
+	} // namespace WorldTimeInterface
 } // namespace Audio
diff --git a/Engine/Source/Runtime/Engine/Private/ActiveSound.cpp b/Engine/Source/Runtime/Engine/Private/ActiveSound.cpp
--- a/Engine/Source/Runtime/Engine/Private/ActiveSound.cpp
+++ b/Engine/Source/Runtime/Engine/Private/ActiveSound.cpp
@@ -843,7 +843,7 @@ void FActiveSound::UpdateInterfaceParameters(const TArray<FListener>& InListener
 		return;
 	}
 
-	if (!Sound || IsPreviewSound())
+	if (!Sound)
 	{
 		return;
 	}
@@ -857,12 +857,20 @@ void FActiveSound::UpdateInterfaceParameters(const TArray<FListener>& InListener
 	FParameterInterfacePtr SpatializationInterface = SpatializationInterface::GetInterface();
 	FParameterInterfacePtr SourceOrientationInterface = SourceOrientationInterface::GetInterface();
 	FParameterInterfacePtr ListenerOrientationInterface = ListenerOrientationInterface::GetInterface();
+	FParameterInterfacePtr SourceLocationInterface = SourceLocationInterface::GetInterface();
+	FParameterInterfacePtr ListenerLocationInterface = ListenerLocationInterface::GetInterface();
+	FParameterInterfacePtr SourceOcclusionInterface = SourceOcclusionInterface::GetInterface();
+	FParameterInterfacePtr WorldTimeInterface = WorldTimeInterface::GetInterface();
 
 	bool bImplementsAttenuation = Sound->ImplementsParameterInterface(AttenuationInterface);
 	bool bImplementsSpatialization = Sound->ImplementsParameterInterface(SpatializationInterface);
 	bool bImplementsSourceOrientation = Sound->ImplementsParameterInterface(SourceOrientationInterface);
 	bool bImplementsListenerOrientation = Sound->ImplementsParameterInterface(ListenerOrientationInterface);
-	
+	bool bImplementsSourceLocation = Sound->ImplementsParameterInterface(SourceLocationInterface);
+	bool bImplementsListenerLocation = Sound->ImplementsParameterInterface(ListenerLocationInterface);
+	bool bImplementsSourceOcclusion = Sound->ImplementsParameterInterface(SourceOcclusionInterface);
+	bool bImplementsWorldTime = Sound->ImplementsParameterInterface(WorldTimeInterface);
+
 	for (const TPair<UPTRINT, FWaveInstance*>& InstancePair : WaveInstances)
 	{
 		const FWaveInstance* Instance = InstancePair.Value;
@@ -878,63 +886,126 @@ void FActiveSound::UpdateInterfaceParameters(const TArray<FListener>& InListener
 			bImplementsSpatialization |= SoundInstance->ImplementsParameterInterface(SpatializationInterface);
 			bImplementsSourceOrientation |= SoundInstance->ImplementsParameterInterface(SourceOrientationInterface);
 			bImplementsListenerOrientation |= SoundInstance->ImplementsParameterInterface(ListenerOrientationInterface);
+			bImplementsSourceLocation |= SoundInstance->ImplementsParameterInterface(SourceLocationInterface);
+			bImplementsListenerLocation |= SoundInstance->ImplementsParameterInterface(ListenerLocationInterface);
+			bImplementsSourceOcclusion |= SoundInstance->ImplementsParameterInterface(SourceOcclusionInterface);
+			bImplementsWorldTime |= SoundInstance->ImplementsParameterInterface(WorldTimeInterface);
 		}
 	}
 
-	if (false == (bImplementsAttenuation || bImplementsSpatialization || bImplementsSourceOrientation || bImplementsListenerOrientation))
+	if (false == (bImplementsAttenuation || bImplementsSpatialization || bImplementsSourceOrientation || bImplementsListenerOrientation 
+		|| bImplementsSourceLocation || bImplementsListenerLocation || bImplementsSourceOcclusion || bImplementsWorldTime))
 	{
 		return;
 	}
 
 	TArray<FAudioParameter> ParamsToUpdate;
 
-	const FListener& Listener = InListeners[ClosestListenerIndex];
-	const FVector SourceDirection = Transform.GetLocation() - Listener.Transform.GetLocation();
-	
-	if (bImplementsAttenuation)
+	// World time is available in preview, so we can add params to the update
+	if (bImplementsWorldTime)
 	{
-		const float Distance = SourceDirection.Size();
-		ParamsToUpdate.Add({ AttenuationInterface::Inputs::Distance, Distance });
+		const float CurrentWorldTimeSeconds = World.Get() ? World->GetTimeSeconds() : FApp::GetGameTime();
+		const float ElapsedWorldTimeSeconds = PlaybackTimeUnscaled;
+		const float WorldTimeWhenPlayedSeconds = CurrentWorldTimeSeconds - PlaybackTimeUnscaled;
+		
+		ParamsToUpdate.Append(
+			{
+				{ WorldTimeInterface::Inputs::WorldTimeWhenPlayed,  WorldTimeWhenPlayedSeconds},
+				{ WorldTimeInterface::Inputs::CurrentWorldTime,  CurrentWorldTimeSeconds},
+				{ WorldTimeInterface::Inputs::ElapsedWorldTime,  ElapsedWorldTimeSeconds}
+			});
 	}
 
-	if (bImplementsSpatialization)
+	// If source is not a Preview Sound, resume adding Audio Params related to the Listener's 3D environment
+	if (!IsPreviewSound())
 	{
-		const FVector SourceDirectionNormal = Listener.Transform.InverseTransformVectorNoScale(SourceDirection).GetSafeNormal();
-		const FVector2D SourceAzimuthAndElevation = FMath::GetAzimuthAndElevation(SourceDirectionNormal, FVector::ForwardVector, FVector::RightVector, FVector::UpVector);
-		const float Azimuth = FMath::RadiansToDegrees(SourceAzimuthAndElevation.X);
-		const float Elevation = FMath::RadiansToDegrees(SourceAzimuthAndElevation.Y);
-		ParamsToUpdate.Append(
+		const FListener& Listener = InListeners[ClosestListenerIndex];
+		const FVector ListenerLocation = Listener.Transform.GetLocation();
+		const FVector SourceLocation = Transform.GetLocation();
+		const FVector SourceDirection = SourceLocation - ListenerLocation;
+		const bool bSourceOccluded = bIsOccluded;
+
+		if (bImplementsAttenuation)
 		{
-			{ SpatializationInterface::Inputs::Azimuth, Azimuth },
-			{ SpatializationInterface::Inputs::Elevation, Elevation }
-		});
-	}
+			const float Distance = SourceDirection.Size();
+			ParamsToUpdate.Add({ AttenuationInterface::Inputs::Distance, Distance });
+		}
 
-	if (bImplementsSourceOrientation)
-	{
-		const FVector ListenerDirectionNormal = Transform.InverseTransformVectorNoScale(-SourceDirection).GetSafeNormal();
-		const FVector2D ListenerAzimuthAndElevation = FMath::GetAzimuthAndElevation(ListenerDirectionNormal, FVector::ForwardVector, FVector::RightVector, FVector::UpVector);
-		const float Azimuth = FMath::RadiansToDegrees(ListenerAzimuthAndElevation.X);
-		const float Elevation = FMath::RadiansToDegrees(ListenerAzimuthAndElevation.Y);
-		ParamsToUpdate.Append(
+		if (bImplementsSpatialization)
 		{
-			{ SourceOrientationInterface::Inputs::Azimuth, Azimuth },
-			{ SourceOrientationInterface::Inputs::Elevation, Elevation }
-		});
-	}
+			const FVector SourceDirectionNormal = Listener.Transform.InverseTransformVectorNoScale(SourceDirection).GetSafeNormal();
+			const FVector2D SourceAzimuthAndElevation = FMath::GetAzimuthAndElevation(SourceDirectionNormal, FVector::ForwardVector, FVector::RightVector, FVector::UpVector);
+			const float Azimuth = FMath::RadiansToDegrees(SourceAzimuthAndElevation.X);
+			const float Elevation = FMath::RadiansToDegrees(SourceAzimuthAndElevation.Y);
+			ParamsToUpdate.Append(
+				{
+					{ SpatializationInterface::Inputs::Azimuth, Azimuth },
+					{ SpatializationInterface::Inputs::Elevation, Elevation }
+				});
+		}
 
-	if (bImplementsListenerOrientation)
-	{
-		const FVector ListenerDirectionNormal = Listener.Transform.GetRotation().GetForwardVector();
-		const FVector2D ListenerAzimuthAndElevation = FMath::GetAzimuthAndElevation(ListenerDirectionNormal, FVector::ForwardVector, FVector::RightVector, FVector::UpVector);
-		const float Azimuth = FMath::RadiansToDegrees(ListenerAzimuthAndElevation.X);
-		const float Elevation = FMath::RadiansToDegrees(ListenerAzimuthAndElevation.Y);
-		
-		ParamsToUpdate.Append(
+		if (bImplementsSourceOrientation)
+		{
+			const FVector ListenerDirectionNormal = Transform.InverseTransformVectorNoScale(-SourceDirection).GetSafeNormal();
+			const FVector2D ListenerAzimuthAndElevation = FMath::GetAzimuthAndElevation(ListenerDirectionNormal, FVector::ForwardVector, FVector::RightVector, FVector::UpVector);
+			const float Azimuth = FMath::RadiansToDegrees(ListenerAzimuthAndElevation.X);
+			const float Elevation = FMath::RadiansToDegrees(ListenerAzimuthAndElevation.Y);
+			ParamsToUpdate.Append(
+				{
+					{ SourceOrientationInterface::Inputs::Azimuth, Azimuth },
+					{ SourceOrientationInterface::Inputs::Elevation, Elevation }
+				});
+		}
+
+		if (bImplementsListenerOrientation)
+		{
+			const FVector ListenerDirectionNormal = Listener.Transform.GetRotation().GetForwardVector();
+			const FVector2D ListenerAzimuthAndElevation = FMath::GetAzimuthAndElevation(ListenerDirectionNormal, FVector::ForwardVector, FVector::RightVector, FVector::UpVector);
+			const float Azimuth = FMath::RadiansToDegrees(ListenerAzimuthAndElevation.X);
+			const float Elevation = FMath::RadiansToDegrees(ListenerAzimuthAndElevation.Y);
+
+			ParamsToUpdate.Append(
+				{
+					{ ListenerOrientationInterface::Inputs::Azimuth, Azimuth },
+					{ ListenerOrientationInterface::Inputs::Elevation, Elevation }
+				});
+		}
+
+		if (bImplementsSourceLocation)
 		{
-			{ ListenerOrientationInterface::Inputs::Azimuth, Azimuth },
-			{ ListenerOrientationInterface::Inputs::Elevation, Elevation }
-		});
+			const float SourceX = SourceLocation.X;
+			const float SourceY = SourceLocation.Y;
+			const float SourceZ = SourceLocation.Z;
+
+			ParamsToUpdate.Append(
+				{
+					{ SourceLocationInterface::Inputs::X, SourceX },
+					{ SourceLocationInterface::Inputs::Y, SourceY },
+					{ SourceLocationInterface::Inputs::Z, SourceZ }
+				});
+		}
+
+		if (bImplementsListenerLocation)
+		{
+			const float ListenerX = ListenerLocation.X;
+			const float ListenerY = ListenerLocation.Y;
+			const float ListenerZ = ListenerLocation.Z;
+
+			ParamsToUpdate.Append(
+				{
+					{ ListenerLocationInterface::Inputs::X, ListenerX },
+					{ ListenerLocationInterface::Inputs::Y, ListenerY },
+					{ ListenerLocationInterface::Inputs::Z, ListenerZ }
+				});
+		}
+
+		if (bImplementsSourceOcclusion)
+		{
+			ParamsToUpdate.Append(
+				{
+					{ SourceOcclusionInterface::Inputs::IsOccluded, bSourceOccluded }
+				});
+		}
 	}
 
 	InstanceTransmitter->SetParameters(MoveTemp(ParamsToUpdate));
diff --git a/Engine/Source/Runtime/Engine/Private/AudioSettings.cpp b/Engine/Source/Runtime/Engine/Private/AudioSettings.cpp
--- a/Engine/Source/Runtime/Engine/Private/AudioSettings.cpp
+++ b/Engine/Source/Runtime/Engine/Private/AudioSettings.cpp
@@ -289,6 +289,10 @@ void UAudioSettings::RegisterParameterInterfaces()
 		InterfaceRegistry.RegisterInterface(Audio::SpatializationInterface::GetInterface());
 		InterfaceRegistry.RegisterInterface(Audio::SourceOrientationInterface::GetInterface());
 		InterfaceRegistry.RegisterInterface(Audio::ListenerOrientationInterface::GetInterface());
+		InterfaceRegistry.RegisterInterface(Audio::SourceLocationInterface::GetInterface());
+		InterfaceRegistry.RegisterInterface(Audio::ListenerLocationInterface::GetInterface());
+		InterfaceRegistry.RegisterInterface(Audio::SourceOcclusionInterface::GetInterface());
+		InterfaceRegistry.RegisterInterface(Audio::WorldTimeInterface::GetInterface());
 	}
 }
 
diff --git a/Engine/Source/Runtime/Engine/Private/SoundAttenuation.cpp b/Engine/Source/Runtime/Engine/Private/SoundAttenuation.cpp
--- a/Engine/Source/Runtime/Engine/Private/SoundAttenuation.cpp
+++ b/Engine/Source/Runtime/Engine/Private/SoundAttenuation.cpp
@@ -464,5 +464,224 @@ namespace Audio
 		}
 	} // namespace ListenerOrientationInterface
 #undef AUDIO_PARAMETER_INTERFACE_NAMESPACE
+
+#define AUDIO_PARAMETER_INTERFACE_NAMESPACE "UE.Source.Location"
+	namespace SourceLocationInterface
+	{
+		const FName Name = AUDIO_PARAMETER_INTERFACE_NAMESPACE;
+
+		namespace Inputs
+		{
+			const FName X = AUDIO_PARAMETER_INTERFACE_MEMBER_DEFINE("X");
+			const FName Y = AUDIO_PARAMETER_INTERFACE_MEMBER_DEFINE("Y");
+			const FName Z = AUDIO_PARAMETER_INTERFACE_MEMBER_DEFINE("Z");
+		} // namespace Inputs
+
+		Audio::FParameterInterfacePtr GetInterface()
+		{
+			struct FInterface : public Audio::FParameterInterface
+			{
+				FInterface()
+					: FParameterInterface(SourceLocationInterface::Name, { 1, 0 })
+				{
+					Inputs =
+					{
+						{
+							FText(),
+							NSLOCTEXT("SourceLocation", "XDescription", "X-axis of Source's World Location."),
+							FName(),
+							{ Inputs::X, 0.0f }
+						},
+						{
+							FText(),
+							NSLOCTEXT("SourceLocation", "YDescription", "Y-axis of Source's World Location."),
+							FName(),
+							{ Inputs::Y, 0.0f }
+						},
+						{
+							FText(),
+							NSLOCTEXT("SourceLocation", "ZDescription", "Z-axis of Source's World Location."),
+							FName(),
+							{ Inputs::Z, 0.0f }
+						}
+					};
+#if WITH_EDITOR
+					Description = LOCTEXT("UE.Source.Location.Description",
+						"Provides information on the world location of the Source.");
+#endif
+				}
+			};
+
+			static FParameterInterfacePtr InterfacePtr;
+			if (!InterfacePtr.IsValid())
+			{
+				InterfacePtr = MakeShared<FInterface>();
+			}
+
+			return InterfacePtr;
+		}
+	} // namespace SourceLocationInterface
+#undef AUDIO_PARAMETER_INTERFACE_NAMESPACE
+
+#define AUDIO_PARAMETER_INTERFACE_NAMESPACE "UE.Listener.Location"
+	namespace ListenerLocationInterface
+	{
+		const FName Name = AUDIO_PARAMETER_INTERFACE_NAMESPACE;
+
+		namespace Inputs
+		{
+			const FName X = AUDIO_PARAMETER_INTERFACE_MEMBER_DEFINE("X");
+			const FName Y = AUDIO_PARAMETER_INTERFACE_MEMBER_DEFINE("Y");
+			const FName Z = AUDIO_PARAMETER_INTERFACE_MEMBER_DEFINE("Z");
+		} // namespace Inputs
+
+		Audio::FParameterInterfacePtr GetInterface()
+		{
+			struct FInterface : public Audio::FParameterInterface
+			{
+				FInterface()
+					: FParameterInterface(ListenerLocationInterface::Name, { 1, 0 })
+				{
+					Inputs =
+					{
+						{
+							FText(),
+							NSLOCTEXT("ListenerLocation", "XDescription", "X-axis of Listener's World Location."),
+							FName(),
+							{ Inputs::X, 0.0f }
+						},
+						{
+							FText(),
+							NSLOCTEXT("ListenerLocation", "YDescription", "Y-axis of Listener's World Location."),
+							FName(),
+							{ Inputs::Y, 0.0f }
+						},
+						{
+							FText(),
+							NSLOCTEXT("ListenerLocation", "ZDescription", "Z-axis of Listener's World Location."),
+							FName(),
+							{ Inputs::Z, 0.0f }
+						}
+					};
+#if WITH_EDITOR
+					Description = LOCTEXT("UE.Listener.Location.Description",
+						"Provides information on the world location of the Listener.");
+#endif
+				}
+			};
+
+			static FParameterInterfacePtr InterfacePtr;
+			if (!InterfacePtr.IsValid())
+			{
+				InterfacePtr = MakeShared<FInterface>();
+			}
+
+			return InterfacePtr;
+		}
+	} // namespace ListenerLocationInterface
+#undef AUDIO_PARAMETER_INTERFACE_NAMESPACE
+
+#define AUDIO_PARAMETER_INTERFACE_NAMESPACE "UE.Source.Occlusion"
+	namespace SourceOcclusionInterface
+	{
+		const FName Name = AUDIO_PARAMETER_INTERFACE_NAMESPACE;
+
+		namespace Inputs
+		{
+			const FName IsOccluded = AUDIO_PARAMETER_INTERFACE_MEMBER_DEFINE("IsOccluded");
+		}
+
+		Audio::FParameterInterfacePtr GetInterface()
+		{
+			struct FInterface : public Audio::FParameterInterface
+			{
+				FInterface()
+					: FParameterInterface(SourceOcclusionInterface::Name, { 1, 0 })
+				{
+					Inputs =
+					{
+						{
+							FText(),
+							NSLOCTEXT("IsOccluded", "IsOccludedDescription", "Whether or not this sound is occluded (configure via AttenuationSettings)."),
+							FName(),
+							{ Inputs::IsOccluded, false }
+						}
+					};
+#if WITH_EDITOR
+					Description = LOCTEXT("UE.Source.Occlusion.Description",
+						"Provides information around Occlusion as configured by this source's Attenuation Settings.");
+#endif
+				}
+			};
+
+			static FParameterInterfacePtr InterfacePtr;
+			if (!InterfacePtr.IsValid())
+			{
+				InterfacePtr = MakeShared<FInterface>();
+			}
+
+			return InterfacePtr;
+		}
+	} // namespace SourceOcclusionInterface
+#undef AUDIO_PARAMETER_INTERFACE_NAMESPACE
+
+#define AUDIO_PARAMETER_INTERFACE_NAMESPACE "UE.World.Time"
+	namespace WorldTimeInterface
+	{
+		const FName Name = AUDIO_PARAMETER_INTERFACE_NAMESPACE;
+
+		namespace Inputs
+		{
+			const FName WorldTimeWhenPlayed = AUDIO_PARAMETER_INTERFACE_MEMBER_DEFINE("WorldTimeWhenPlayed");
+			const FName CurrentWorldTime = AUDIO_PARAMETER_INTERFACE_MEMBER_DEFINE("CurrentWorldTime");
+			const FName ElapsedWorldTime = AUDIO_PARAMETER_INTERFACE_MEMBER_DEFINE("ElapsedWorldTime");
+		}
+
+		Audio::FParameterInterfacePtr GetInterface()
+		{
+			struct FInterface : public Audio::FParameterInterface
+			{
+				FInterface()
+					: FParameterInterface(WorldTimeInterface::Name, { 1, 0 })
+				{
+					Inputs =
+					{
+						{
+							FText(),
+							NSLOCTEXT("WorldTime", "WorldTimeWhenPlayedDescription", "[0, X] World time (in seconds) the moment this sound was played."),
+							FName(),
+							{ Inputs::WorldTimeWhenPlayed, 0.0f }
+						},
+						{
+							FText(),
+							NSLOCTEXT("WorldTime", "CurrentWorldTimeDescription", "[0, X] The current value of GameTimeSeconds from the World this sound belongs to."),
+							FName(),
+							{ Inputs::CurrentWorldTime, 0.0f }
+						},
+						{
+							FText(),
+							NSLOCTEXT("WorldTime", "ElapsedWorldTimeDescription", "[0, X] How long (in seconds) this asset has been playing for."),
+							FName(),
+							{ Inputs::ElapsedWorldTime, 0.0f }
+						}
+					};
+#if WITH_EDITOR
+					Description = LOCTEXT("UE.World.Time.Description",
+						"Provides information related to the game various time measurements, like how long since this asset was last played or the elapsed time of the World this sound belongs to.");
+#endif
+				}
+			};
+
+			static FParameterInterfacePtr InterfacePtr;
+			if (!InterfacePtr.IsValid())
+			{
+				InterfacePtr = MakeShared<FInterface>();
+			}
+
+			return InterfacePtr;
+		}
+	} // WorldTimeInterface
+#undef AUDIO_PARAMETER_INTERFACE_NAMESPACE
+
 } // namespace Audio
 #undef LOCTEXT_NAMESPACE
