Add LogMessage interface to FComponentSpacePoseContext

#Synced-CL 51161208

diff --git a/Engine/Source/Runtime/Engine/Classes/Animation/AnimNodeBase.h b/Engine/Source/Runtime/Engine/Classes/Animation/AnimNodeBase.h
--- a/Engine/Source/Runtime/Engine/Classes/Animation/AnimNodeBase.h
+++ b/Engine/Source/Runtime/Engine/Classes/Animation/AnimNodeBase.h
@@ -624,6 +624,10 @@ public:
 	// Note: this copy assignment operator copies the whole object but the copy constructor only copies part of the object.
 	FComponentSpacePoseContext& operator=(const FComponentSpacePoseContext&) = default;
 
+	// Log evaluation message
+	void LogMessage(const TSharedRef<FTokenizedMessage>& InMessage) const { LogMessageInternal("Evaluate", InMessage); }
+	void LogMessage(EMessageSeverity::Type InSeverity, FText InMessage) const { LogMessage(FTokenizedMessage::Create(InSeverity, InMessage)); }
+
 	ENGINE_API void ResetToRefPose();
 
 	ENGINE_API bool ContainsNaN() const;
