Fix up Material Parameter duplicates validation errors when using r.Material.PedanticErrorChecksEnabled=1 .

Before, this was reporting only the parameters with different Default Values than the first one chosen by the code, which can be random. When reporting those parameters, it was attempting to show a comparison between the current parameter and the previous duplicate/first one found.
It was also broken during a previous change, because it was using the same path for both parameters when doing a comparison.

Now, when a duplicate with a different Default Value appears, all are reported, with their respective values and their asset paths.

[RN] minor | Rendering Architecture
#Synced-CL 51289354

diff --git a/Engine/Source/Runtime/Engine/Private/Materials/Material.cpp b/Engine/Source/Runtime/Engine/Private/Materials/Material.cpp
--- a/Engine/Source/Runtime/Engine/Private/Materials/Material.cpp
+++ b/Engine/Source/Runtime/Engine/Private/Materials/Material.cpp
@@ -7297,16 +7297,25 @@ bool UMaterial::CheckInValidStateForCompilation(class FMaterialCompiler* Compile
 	bool bSuccess = true;
 	
 	// Check and report errors due to duplicate parameters set to distinct values.
-	if (!CachedExpressionData->DuplicateParameterErrors.IsEmpty())
+	for (const auto& Param : CachedExpressionData->ParameterExpressionsAndValues)
 	{
-		for (const TPair<TObjectPtr<UMaterialExpression>, FString>& Error : CachedExpressionData->DuplicateParameterErrors)
+		const FString ParamName = Param.Key.ToString();
+		const FMaterialCachedExpressionData::FParamExpressionsValues& ParameterExpressionsValues = Param.Value;
+		if (ParameterExpressionsValues.bHasDuplicatesWithDifferentValues)
 		{
-			Compiler->AppendExpressionError(Error.Get<0>(), *Error.Get<1>());
-		}
+			bSuccess = false;
+			for (const auto& ExpressionsValues : ParameterExpressionsValues.ExpressionsValues)
+			{
+				UMaterialExpression* Expression = ExpressionsValues.Get<0>();
+				const FString& ParamValue = ExpressionsValues.Get<1>();
 
-		bSuccess = false;
-	}
+				FString ErrorMessage = FString::Printf(TEXT("Parameter %s with default value %s from asset %s has duplicates with different default values."),
+					*ParamName, *ParamValue, *Expression->GetAssetPathName());
 
+				Compiler->AppendExpressionError(Expression,  *ErrorMessage);
+			}
+		}
+	}
 	return bSuccess;
 }
 
diff --git a/Engine/Source/Runtime/Engine/Private/Materials/MaterialCachedData.cpp b/Engine/Source/Runtime/Engine/Private/Materials/MaterialCachedData.cpp
--- a/Engine/Source/Runtime/Engine/Private/Materials/MaterialCachedData.cpp
+++ b/Engine/Source/Runtime/Engine/Private/Materials/MaterialCachedData.cpp
@@ -218,245 +218,218 @@ static bool TryAddParameter(FMaterialCachedExpressionData& CachedData,
 	return false;
 }
 
-bool FMaterialCachedExpressionData::AddParameter(const FMaterialParameterInfo& ParameterInfo, const FMaterialParameterMetadata& ParameterMeta, const FString& InAssetPath, UObject*& OutReferencedTexture, UTextureCollection*& OutReferencedTextureCollection, FText& OutErrorMessage)
+FString FMaterialCachedExpressionData::GetParameterValueAsText(const FMaterialParameterMetadata& ParameterMeta)
 {
-	check(EditorOnlyData);
+	switch (ParameterMeta.Value.Type)
+	{
+	case EMaterialParameterType::Scalar:
+		return FString::Printf(TEXT("%f"), ParameterMeta.Value.AsScalar());
 
-	if (!InAssetPath.IsEmpty())
+	case EMaterialParameterType::Vector:
+		return ParameterMeta.Value.AsLinearColor().ToString();
+
+	case EMaterialParameterType::DoubleVector:
+		return ParameterMeta.Value.AsVector4d().ToString();
+
+	case EMaterialParameterType::Texture:
+		return TSoftObjectPtr<UTexture>(ParameterMeta.Value.Texture).ToString();
+
+	case EMaterialParameterType::TextureCollection:
+		return TSoftObjectPtr<UTextureCollection>(ParameterMeta.Value.TextureCollection).ToString();
+
+	case EMaterialParameterType::ParameterCollection:
+		return TSoftObjectPtr<UMaterialParameterCollection>(ParameterMeta.Value.ParameterCollection).ToString();
+
+	case EMaterialParameterType::Font:
+		return FString::Printf(TEXT("%s(%i)"), *TSoftObjectPtr<UFont>(ParameterMeta.Value.Font.Value).ToString(), ParameterMeta.Value.Font.Page);
+
+	case EMaterialParameterType::RuntimeVirtualTexture:
+		return TSoftObjectPtr<URuntimeVirtualTexture>(ParameterMeta.Value.RuntimeVirtualTexture).ToString();
+
+	case EMaterialParameterType::SparseVolumeTexture:
+		return TSoftObjectPtr<USparseVolumeTexture>(ParameterMeta.Value.SparseVolumeTexture).ToString();
+
+	case EMaterialParameterType::StaticSwitch:
+		return FString::Printf(TEXT("%s%s"), ParameterMeta.Value.AsStaticSwitch() ? TEXT("true") : TEXT("false"), 
+			ParameterMeta.bDynamicSwitchParameter ? TEXT("(dynamic)") : TEXT(""));
+
+	case EMaterialParameterType::StaticComponentMask:
 	{
-		FParameterAssetPathsEntry& Entry = EditorOnlyData->AssetPathsMap.FindOrAdd(ParameterInfo.Name);
-		TMap<FString, int32>& AssetPaths = Entry.AssetPaths;
+		const FStaticComponentMaskValue MaskValue = ParameterMeta.Value.AsStaticComponentMask();
+		return FString::Printf(TEXT("R=%s,G=%s,B=%s,A=%s")
+			, MaskValue.R ? TEXT("true") : TEXT("false")
+			, MaskValue.G ? TEXT("true") : TEXT("false")
+			, MaskValue.B ? TEXT("true") : TEXT("false")
+			, MaskValue.A ? TEXT("true") : TEXT("false"));
+	}
+	default:
+		return {};
+	}
 
-		int32& PathsCounter = AssetPaths.FindOrAdd(InAssetPath);
-		++PathsCounter;
+}
+
+// Returns whether a parameter that is a duplicate has the same value as the existing one
+bool FMaterialCachedExpressionData::CheckDuplicateParameterValuesEqual(const int32 InExistingParamIdx, const FMaterialParameterMetadata& DuplicateParameterMeta)
+{
+	switch (DuplicateParameterMeta.Value.Type)
+	{
+	case EMaterialParameterType::Scalar:
+		return ScalarValues[InExistingParamIdx] == DuplicateParameterMeta.Value.AsScalar();
+
+	case EMaterialParameterType::Vector:
+		return  VectorValues[InExistingParamIdx] == DuplicateParameterMeta.Value.AsLinearColor();
+
+	case EMaterialParameterType::DoubleVector:
+		return DoubleVectorValues[InExistingParamIdx] == DuplicateParameterMeta.Value.AsVector4d();
+
+	case EMaterialParameterType::Texture:
+		return TextureValues[InExistingParamIdx] == DuplicateParameterMeta.Value.Texture;
+
+	case EMaterialParameterType::TextureCollection:
+		return TextureCollectionValues[InExistingParamIdx] == DuplicateParameterMeta.Value.TextureCollection;
+
+	case EMaterialParameterType::ParameterCollection:
+		return ParameterCollectionValues[InExistingParamIdx] == DuplicateParameterMeta.Value.ParameterCollection;
+
+	case EMaterialParameterType::Font:
+		return FontValues[InExistingParamIdx] == DuplicateParameterMeta.Value.Font.Value && FontPageValues[InExistingParamIdx] == DuplicateParameterMeta.Value.Font.Page;
+
+	case EMaterialParameterType::RuntimeVirtualTexture:
+		return RuntimeVirtualTextureValues[InExistingParamIdx] == DuplicateParameterMeta.Value.RuntimeVirtualTexture;
+
+	case EMaterialParameterType::SparseVolumeTexture:
+		return SparseVolumeTextureValues[InExistingParamIdx] == DuplicateParameterMeta.Value.SparseVolumeTexture;
+
+	case EMaterialParameterType::StaticSwitch:
+		return (StaticSwitchValues[InExistingParamIdx] == DuplicateParameterMeta.Value.AsStaticSwitch()) && (DynamicSwitchValues[InExistingParamIdx] == DuplicateParameterMeta.bDynamicSwitchParameter);
+
+	case EMaterialParameterType::StaticComponentMask:
+		return EditorOnlyData->StaticComponentMaskValues[InExistingParamIdx] == DuplicateParameterMeta.Value.AsStaticComponentMask();
+
+	default:
+		return true;
 	}
+}
 
-	const FMaterialCachedParameterEditorInfo EditorInfo(ParameterMeta.ExpressionGuid, ParameterMeta.Description, ParameterMeta.Group, ParameterMeta.SortPriority);
-	int32 Index = INDEX_NONE;
-	TOptional<FMaterialCachedParameterEditorInfo> PreviousEditorInfo;
-	if (TryAddParameter(*this, ParameterMeta.Value.Type, ParameterInfo, EditorInfo, Index, PreviousEditorInfo))
+void FMaterialCachedExpressionData::InsertParameterValue(const FMaterialParameterMetadata& ParameterMeta, const int32 Index, UObject*& OutReferencedTexture, UTextureCollection*& OutReferencedTextureCollection)
+{
+	switch (ParameterMeta.Value.Type)
 	{
-		switch (ParameterMeta.Value.Type)
+	case EMaterialParameterType::Scalar:
+		ScalarValues.Insert(ParameterMeta.Value.AsScalar(), Index);
+		EditorOnlyData->ScalarMinMaxValues.Insert(FVector2D(ParameterMeta.ScalarMin, ParameterMeta.ScalarMax), Index);
+		EditorOnlyData->ScalarEnumerationValues.Insert(ParameterMeta.ScalarEnumeration, Index);
+		EditorOnlyData->ScalarEnumerationIndexValues.Insert(ParameterMeta.ScalarEnumerationIndex, Index);
+		ScalarPrimitiveDataIndexValues.Insert(ParameterMeta.PrimitiveDataIndex, Index);
+		if (ParameterMeta.bUsedAsAtlasPosition)
 		{
-		case EMaterialParameterType::Scalar:
-			ScalarValues.Insert(ParameterMeta.Value.AsScalar(), Index);
-			EditorOnlyData->ScalarMinMaxValues.Insert(FVector2D(ParameterMeta.ScalarMin, ParameterMeta.ScalarMax), Index);
-			EditorOnlyData->ScalarEnumerationValues.Insert(ParameterMeta.ScalarEnumeration, Index);
-			EditorOnlyData->ScalarEnumerationIndexValues.Insert(ParameterMeta.ScalarEnumerationIndex, Index);
-			ScalarPrimitiveDataIndexValues.Insert(ParameterMeta.PrimitiveDataIndex, Index);
-			if (ParameterMeta.bUsedAsAtlasPosition)
-			{
-				EditorOnlyData->ScalarCurveValues.Insert(ParameterMeta.ScalarCurve.Get(), Index);
-				EditorOnlyData->ScalarCurveAtlasValues.Insert(ParameterMeta.ScalarAtlas.Get(), Index);
-				OutReferencedTexture = ParameterMeta.ScalarAtlas.Get();
-			}
-			else
-			{
-				EditorOnlyData->ScalarCurveValues.Insert(nullptr, Index);
-				EditorOnlyData->ScalarCurveAtlasValues.Insert(nullptr, Index);
-			}
-			break;
-
-		case EMaterialParameterType::Vector:
-			VectorValues.Insert(ParameterMeta.Value.AsLinearColor(), Index);
-			EditorOnlyData->VectorChannelNameValues.Insert(ParameterMeta.ChannelNames, Index);
-			EditorOnlyData->VectorUsedAsChannelMaskValues.Insert(ParameterMeta.bUsedAsChannelMask, Index);
-			VectorPrimitiveDataIndexValues.Insert(ParameterMeta.PrimitiveDataIndex, Index);
-			break;
-
-		case EMaterialParameterType::DoubleVector:
-			DoubleVectorValues.Insert(ParameterMeta.Value.AsVector4d(), Index);
-			break;
-
-		case EMaterialParameterType::Texture:
-			TextureValues.Insert(ParameterMeta.Value.Texture, Index);
-			EditorOnlyData->TextureChannelNameValues.Insert(ParameterMeta.ChannelNames, Index);
-			OutReferencedTexture = ParameterMeta.Value.Texture;
-			break;
-
-		case EMaterialParameterType::TextureCollection:
-			TextureCollectionValues.Insert(ParameterMeta.Value.TextureCollection, Index);
-			OutReferencedTextureCollection = ParameterMeta.Value.TextureCollection;
-			break;
-
-		case EMaterialParameterType::ParameterCollection:
-			ParameterCollectionValues.Insert(ParameterMeta.Value.ParameterCollection, Index);
-			break;
-
-		case EMaterialParameterType::Font:
-			FontValues.Insert(ParameterMeta.Value.Font.Value, Index);
-			FontPageValues.Insert(ParameterMeta.Value.Font.Page, Index);
-			if (ParameterMeta.Value.Font.Value && ParameterMeta.Value.Font.Value->Textures.IsValidIndex(ParameterMeta.Value.Font.Page))
-			{
-				OutReferencedTexture = ParameterMeta.Value.Font.Value->Textures[ParameterMeta.Value.Font.Page];
-			}
-			break;
-
-		case EMaterialParameterType::RuntimeVirtualTexture:
-			RuntimeVirtualTextureValues.Insert(ParameterMeta.Value.RuntimeVirtualTexture, Index);
-			OutReferencedTexture = ParameterMeta.Value.RuntimeVirtualTexture;
-			break;
-
-		case EMaterialParameterType::SparseVolumeTexture:
-			SparseVolumeTextureValues.Insert(ParameterMeta.Value.SparseVolumeTexture, Index);
-			OutReferencedTexture = ParameterMeta.Value.SparseVolumeTexture;
-			break;
-
-		case EMaterialParameterType::StaticSwitch:
-			StaticSwitchValues.Insert(ParameterMeta.Value.AsStaticSwitch(), Index);
-			DynamicSwitchValues.Insert(ParameterMeta.bDynamicSwitchParameter, Index);
-			break;
-
-		case EMaterialParameterType::StaticComponentMask:
-			EditorOnlyData->StaticComponentMaskValues.Insert(ParameterMeta.Value.AsStaticComponentMask(), Index);
-			break;
-
-		default:
-			checkNoEntry();
-			break;
+			EditorOnlyData->ScalarCurveValues.Insert(ParameterMeta.ScalarCurve.Get(), Index);
+			EditorOnlyData->ScalarCurveAtlasValues.Insert(ParameterMeta.ScalarAtlas.Get(), Index);
+			OutReferencedTexture = ParameterMeta.ScalarAtlas.Get();
 		}
-	}
-	else
-	{
-		auto GetEditorInfoAsText = [AssetPath = InAssetPath](const FMaterialCachedParameterEditorInfo& InEditorInfo) -> FText
-			{
-				if(!AssetPath.IsEmpty())
-				{
-					return FText::Format(LOCTEXT("ReportParameterMetaData", "group:'{0}', asset:'{1}'"), FText::FromName(InEditorInfo.Group), FText::FromString(AssetPath));
-				}
-				else
-				{
-					return FText::Format(LOCTEXT("ReportInvalidatedParameterMetaData", "group:'{0}', description:'{1}'"), FText::FromName(InEditorInfo.Group), FText::FromString(InEditorInfo.Description));
-				}
-			};
+		else
+		{
+			EditorOnlyData->ScalarCurveValues.Insert(nullptr, Index);
+			EditorOnlyData->ScalarCurveAtlasValues.Insert(nullptr, Index);
+		}
+		break;
 
-		auto ReportError = [&GetEditorInfoAsText](const FString& InCachedValue, const FString& InValue, const FMaterialCachedParameterEditorInfo& InCachedEditorInfo, const FMaterialCachedParameterEditorInfo& InEditorInfo) -> FText
-			{
-				return FText::Format(LOCTEXT("ReportDifferentParameterValueError", "{0} ({1}) vs. cached: {2} ({3})"), 
-					FText::FromString(InValue), GetEditorInfoAsText(InEditorInfo), FText::FromString(InCachedValue), GetEditorInfoAsText(InCachedEditorInfo));
-			};
+	case EMaterialParameterType::Vector:
+		VectorValues.Insert(ParameterMeta.Value.AsLinearColor(), Index);
+		EditorOnlyData->VectorChannelNameValues.Insert(ParameterMeta.ChannelNames, Index);
+		EditorOnlyData->VectorUsedAsChannelMaskValues.Insert(ParameterMeta.bUsedAsChannelMask, Index);
+		VectorPrimitiveDataIndexValues.Insert(ParameterMeta.PrimitiveDataIndex, Index);
+		break;
 
-		bool bSameValue = false;
-		switch (ParameterMeta.Value.Type)
-		{
-		case EMaterialParameterType::Scalar:
-			bSameValue = ScalarValues[Index] == ParameterMeta.Value.AsScalar();
-			if (!bSameValue)
-			{
-				OutErrorMessage = ReportError(FString::Printf(TEXT("%f"), ScalarValues[Index]), FString::Printf(TEXT("%f"), ParameterMeta.Value.AsScalar()), *PreviousEditorInfo, EditorInfo);
-			}
-			break;
+	case EMaterialParameterType::DoubleVector:
+		DoubleVectorValues.Insert(ParameterMeta.Value.AsVector4d(), Index);
+		break;
 
-		case EMaterialParameterType::Vector:
-			bSameValue = VectorValues[Index] == ParameterMeta.Value.AsLinearColor();
-			if (!bSameValue)
-			{
-				OutErrorMessage = ReportError(VectorValues[Index].ToString(), ParameterMeta.Value.AsLinearColor().ToString(), *PreviousEditorInfo, EditorInfo);
-			}
-			break;
+	case EMaterialParameterType::Texture:
+		TextureValues.Insert(ParameterMeta.Value.Texture, Index);
+		EditorOnlyData->TextureChannelNameValues.Insert(ParameterMeta.ChannelNames, Index);
+		OutReferencedTexture = ParameterMeta.Value.Texture;
+		break;
 
-		case EMaterialParameterType::DoubleVector:
-			bSameValue = DoubleVectorValues[Index] == ParameterMeta.Value.AsVector4d();
-			if (!bSameValue)
-			{
-				OutErrorMessage = ReportError(DoubleVectorValues[Index].ToString(), ParameterMeta.Value.AsVector4d().ToString(), *PreviousEditorInfo, EditorInfo);
-			}
-			break;
+	case EMaterialParameterType::TextureCollection:
+		TextureCollectionValues.Insert(ParameterMeta.Value.TextureCollection, Index);
+		OutReferencedTextureCollection = ParameterMeta.Value.TextureCollection;
+		break;
 
-		case EMaterialParameterType::Texture:
-			bSameValue = TextureValues[Index] == ParameterMeta.Value.Texture;
-			if (!bSameValue)
-			{
-				OutErrorMessage = ReportError(TextureValues[Index].ToString(), TSoftObjectPtr<UTexture>(ParameterMeta.Value.Texture).ToString(), *PreviousEditorInfo, EditorInfo);
-			}
-			break;
+	case EMaterialParameterType::ParameterCollection:
+		ParameterCollectionValues.Insert(ParameterMeta.Value.ParameterCollection, Index);
+		break;
 
-		case EMaterialParameterType::TextureCollection:
-			bSameValue = TextureCollectionValues[Index] == ParameterMeta.Value.TextureCollection;
-			if (!bSameValue)
-			{
-				OutErrorMessage = ReportError(TextureCollectionValues[Index].ToString(), TSoftObjectPtr<UTextureCollection>(ParameterMeta.Value.TextureCollection).ToString(), *PreviousEditorInfo, EditorInfo);
-			}
-			break;
+	case EMaterialParameterType::Font:
+		FontValues.Insert(ParameterMeta.Value.Font.Value, Index);
+		FontPageValues.Insert(ParameterMeta.Value.Font.Page, Index);
+		if (ParameterMeta.Value.Font.Value && ParameterMeta.Value.Font.Value->Textures.IsValidIndex(ParameterMeta.Value.Font.Page))
+		{
+			OutReferencedTexture = ParameterMeta.Value.Font.Value->Textures[ParameterMeta.Value.Font.Page];
+		}
+		break;
 
-		case EMaterialParameterType::ParameterCollection:
-			bSameValue = ParameterCollectionValues[Index] == ParameterMeta.Value.ParameterCollection;
-			if (!bSameValue)
-			{
-				OutErrorMessage = ReportError(ParameterCollectionValues[Index].ToString(), TSoftObjectPtr<UMaterialParameterCollection>(ParameterMeta.Value.ParameterCollection).ToString(), *PreviousEditorInfo, EditorInfo);
-			}
-			break;
+	case EMaterialParameterType::RuntimeVirtualTexture:
+		RuntimeVirtualTextureValues.Insert(ParameterMeta.Value.RuntimeVirtualTexture, Index);
+		OutReferencedTexture = ParameterMeta.Value.RuntimeVirtualTexture;
+		break;
 
-		case EMaterialParameterType::Font:
-			bSameValue = FontValues[Index] == ParameterMeta.Value.Font.Value && FontPageValues[Index] == ParameterMeta.Value.Font.Page;
-			if (!bSameValue)
-			{
-				OutErrorMessage = ReportError(
-					FString::Printf(TEXT("%s(%i)"), *FontValues[Index].ToString(), FontPageValues[Index]), 
-					FString::Printf(TEXT("%s(%i)"), *TSoftObjectPtr<UFont>(ParameterMeta.Value.Font.Value).ToString(), ParameterMeta.Value.Font.Page), 
-					*PreviousEditorInfo, EditorInfo);
-			}
-			break;
+	case EMaterialParameterType::SparseVolumeTexture:
+		SparseVolumeTextureValues.Insert(ParameterMeta.Value.SparseVolumeTexture, Index);
+		OutReferencedTexture = ParameterMeta.Value.SparseVolumeTexture;
+		break;
 
-		case EMaterialParameterType::RuntimeVirtualTexture:
-			bSameValue = RuntimeVirtualTextureValues[Index] == ParameterMeta.Value.RuntimeVirtualTexture;
-			if (!bSameValue)
-			{
-				OutErrorMessage = ReportError(RuntimeVirtualTextureValues[Index].ToString(), TSoftObjectPtr<URuntimeVirtualTexture>(ParameterMeta.Value.RuntimeVirtualTexture).ToString(), *PreviousEditorInfo, EditorInfo);
-			}
-			break;
+	case EMaterialParameterType::StaticSwitch:
+		StaticSwitchValues.Insert(ParameterMeta.Value.AsStaticSwitch(), Index);
+		DynamicSwitchValues.Insert(ParameterMeta.bDynamicSwitchParameter, Index);
+		break;
 
-		case EMaterialParameterType::SparseVolumeTexture:
-			bSameValue = SparseVolumeTextureValues[Index] == ParameterMeta.Value.SparseVolumeTexture;
-			if (!bSameValue)
-			{
-				OutErrorMessage = ReportError(SparseVolumeTextureValues[Index].ToString(), TSoftObjectPtr<USparseVolumeTexture>(ParameterMeta.Value.SparseVolumeTexture).ToString(), *PreviousEditorInfo, EditorInfo);
-			}
-			break;
+	case EMaterialParameterType::StaticComponentMask:
+		EditorOnlyData->StaticComponentMaskValues.Insert(ParameterMeta.Value.AsStaticComponentMask(), Index);
+		break;
 
-		case EMaterialParameterType::StaticSwitch:
-			bSameValue = (StaticSwitchValues[Index] == ParameterMeta.Value.AsStaticSwitch()) && (DynamicSwitchValues[Index] == ParameterMeta.bDynamicSwitchParameter);
-			if (!bSameValue)
-			{
-				auto GetStaticSwitchString = [](bool bInValue, bool bIsDynamicSwitch) -> FString
-					{
-						return FText::Format(LOCTEXT("StaticSwitchValue", "{0}{1}"), bInValue ? LOCTEXT("true", "true") : LOCTEXT("false", "false"), bIsDynamicSwitch ? LOCTEXT("dynamic", "(dynamic)") : FText()).ToString();
-					};
-				OutErrorMessage = ReportError(
-					GetStaticSwitchString(StaticSwitchValues[Index], DynamicSwitchValues[Index]), 
-					GetStaticSwitchString(ParameterMeta.Value.AsStaticSwitch(), ParameterMeta.bDynamicSwitchParameter), 
-					*PreviousEditorInfo, EditorInfo);
-			}
-			break;
+	default:
+		checkNoEntry();
+		break;
+	}
+}
 
-		case EMaterialParameterType::StaticComponentMask:
-			bSameValue = EditorOnlyData->StaticComponentMaskValues[Index] == ParameterMeta.Value.AsStaticComponentMask();
-			if (!bSameValue)
-			{
-				auto GetStaticComponentMaskString = [](const FStaticComponentMaskValue& InValue) -> FString
-					{
-						return FText::Format(LOCTEXT("StaticComponentMaskValue", "R={0},G={1},B={2},A={3}")
-							, InValue.R ? LOCTEXT("true", "true") : LOCTEXT("false", "false")
-							, InValue.G ? LOCTEXT("true", "true") : LOCTEXT("false", "false")
-							, InValue.B ? LOCTEXT("true", "true") : LOCTEXT("false", "false")
-							, InValue.A ? LOCTEXT("true", "true") : LOCTEXT("false", "false")).ToString();
-					};
-				OutErrorMessage = ReportError(
-					GetStaticComponentMaskString(EditorOnlyData->StaticComponentMaskValues[Index]), 
-					GetStaticComponentMaskString(ParameterMeta.Value.AsStaticComponentMask()), 
-					*PreviousEditorInfo, EditorInfo);
-			}
-			break;
+void FMaterialCachedExpressionData::AddParameter(const FMaterialParameterInfo& ParameterInfo, const FMaterialParameterMetadata& ParameterMeta, const FString& InAssetPath, UObject*& OutReferencedTexture, UTextureCollection*& OutReferencedTextureCollection, UMaterialExpression* InExpression)
+{
+	check(EditorOnlyData);
 
-		default:
-			bSameValue = true;
-			break;
-		}
+	if (!InAssetPath.IsEmpty())
+	{
+		FParameterAssetPathsEntry& Entry = EditorOnlyData->AssetPathsMap.FindOrAdd(ParameterInfo.Name);
+		TMap<FString, int32>& AssetPaths = Entry.AssetPaths;
 
-		return bSameValue;
+		int32& PathsCounter = AssetPaths.FindOrAdd(InAssetPath);
+		++PathsCounter;
 	}
 
-	return true;
+	const FMaterialCachedParameterEditorInfo EditorInfo(ParameterMeta.ExpressionGuid, ParameterMeta.Description, ParameterMeta.Group, ParameterMeta.SortPriority);
+	int32 Index = INDEX_NONE;
+	TOptional<FMaterialCachedParameterEditorInfo> PreviousEditorInfo;
+
+	const bool bNonDuplicate = TryAddParameter(*this, ParameterMeta.Value.Type, ParameterInfo, EditorInfo, Index, PreviousEditorInfo);
+	if (bNonDuplicate)
+	{
+		InsertParameterValue(ParameterMeta, Index, OutReferencedTexture, OutReferencedTextureCollection);
+	}
+
+	// ParameterExpressionsAndValues is only used in case PedanticErrorChecksEnabled=1 currently. Avoid doing operations if disabled.
+	const IConsoleVariable* CVar = IConsoleManager::Get().FindConsoleVariable(TEXT("r.Material.PedanticErrorChecksEnabled"));
+	if (CVar && !CVar->GetBool())
+	{
+		return;
+	}
+
+	const bool bEqualValues = CheckDuplicateParameterValuesEqual(Index, ParameterMeta);
+	const FString CurrentParamValueText = GetParameterValueAsText(ParameterMeta);
+
+	FParamExpressionsValues& ParamExpressionsValues = ParameterExpressionsAndValues.FindOrAdd(ParameterInfo.Name);
+	ParamExpressionsValues.ExpressionsValues.AddUnique({ InExpression,  CurrentParamValueText });
+	ParamExpressionsValues.bHasDuplicatesWithDifferentValues |= !bEqualValues;
 }
 
 void FMaterialCachedExpressionData::UpdateForFunction(const FMaterialCachedExpressionContext& Context, UMaterialFunctionInterface* Function, EMaterialParameterAssociation Association, int32 ParameterIndex)
@@ -561,12 +534,7 @@ void FMaterialCachedExpressionData::UpdateForExpressions(const FMaterialCachedEx
 			const FString AssetPath = Expression->GetAssetPathName();
 
 			// Try add the parameter. If this fails, the parameter is being added twice with different values. Report it as error.
-			if (!AddParameter(ParameterInfo, ParameterMeta, AssetPath, ReferencedTexture, ReferencedTextureCollection, ErrorContext))
-			{
-				FText ErrorMessage = FText::Format(LOCTEXT("DuplicateParameterError", "Parameter '{0}' is set multiple times to different values : {1}. Make sure each parameter is set once or always to the same value."),
-					 FText::FromName(ParameterName), ErrorContext);
-				DuplicateParameterErrors.AddUnique({ Expression, ErrorMessage.ToString() });
-			}
+			AddParameter(ParameterInfo, ParameterMeta, AssetPath, ReferencedTexture, ReferencedTextureCollection, Expression);
 		}
 
 
diff --git a/Engine/Source/Runtime/Engine/Public/MaterialCachedData.h b/Engine/Source/Runtime/Engine/Public/MaterialCachedData.h
--- a/Engine/Source/Runtime/Engine/Public/MaterialCachedData.h
+++ b/Engine/Source/Runtime/Engine/Public/MaterialCachedData.h
@@ -223,12 +223,12 @@ struct FMaterialCachedExpressionData
 	void Validate(const UMaterialInterface& Material);
 
 	/** Adds a parameter. If this returns false, a parameter with identical name has already been added but it was set to a different value. */
-	bool AddParameter(const FMaterialParameterInfo& ParameterInfo, const FMaterialParameterMetadata& ParameterMeta, const FString& InAssetPath ,UObject*& OutReferencedTexture, UTextureCollection*& OutReferencedTextureCollection, FText& OutErrorMessage);
+	void AddParameter(const FMaterialParameterInfo& ParameterInfo, const FMaterialParameterMetadata& ParameterMeta, const FString& InAssetPath ,UObject*& OutReferencedTexture, UTextureCollection*& OutReferencedTextureCollection, UMaterialExpression* InExpression);
 
-	bool AddParameter(const FMaterialParameterInfo& ParameterInfo, const FMaterialParameterMetadata& ParameterMeta, const FString& InAssetPath, UObject*& OutReferencedTexture, FText& OutErrorMessage)
+	void AddParameter(const FMaterialParameterInfo& ParameterInfo, const FMaterialParameterMetadata& ParameterMeta, const FString& InAssetPath, UObject*& OutReferencedTexture, UMaterialExpression* InExpression)
 	{
 		UTextureCollection* UnusedReferencedTextureCollection = nullptr;
-		return AddParameter(ParameterInfo, ParameterMeta, InAssetPath, OutReferencedTexture, UnusedReferencedTextureCollection, OutErrorMessage);
+		return AddParameter(ParameterInfo, ParameterMeta, InAssetPath, OutReferencedTexture, UnusedReferencedTextureCollection, InExpression);
 	}
 #endif // WITH_EDITOR
 
@@ -301,6 +301,13 @@ struct FMaterialCachedExpressionData
 	bool Serialize(FArchive& Ar);
 	void PostSerialize(const FArchive& Ar);
 
+private:
+	bool CheckDuplicateParameterValuesEqual(const int32 InExistingParamIdx, const FMaterialParameterMetadata& DuplicateParameterMeta);
+	FString GetParameterValueAsText(const FMaterialParameterMetadata& ParameterMeta);
+	void InsertParameterValue(const FMaterialParameterMetadata& ParameterMeta, const int32 Index, UObject*& OutReferencedTexture, UTextureCollection*& OutReferencedTextureCollection);
+
+public:
+
 #if WITH_EDITORONLY_DATA
 	TSharedPtr<FMaterialCachedExpressionEditorOnlyData> EditorOnlyData;
 #endif // WITH_EDITORONLY_DATA
@@ -445,8 +452,13 @@ struct FMaterialCachedExpressionData
 	uint64 PropertyConnectedMask = 0;
 
 #if WITH_EDITOR
-	/** Array of errors reporting a parameter being set multiple times to distinct values. */
-	TArray<TPair<TObjectPtr<class UMaterialExpression>, FString>> DuplicateParameterErrors;
+	struct FParamExpressionsValues
+	{
+		TArray<TPair<TObjectPtr<class UMaterialExpression>, FString>> ExpressionsValues;
+		bool bHasDuplicatesWithDifferentValues = false;
+	};
+	/** Stores Parameters to expresions and string values. Can be used for error logging duplicates.*/
+	TMap<FName, FParamExpressionsValues> ParameterExpressionsAndValues;
 #endif
 };
 
