FReferenceSkeleton::EnsureParentsExist: make sure that the bone indices are prepared for both input and output to avoid retaining information that does not comply with the ref skeleton

#rnx

#Synced-CL 51330703

diff --git a/Engine/Source/Runtime/Engine/Private/Animation/ReferenceSkeleton.cpp b/Engine/Source/Runtime/Engine/Private/Animation/ReferenceSkeleton.cpp
--- a/Engine/Source/Runtime/Engine/Private/Animation/ReferenceSkeleton.cpp
+++ b/Engine/Source/Runtime/Engine/Private/Animation/ReferenceSkeleton.cpp
@@ -779,32 +779,56 @@ struct FEnsureParentsExistScratchArea : public TThreadSingleton<FEnsureParentsEx
 
 void FReferenceSkeleton::EnsureParentsExist(TArray<FBoneIndexType>& InOutBoneSortedArray) const
 {
+	bool bWarn = false;
+#if	!(UE_BUILD_SHIPPING || UE_BUILD_TEST)
+	// Warn if we're getting bad data.
+	// Bones are matched as int32, and a non found bone will be set to INDEX_NONE == -1
+	// This should never happen, so if it does, something is wrong!
+	bWarn = true;
+#endif
+	
 	const int32 NumBones = GetNum();
+	if (NumBones == 0)
+	{
+		if (bWarn && !InOutBoneSortedArray.IsEmpty())
+		{
+			UE_LOGF(LogAnimation, Log, "FReferenceSkeleton::EnsureParentsExist, RefSkeleton is empty. Removing all bones from InOutBoneSortedArray.");
+		}
+		InOutBoneSortedArray.Empty();
+		return;
+	}
+	
+	// remove invalid indices
+	InOutBoneSortedArray.RemoveAll([bWarn, NumBones](const FBoneIndexType BoneIndex)
+	{
+		const bool bRemove = BoneIndex >= NumBones;	
+		if (bWarn && bRemove)
+		{
+			UE_LOGF(LogAnimation, Log,
+				"FReferenceSkeleton::EnsureParentsExist, invalid BoneIndex[%d / %d]. Removing bone from InOutBoneSortedArray.",
+				BoneIndex, NumBones);	
+		}
+		return bRemove;
+	});
+	
+	if (InOutBoneSortedArray.IsEmpty())
+	{
+		return;
+	}
+	
 	// Iterate through existing array.
-	int32 i = 0;
-
 	TArray<bool>& BoneExists = FEnsureParentsExistScratchArea::Get().BoneExists;
 	BoneExists.Reset();
-	BoneExists.SetNumZeroed(NumBones);
+	BoneExists.Init(false, NumBones);
 
+	int32 i = 0;
 	while (i < InOutBoneSortedArray.Num())
 	{
 		const int32 BoneIndex = InOutBoneSortedArray[i];
 
 		// For the root bone, just move on.
 		if (BoneIndex > 0)
-		{
-#if	!(UE_BUILD_SHIPPING || UE_BUILD_TEST)
-			// Warn if we're getting bad data.
-			// Bones are matched as int32, and a non found bone will be set to INDEX_NONE == -1
-			// This should never happen, so if it does, something is wrong!
-			if (BoneIndex >= NumBones)
-			{
-				UE_LOG(LogAnimation, Log, TEXT("FAnimationRuntime::EnsureParentsExist, BoneIndex >= RefSkeleton.GetNum()."));
-				i++;
-				continue;
-			}
-#endif
+		{			
 			BoneExists[BoneIndex] = true;
 
 			const int32 ParentIndex = GetParentIndex(BoneIndex);
@@ -812,7 +836,7 @@ void FReferenceSkeleton::EnsureParentsExist(TArray<FBoneIndexType>& InOutBoneSor
 			// If we do not have this parent in the array, we add it in this location, and leave 'i' where it is.
 			// This can happen if somebody removes bones in the physics asset, then it will try add back in, and in the process, 
 			// parent can be missing
-			if (!BoneExists[ParentIndex])
+			if (BoneExists.IsValidIndex(ParentIndex) && !BoneExists[ParentIndex])
 			{
 				InOutBoneSortedArray.InsertUninitialized(i);
 				InOutBoneSortedArray[i] = ParentIndex;
