Fix multiple "r.Material.PedanticErrorChecksEnabled" FindConsoleObject() calls by making return value static.
[RN] minor | Rendering Architecture

#Synced-CL 51289354

diff --git a/Engine/Source/Runtime/Engine/Private/Materials/MaterialCachedData.cpp b/Engine/Source/Runtime/Engine/Private/Materials/MaterialCachedData.cpp
--- a/Engine/Source/Runtime/Engine/Private/Materials/MaterialCachedData.cpp
+++ b/Engine/Source/Runtime/Engine/Private/Materials/MaterialCachedData.cpp
@@ -418,7 +418,7 @@ void FMaterialCachedExpressionData::AddParameter(const FMaterialParameterInfo& P
 	}
 
 	// ParameterExpressionsAndValues is only used in case PedanticErrorChecksEnabled=1 currently. Avoid doing operations if disabled.
-	const IConsoleVariable* CVar = IConsoleManager::Get().FindConsoleVariable(TEXT("r.Material.PedanticErrorChecksEnabled"));
+	static const IConsoleVariable* CVar = IConsoleManager::Get().FindConsoleVariable(TEXT("r.Material.PedanticErrorChecksEnabled"));
 	if (CVar && !CVar->GetBool())
 	{
 		return;
