Substrate - Toon - Added experimental status for Toon BSDF node and Toon profile asset for 5.8.

#tests PC editor, adding Substrate toon BSDF and opening Toon profile asset.
#rnx

#Synced-CL 51339376

diff --git a/Engine/Source/Runtime/Engine/Classes/Engine/ToonProfile.h b/Engine/Source/Runtime/Engine/Classes/Engine/ToonProfile.h
--- a/Engine/Source/Runtime/Engine/Classes/Engine/ToonProfile.h
+++ b/Engine/Source/Runtime/Engine/Classes/Engine/ToonProfile.h
@@ -164,7 +164,7 @@ struct FToonProfileStruct
  * Toon profile asset, can be specified at a material. 
  * Don't change at runtime. All properties in here are per material.
  */
-UCLASS(autoexpandcategories = ToonProfile, MinimalAPI)
+UCLASS(autoexpandcategories = ToonProfile, MinimalAPI, meta = (DevelopmentStatus = "Experimental"))
 class UToonProfile : public UObject
 {
 	GENERATED_UCLASS_BODY()
diff --git a/Engine/Source/Runtime/Engine/Private/Materials/MaterialExpressionSubstrate.cpp b/Engine/Source/Runtime/Engine/Private/Materials/MaterialExpressionSubstrate.cpp
--- a/Engine/Source/Runtime/Engine/Private/Materials/MaterialExpressionSubstrate.cpp
+++ b/Engine/Source/Runtime/Engine/Private/Materials/MaterialExpressionSubstrate.cpp
@@ -2531,7 +2531,7 @@ int32 UMaterialExpressionSubstrateToonBSDF::Compile(class FMaterialCompiler* Com
 
 void UMaterialExpressionSubstrateToonBSDF::GetCaption(TArray<FString>& OutCaptions) const
 {
-	OutCaptions.Add(TEXT("Substrate Toon BSDF"));
+	OutCaptions.Add(TEXT("Substrate Toon BSDF - EXPERIMENTAL"));
 }
 
 EMaterialValueType UMaterialExpressionSubstrateToonBSDF::GetOutputValueType(int32 OutputIndex)
diff --git a/Engine/Source/Runtime/Engine/Public/Materials/MaterialExpressionSubstrate.h b/Engine/Source/Runtime/Engine/Public/Materials/MaterialExpressionSubstrate.h
--- a/Engine/Source/Runtime/Engine/Public/Materials/MaterialExpressionSubstrate.h
+++ b/Engine/Source/Runtime/Engine/Public/Materials/MaterialExpressionSubstrate.h
@@ -857,7 +857,7 @@ class UMaterialExpressionSubstrateLightFunction : public UMaterialExpressionSubs
 	//~ End UMaterialExpression Interface
 };
 
-UCLASS(MinimalAPI, collapsecategories, hidecategories = Object, DisplayName = "Substrate Toon BSDF")
+UCLASS(MinimalAPI, collapsecategories, hidecategories = Object, DisplayName = "Substrate Toon BSDF", meta = (DevelopmentStatus = "Experimental"))
 class UMaterialExpressionSubstrateToonBSDF : public UMaterialExpressionSubstrateBSDF
 {
 	GENERATED_UCLASS_BODY()
