Avoid crashing when blend space sample uses a null animation sequence

#Synced-CL 51338101

diff --git a/Engine/Source/Runtime/Engine/Private/Animation/BlendSpace.cpp b/Engine/Source/Runtime/Engine/Private/Animation/BlendSpace.cpp
--- a/Engine/Source/Runtime/Engine/Private/Animation/BlendSpace.cpp
+++ b/Engine/Source/Runtime/Engine/Private/Animation/BlendSpace.cpp
@@ -2044,7 +2044,8 @@ void UBlendSpace::TickFollowerSamples(TArray<FBlendSampleData>& SampleDataList,
 	{
 		FBlendSampleData& SampleDataItem = SampleDataList[SampleIndex];
 		const FBlendSample& Sample = SampleData[SampleDataItem.SampleDataIndex];
-		if (HighestWeightIndex != SampleIndex)
+
+		if (HighestWeightIndex != SampleIndex && Sample.Animation != nullptr)
 		{
 			if (bResetMarkerDataOnFollowers)
 			{
