Exposed a couple of cvars in UPrimitiveComponent

#rnx
#Synced-CL 51210104

diff --git a/Engine/Source/Runtime/Engine/Classes/Components/PrimitiveComponent.h b/Engine/Source/Runtime/Engine/Classes/Components/PrimitiveComponent.h
--- a/Engine/Source/Runtime/Engine/Classes/Components/PrimitiveComponent.h
+++ b/Engine/Source/Runtime/Engine/Classes/Components/PrimitiveComponent.h
@@ -63,8 +63,8 @@ namespace PrimitiveComponentCVars
 {
 	extern float HitDistanceToleranceCVar;
 	extern float InitialOverlapToleranceCVar;
-	extern int32 bAllowCachedOverlapsCVar;
-	extern int32 bEnableFastOverlapCheck;
+	ENGINE_API extern int32 bAllowCachedOverlapsCVar;
+	ENGINE_API extern int32 bEnableFastOverlapCheck;
 }
 
 /** Determines whether a Character can attempt to step up onto a component when they walk in to it. */
