Allow for notify event dispatch to be disabled on a per-mesh component basis.  This doesn't prevent the notifies from being sampled but the associated events will not be dispatched at the end of animation update when the flag is set.

#Synced-CL 51330716

diff --git a/Engine/Source/Editor/UnrealEd/Private/ThumbnailHelpers.cpp b/Engine/Source/Editor/UnrealEd/Private/ThumbnailHelpers.cpp
--- a/Engine/Source/Editor/UnrealEd/Private/ThumbnailHelpers.cpp
+++ b/Engine/Source/Editor/UnrealEd/Private/ThumbnailHelpers.cpp
@@ -643,6 +643,10 @@ FAnimationSequenceThumbnailScene::FAnimationSequenceThumbnailScene()
 	PreviewActor = GetWorld()->SpawnActor<AAnimationThumbnailSkeletalMeshActor>(SpawnInfo);
 
 	PreviewActor->SetActorEnableCollision(false);
+	if (USkeletalMeshComponent* SkelMeshComp = PreviewActor->GetSkeletalMeshComponent())
+	{
+		SkelMeshComp->bSuppressNotifyEventDispatch = true;
+	}
 }
 
 bool FAnimationSequenceThumbnailScene::SetAnimation(UAnimSequenceBase* InAnimation)
@@ -774,6 +778,10 @@ FBlendSpaceThumbnailScene::FBlendSpaceThumbnailScene()
 	PreviewActor = GetWorld()->SpawnActor<AAnimationThumbnailSkeletalMeshActor>(SpawnInfo);
 
 	PreviewActor->SetActorEnableCollision(false);
+	if (USkeletalMeshComponent* SkelMeshComp = PreviewActor->GetSkeletalMeshComponent())
+	{
+		SkelMeshComp->bSuppressNotifyEventDispatch = true;
+	}
 }
 
 bool FBlendSpaceThumbnailScene::SetBlendSpace(class UBlendSpace* InBlendSpace)
@@ -908,6 +916,10 @@ FAnimBlueprintThumbnailScene::FAnimBlueprintThumbnailScene()
 	PreviewActor = GetWorld()->SpawnActor<ASkeletalMeshActor>(SpawnInfo);
 
 	PreviewActor->SetActorEnableCollision(false);
+	if (USkeletalMeshComponent* SkelMeshComp = PreviewActor->GetSkeletalMeshComponent())
+	{
+		SkelMeshComp->bSuppressNotifyEventDispatch = true;
+	}
 }
 
 bool FAnimBlueprintThumbnailScene::SetAnimBlueprint(class UAnimBlueprint* InBlueprint)
diff --git a/Engine/Source/Runtime/Engine/Classes/Components/SkeletalMeshComponent.h b/Engine/Source/Runtime/Engine/Classes/Components/SkeletalMeshComponent.h
--- a/Engine/Source/Runtime/Engine/Classes/Components/SkeletalMeshComponent.h
+++ b/Engine/Source/Runtime/Engine/Classes/Components/SkeletalMeshComponent.h
@@ -844,6 +844,10 @@ public:
 	UPROPERTY(EditAnywhere, BlueprintReadWrite, AdvancedDisplay, Category = Optimization)
 	uint8 bSkipBoundsUpdateWhenInterpolating:1;
 
+	/** Prevent notify events from being dispatched for this mesh. */
+	UPROPERTY(EditAnywhere, BlueprintReadWrite, Category = SkeletalMesh)
+	uint8 bSuppressNotifyEventDispatch : 1;
+
 protected:
 
 	/** Whether the clothing simulation is suspended (not the same as disabled, we no longer run the sim but keep the last valid sim data around) */
@@ -1048,7 +1052,10 @@ PRAGMA_DISABLE_DEPRECATION_WARNINGS
 	const TArray<UAnimInstance*>& GetLinkedAnimInstances() const { return LinkedInstances; }
 PRAGMA_ENABLE_DEPRECATION_WARNINGS
 
-	/** If true, next time that ConditionallyDispatchQueuedAnimEvents() is called it will trigger any queued anim notifies available. */
+	/** 
+	 * If true, next time that ConditionallyDispatchQueuedAnimEvents() is called it will trigger any 
+	 * queued anim notifies available. Only to be used when not performing a regular anim instance update.
+	 */
 	ENGINE_API void AllowQueuedAnimEventsNextDispatch();
 	
 private:
diff --git a/Engine/Source/Runtime/Engine/Private/Components/SkeletalMeshComponent.cpp b/Engine/Source/Runtime/Engine/Private/Components/SkeletalMeshComponent.cpp
--- a/Engine/Source/Runtime/Engine/Private/Components/SkeletalMeshComponent.cpp
+++ b/Engine/Source/Runtime/Engine/Private/Components/SkeletalMeshComponent.cpp
@@ -485,6 +485,8 @@ USkeletalMeshComponent::USkeletalMeshComponent(const FObjectInitializer& ObjectI
 	bEnableUpdateRateOptimizations = false;
 	RagdollAggregateThreshold = UPhysicsSettings::Get()->RagdollAggregateThreshold;
 
+	bSuppressNotifyEventDispatch = false;
+
 	bUpdateMeshWhenKinematic = false;
 	bEnableAnimation = true;
 
@@ -1677,7 +1679,7 @@ void USkeletalMeshComponent::TickAnimation(float DeltaTime, bool bNeedsValidRoot
 	if (GetSkeletalMeshAsset() != nullptr && !bIsCompiling)
 	{
 		// We're about to UpdateAnimation, this will potentially queue events that we'll need to dispatch.
-		bNeedsQueuedAnimEventsDispatched = true;
+		bNeedsQueuedAnimEventsDispatched = !bSuppressNotifyEventDispatch;
 
 		// Tick all of our anim instances
 		TickAnimInstances(DeltaTime, bNeedsValidRootMotion);
@@ -3742,7 +3744,14 @@ PRAGMA_ENABLE_DEPRECATION_WARNINGS
 
 void USkeletalMeshComponent::AllowQueuedAnimEventsNextDispatch()
 {
-	bNeedsQueuedAnimEventsDispatched = true;
+	if (!bSuppressNotifyEventDispatch)
+	{
+		bNeedsQueuedAnimEventsDispatched = true;
+	}
+	else
+	{
+		UE_LOG(LogSkeletalMesh, Warning, TEXT("AllowQueuedAnimEventsNextDispatch called when notify dispatch is disabled on mesh %s. Re-enable dispatch to allow this."), *GetName());
+	}	
 }
 
 UAnimInstance* USkeletalMeshComponent::GetLinkedAnimGraphInstanceByTag(FName InName) const
