[HWRT] Only sort and loop over geometries that are evicted or require streaming in the last step of FRayTracingGeometryManager::UpdateReferenceBasedResidency(...)

- Render thread overhead on MegaLights demo: 0.36ms -> 0.18ms (50% reduction)

#Synced-CL 51330724

diff --git a/Engine/Source/Runtime/Engine/Private/Rendering/RayTracingGeometryManager.cpp b/Engine/Source/Runtime/Engine/Private/Rendering/RayTracingGeometryManager.cpp
--- a/Engine/Source/Runtime/Engine/Private/Rendering/RayTracingGeometryManager.cpp
+++ b/Engine/Source/Runtime/Engine/Private/Rendering/RayTracingGeometryManager.cpp
@@ -792,8 +792,8 @@ void FRayTracingGeometryManager::UpdateReferenceBasedResidency(FRHICommandList&
 {
 	TSet<RayTracing::FGeometryHandle> NotReferencedResidentGeometries = EvictableGeometries;
 
-	TArray<RayTracing::FGeometryHandle> LocalReferencedGeometries;
-	LocalReferencedGeometries.Reserve(RegisteredGeometries.Num());
+	TArray<RayTracing::FGeometryHandle> ReferencedGeometriesEvictedOrRequiringStreaming;
+	ReferencedGeometriesEvictedOrRequiringStreaming.Reserve(RegisteredGeometries.Num());
 
 	int64 ReferencedSize = 0;
 	int64 RequestedSize = TotalAlwaysResidentSize;
@@ -806,7 +806,15 @@ void FRayTracingGeometryManager::UpdateReferenceBasedResidency(FRHICommandList&
 		FRegisteredGeometry& RegisteredGeometry = GetRegisteredGeometry(GeometryHandle);
 		RegisteredGeometry.LastReferencedFrame = GFrameCounterRenderThread;
 
-		LocalReferencedGeometries.Add(GeometryHandle);
+		bool bStreamBVH;
+		bool bStreamBuffers;
+		DoesRayTracingGeometryRequireStreaming(RegisteredGeometry, bStreamBVH, bStreamBuffers);
+
+		if (RegisteredGeometry.bEvicted || bStreamBVH || bStreamBuffers)
+		{
+			ReferencedGeometriesEvictedOrRequiringStreaming.Add(GeometryHandle);
+		}
+
 		NotReferencedResidentGeometries.Remove(GeometryHandle);
 
 		ReferencedSize += RegisteredGeometry.Size;
@@ -823,7 +831,7 @@ void FRayTracingGeometryManager::UpdateReferenceBasedResidency(FRHICommandList&
 	}
 
 	// Step 2
-	// - add all geometries in referenced groups to LocalReferencedGeometries
+	// - add all geometries in referenced groups to ReferencedGeometriesEvictedOrRequiringStreaming
 	//		- need to make all geometries in group resident otherwise might not have valid geometry when reducing LOD
 	//		- TODO: Could track TargetLOD and only make [TargetLOD ... LastLOD] range resident
 	// - also update LastReferencedFrame and calculate memory required to make evicted geometries resident
@@ -849,7 +857,15 @@ void FRayTracingGeometryManager::UpdateReferenceBasedResidency(FRHICommandList&
 				}
 				else
 				{
-					LocalReferencedGeometries.Add(GeometryHandle);
+					bool bStreamBVH;
+					bool bStreamBuffers;
+					DoesRayTracingGeometryRequireStreaming(RegisteredGeometry, bStreamBVH, bStreamBuffers);
+
+					if (RegisteredGeometry.bEvicted || bStreamBVH || bStreamBuffers)
+					{
+						ReferencedGeometriesEvictedOrRequiringStreaming.Add(GeometryHandle);
+					}
+
 					NotReferencedResidentGeometries.Remove(GeometryHandle);
 
 					RequestedSize += RegisteredGeometry.Size;
@@ -864,7 +880,7 @@ void FRayTracingGeometryManager::UpdateReferenceBasedResidency(FRHICommandList&
 	}
 
 #if DO_CHECK
-	// ensure(LocalReferencedGeometries.Num() == TSet(LocalReferencedGeometries).Num());
+	// ensure(ReferencedGeometriesEvictedOrRequiringStreaming.Num() == TSet(ReferencedGeometriesEvictedOrRequiringStreaming).Num());
 #endif
 
 	const int64 ResidentGeometryMemoryPoolSize = FUnitConversion::Convert(GRayTracingResidentGeometryMemoryPoolSizeInMB, EUnit::Megabytes, EUnit::Bytes);
@@ -919,7 +935,7 @@ void FRayTracingGeometryManager::UpdateReferenceBasedResidency(FRHICommandList&
 	{
 		// Step 4.1
 		//	- sort by size to prioritize smaller geometries
-		Algo::Sort(LocalReferencedGeometries, [this](RayTracing::FGeometryHandle& LHSHandle, RayTracing::FGeometryHandle& RHSHandle)
+		Algo::Sort(ReferencedGeometriesEvictedOrRequiringStreaming, [this](RayTracing::FGeometryHandle& LHSHandle, RayTracing::FGeometryHandle& RHSHandle)
 			{
 				FRegisteredGeometry& LHS = GetRegisteredGeometry(LHSHandle);
 				FRegisteredGeometry& RHS = GetRegisteredGeometry(RHSHandle);
@@ -930,9 +946,9 @@ void FRayTracingGeometryManager::UpdateReferenceBasedResidency(FRHICommandList&
 		// Step 4.2
 		// - make geometries resident until we go over budget
 		int32 Index = 0;
-		while (TotalResidentSize + TotalStreamingSize < ResidentGeometryMemoryPoolSize && Index < LocalReferencedGeometries.Num())
+		while (TotalResidentSize + TotalStreamingSize < ResidentGeometryMemoryPoolSize && Index < ReferencedGeometriesEvictedOrRequiringStreaming.Num())
 		{
-			RayTracing::FGeometryHandle GeometryHandle = LocalReferencedGeometries[Index];
+			RayTracing::FGeometryHandle GeometryHandle = ReferencedGeometriesEvictedOrRequiringStreaming[Index];
 			FRegisteredGeometry& RegisteredGeometry = GetRegisteredGeometry(GeometryHandle);
 
 			if (RegisteredGeometry.Geometry->IsEvicted())
@@ -1113,36 +1129,53 @@ void FRayTracingGeometryManager::CheckIntegrity()
 }
 #endif // DO_CHECK
 
-bool FRayTracingGeometryManager::RequestRayTracingGeometryStreamIn(FRHICommandList& RHICmdList, RayTracing::FGeometryHandle GeometryHandle)
+void FRayTracingGeometryManager::DoesRayTracingGeometryRequireStreaming(const FRegisteredGeometry& RegisteredGeometry, bool& bOutStreamBVH, bool& bOutStreamBuffers)
 {
-	FRegisteredGeometry& RegisteredGeometry = GetRegisteredGeometry(GeometryHandle);
-	FRayTracingGeometry* Geometry = RegisteredGeometry.Geometry;
+	bOutStreamBVH = false;
+	bOutStreamBuffers = false;
 
 	if (RegisteredGeometry.Status == FRegisteredGeometry::EStatus::Streaming)
 	{
 		// skip if there's already a streaming request in-flight for this geometry
-		return true;
+		return;
 	}
 
-	const bool bStreamBVH = Geometry->GetInitializer().Type == ERayTracingGeometryInitializerType::StreamingDestination;
-	const bool bStreamBuffers = bStreamBVH || (RegisteredGeometry.StreamableBuffersSize > 0 && !RegisteredGeometry.StreamableAsset->AreBuffersStreamedIn() && ReferencedGeometryGroupsForDynamicUpdate.Contains(RegisteredGeometry.Geometry->GroupHandle));
+	bOutStreamBVH = RegisteredGeometry.Geometry->GetInitializer().Type == ERayTracingGeometryInitializerType::StreamingDestination;
+	bOutStreamBuffers = bOutStreamBVH || (RegisteredGeometry.StreamableBuffersSize > 0 && !RegisteredGeometry.StreamableAsset->AreBuffersStreamedIn() && ReferencedGeometryGroupsForDynamicUpdate.Contains(RegisteredGeometry.Geometry->GroupHandle));
 
-	if (!bStreamBuffers && !bStreamBVH)
+	if (!bOutStreamBVH && !bOutStreamBuffers)
 	{
 		// no streaming required
-		return true;
+		return;
 	}
 
-	if (!Geometry->GroupHandle.IsNull())
+	if (!RegisteredGeometry.Geometry->GroupHandle.IsNull())
 	{
-		const FRayTracingGeometryGroup& Group = GetRegisteredGroup(Geometry->GroupHandle);
+		const FRayTracingGeometryGroup& Group = GetRegisteredGroup(RegisteredGeometry.Geometry->GroupHandle);
 
-		if (Geometry->LODIndex < Group.CurrentFirstLODIdx)
+		if (RegisteredGeometry.Geometry->LODIndex < Group.CurrentFirstLODIdx)
 		{
-			// streaming request no longer necessary
-			return true;
+			bOutStreamBVH = false;
+			bOutStreamBuffers = false;
+			return;
 		}
 	}
+}
+
+bool FRayTracingGeometryManager::RequestRayTracingGeometryStreamIn(FRHICommandList& RHICmdList, RayTracing::FGeometryHandle GeometryHandle)
+{
+	FRegisteredGeometry& RegisteredGeometry = GetRegisteredGeometry(GeometryHandle);
+	FRayTracingGeometry* Geometry = RegisteredGeometry.Geometry;
+
+	bool bStreamBVH;
+	bool bStreamBuffers;
+	DoesRayTracingGeometryRequireStreaming(RegisteredGeometry, bStreamBVH, bStreamBuffers);
+
+	if (!bStreamBuffers && !bStreamBVH)
+	{
+		// no streaming required
+		return true;
+	}
 
 	// TODO: Support DDC streaming
 
diff --git a/Engine/Source/Runtime/Engine/Public/Rendering/RayTracingGeometryManager.h b/Engine/Source/Runtime/Engine/Public/Rendering/RayTracingGeometryManager.h
--- a/Engine/Source/Runtime/Engine/Public/Rendering/RayTracingGeometryManager.h
+++ b/Engine/Source/Runtime/Engine/Public/Rendering/RayTracingGeometryManager.h
@@ -169,6 +169,8 @@ private:
 		}
 	};
 
+	void DoesRayTracingGeometryRequireStreaming(const FRegisteredGeometry& RegisteredGeometry, bool& bOutStreamBVH, bool& bOutStreamBuffers);
+
 	void CancelStreamingRequest(FRegisteredGeometry& RegisteredGeometry);
 
 	// Helper function to stream out geometry
