FMaterial - Make sure ErrorExpressions stays in sync with CompileErrors when clearing errors

Prevents a crash in FHLSLMaterialTranslator::Error() (line 4425) caused by accessing the CompileErrors array at an invalid index

#Synced-CL 51228359

diff --git a/Engine/Source/Runtime/Engine/Private/Materials/HLSLMaterialTranslator.cpp b/Engine/Source/Runtime/Engine/Private/Materials/HLSLMaterialTranslator.cpp
--- a/Engine/Source/Runtime/Engine/Private/Materials/HLSLMaterialTranslator.cpp
+++ b/Engine/Source/Runtime/Engine/Private/Materials/HLSLMaterialTranslator.cpp
@@ -4420,7 +4420,9 @@ int32 FHLSLMaterialTranslator::Error(const TCHAR* Text)
 		if (ExpressionToError)
 		{
 			ExpressionToError->LastErrorText = Text;
-			for (int32 i = 0; i < ErrorExpressions.Num(); i++)
+			ensure(ErrorExpressions.Num() == CompileErrors.Num());
+			const int32 NumErrors = FMath::Min(ErrorExpressions.Num(), CompileErrors.Num());
+			for (int32 i = 0; i < NumErrors; i++)
 			{
 				if(ErrorExpressions[i] == ExpressionToError && CompileErrors[i] == ErrorString)
 				{
diff --git a/Engine/Source/Runtime/Engine/Private/Materials/MaterialShared.cpp b/Engine/Source/Runtime/Engine/Private/Materials/MaterialShared.cpp
--- a/Engine/Source/Runtime/Engine/Private/Materials/MaterialShared.cpp
+++ b/Engine/Source/Runtime/Engine/Private/Materials/MaterialShared.cpp
@@ -3214,6 +3214,7 @@ bool FMaterial::CacheShaders(const FMaterialShaderMapId& ShaderMapId, EMaterialS
 		{
 			// Clear outdated compile errors as we're not calling Translate on this path
 			CompileErrors.Empty();
+			ErrorExpressions.Empty();
 		}
 #endif // WITH_EDITOR
 	}
