Fixed an issue where Viewmode LODColoration was not working on builds with cooked assets unless ODSC was used

#Synced-CL 50300730

diff --git a/Engine/Source/Runtime/Engine/Private/DebugViewModeHelpers.cpp b/Engine/Source/Runtime/Engine/Private/DebugViewModeHelpers.cpp
--- a/Engine/Source/Runtime/Engine/Private/DebugViewModeHelpers.cpp
+++ b/Engine/Source/Runtime/Engine/Private/DebugViewModeHelpers.cpp
@@ -70,6 +70,7 @@ bool IsDebugViewShaderModeODSCOnly()
 	return false;
 }
 
+// UE_DEPRECATED(5.8, "SupportDebugViewModes is deprecated, please use ForceEnableDebugViewSupport/ForceDisableDebugViewSupport instead.")
 bool SupportDebugViewModes()
 {
 	// Allow command line disable of debug view modes.
@@ -79,9 +80,27 @@ bool SupportDebugViewModes()
 	return !bForceDisableDebugViewModes;
 }
 
+
+/** Returns true if we are forcing debug view support even for non editor or ODSC use cases. */
+bool ForceEnableDebugViewSupport(EShaderPlatform Platform)
+{
+	static FShaderPlatformCachedIniValue<int32> CVar(TEXT("r.ForceDebugViewModes"));
+	return CVar.Get(Platform) == 1;
+}
+
+/** Returns true if we are force disabling debug view support even for editor use cases. */
+bool ForceDisableDebugViewSupport(EShaderPlatform Platform)
+{
+	// Allow command line disable of debug view modes.
+	// This is intended only for a very specific commandlet case where we can demand debug view shaders through ODSC even when they are not actually needed.
+	// Setting this command line avoids _any_ potential overhead from debug shader compilation.
+	static const bool bForceDisableDebugViewModes = FParse::Param(FCommandLine::Get(), TEXT("forceDisableDebugViewModes"));
+	return bForceDisableDebugViewModes;
+}
+
 bool SupportDebugViewShaderMode(EDebugViewShaderMode ShaderMode, EShaderPlatform Platform)
 {
-	if (!SupportDebugViewModes())
+	if (ForceDisableDebugViewSupport(Platform))
 	{
 		return false;
 	}
@@ -134,6 +153,11 @@ bool AllowDebugViewShaderMode(EDebugViewShaderMode ShaderMode, EShaderPlatform P
 		return false;
 	}
 
+	if (ForceEnableDebugViewSupport(Platform))
+	{
+		return true;
+	}
+
 #if WITH_EDITORONLY_DATA
 	bool bRequireODSCShaders = IsDebugViewShaderModeODSCOnly();
 #else
diff --git a/Engine/Source/Runtime/Engine/Public/DebugViewModeHelpers.h b/Engine/Source/Runtime/Engine/Public/DebugViewModeHelpers.h
--- a/Engine/Source/Runtime/Engine/Public/DebugViewModeHelpers.h
+++ b/Engine/Source/Runtime/Engine/Public/DebugViewModeHelpers.h
@@ -43,17 +43,27 @@ ENGINE_API const TCHAR* DebugViewShaderModeToString(EDebugViewShaderMode InShade
 #if WITH_DEBUG_VIEW_MODES
 /** Returns true if debug view shaders are only ODSC compiled (in editor or cooked builds). Returns false if they are can also be compiled for editor only shader maps as well as ODSC. */
 ENGINE_API bool IsDebugViewShaderModeODSCOnly();
-/** Returns true if any debug view rendering is supported. Can be used as input for shader cooking. */
+/** DEPRECATED 5.8 - Returns true if any debug view rendering is supported. Can be used as input for shader cooking. */
+UE_DEPRECATED(5.8, "SupportDebugViewModes is deprecated, please use ForceEnableDebugViewSupport/ForceDisableDebugViewSupport instead.")
 ENGINE_API bool SupportDebugViewModes();
 /** Returns true if a specific debug view supported for a platform. Can be used as input for shader cooking. */
 ENGINE_API bool SupportDebugViewShaderMode(EDebugViewShaderMode ShaderMode, EShaderPlatform Platform);
 /** Returns true if a specific debug view can be displayed at runtime. */
 ENGINE_API bool AllowDebugViewShaderMode(EDebugViewShaderMode ShaderMode, EShaderPlatform Platform, ERHIFeatureLevel::Type FeatureLevel);
+/** Returns true if we are forcing debug view support even for non editor or ODSC use cases. */
+ENGINE_API bool ForceEnableDebugViewSupport(EShaderPlatform Platform);
+/** Returns true if we are force disabling debug view support even for editor use cases. */
+ENGINE_API bool ForceDisableDebugViewSupport(EShaderPlatform Platform);
+
 #else
 inline bool IsDebugViewShaderModeODSCOnly() { return false; }
+/** DEPRECATED 5.8 - Returns true if any debug view rendering is supported. Can be used as input for shader cooking. */
+UE_DEPRECATED(5.8, "SupportDebugViewModes is deprecated, please use ForceEnableDebugViewSupport/ForceDisableDebugViewSupport instead.")
 inline bool SupportDebugViewModes() { return false; }
 inline bool SupportDebugViewShaderMode(EDebugViewShaderMode ShaderMode, EShaderPlatform Platform) { return false; }
 inline bool AllowDebugViewShaderMode(EDebugViewShaderMode ShaderMode, EShaderPlatform Platform, ERHIFeatureLevel::Type FeatureLevel) { return false; }
+inline bool ForceEnableDebugViewSupport(EShaderPlatform Platform) { return false; }
+inline bool ForceDisableDebugViewSupport(EShaderPlatform Platform) { return false; }
 #endif
 
 ENGINE_API int32 GetNumActorsInWorld(UWorld* InWorld);
diff --git a/Engine/Source/Runtime/Renderer/Private/DebugViewModeRendering.cpp b/Engine/Source/Runtime/Renderer/Private/DebugViewModeRendering.cpp
--- a/Engine/Source/Runtime/Renderer/Private/DebugViewModeRendering.cpp
+++ b/Engine/Source/Runtime/Renderer/Private/DebugViewModeRendering.cpp
@@ -162,7 +162,12 @@ TRDGUniformBufferRef<FDebugViewModePassUniformParameters> CreateDebugViewModePas
 
 static bool ShouldCompileDebugViewModeShaders(const FMeshMaterialShaderPermutationParameters& Parameters)
 {
-	if (!SupportDebugViewModes())
+	if (ForceEnableDebugViewSupport(Parameters.Platform))
+	{
+		return true;
+	}
+
+	if (ForceDisableDebugViewSupport(Parameters.Platform))
 	{
 		return false;
 	}
