Deprecated & disabled the r.SkipRedundantTransformUpdate feature and associated API as it leads to inconsistencies, ASAN warnings and is of doubtful performance value (possibly detriment).

#Synced-CL 51368499
#Synced-CL 51369525

diff --git a/Engine/Plugins/Experimental/MovieSceneAnimMixer/Source/MovieSceneAnimMixer/Private/Systems/MovieSceneAnimNextTargetSystem.cpp b/Engine/Plugins/Experimental/MovieSceneAnimMixer/Source/MovieSceneAnimMixer/Private/Systems/MovieSceneAnimNextTargetSystem.cpp
--- a/Engine/Plugins/Experimental/MovieSceneAnimMixer/Source/MovieSceneAnimMixer/Private/Systems/MovieSceneAnimNextTargetSystem.cpp
+++ b/Engine/Plugins/Experimental/MovieSceneAnimMixer/Source/MovieSceneAnimMixer/Private/Systems/MovieSceneAnimNextTargetSystem.cpp
@@ -120,7 +120,9 @@ namespace UE::MovieScene
 			State.SaveState(SkelMeshComp);
 			if (FPrimitiveSceneProxy* Proxy = SkelMeshComp->GetSceneProxy())
 			{
+PRAGMA_DISABLE_DEPRECATION_WARNINGS
 				Proxy->SetCanSkipRedundantTransformUpdates(false);
+PRAGMA_ENABLE_DEPRECATION_WARNINGS
 			}
 			return State;
 		}
@@ -132,7 +134,9 @@ namespace UE::MovieScene
 			{
 				if (FPrimitiveSceneProxy* Proxy = SkelMeshComp->GetSceneProxy())
 				{
+PRAGMA_DISABLE_DEPRECATION_WARNINGS
 					Proxy->SetCanSkipRedundantTransformUpdates(true);
+PRAGMA_ENABLE_DEPRECATION_WARNINGS
 				}
 				InOutCachedValue.RestoreState();
 			}
diff --git a/Engine/Plugins/Runtime/GeometryCache/Source/GeometryCache/Private/GeometryCacheSceneProxy.cpp b/Engine/Plugins/Runtime/GeometryCache/Source/GeometryCache/Private/GeometryCacheSceneProxy.cpp
--- a/Engine/Plugins/Runtime/GeometryCache/Source/GeometryCache/Private/GeometryCacheSceneProxy.cpp
+++ b/Engine/Plugins/Runtime/GeometryCache/Source/GeometryCache/Private/GeometryCacheSceneProxy.cpp
@@ -111,8 +111,9 @@ FGeometryCacheSceneProxy::FGeometryCacheSceneProxy(UGeometryCacheComponent* Comp
 
 	EnableGPUSceneSupportFlags();
 
+PRAGMA_DISABLE_DEPRECATION_WARNINGS
 	bCanSkipRedundantTransformUpdates = false;
-
+PRAGMA_ENABLE_DEPRECATION_WARNINGS
 	// All tracks use the same array of materials since the MaterialIndex in FGeometryCacheMeshBatchInfo 
 	// is an index into the global material array UGeometryCache::Materials
 	TArray<UMaterialInterface*> Materials;
diff --git a/Engine/Source/Runtime/Engine/Classes/Components/SkinnedMeshComponent.h b/Engine/Source/Runtime/Engine/Classes/Components/SkinnedMeshComponent.h
--- a/Engine/Source/Runtime/Engine/Classes/Components/SkinnedMeshComponent.h
+++ b/Engine/Source/Runtime/Engine/Classes/Components/SkinnedMeshComponent.h
@@ -1674,6 +1674,7 @@ public:
 	/**
 	 * Toggles whether this skinned mesh component can issue duplicate transform updates
 	 */
+	UE_DEPRECATED(5.8, "Duplicate transform control will be removed.")
 	ENGINE_API void AllowDuplicateTransformUpdates(bool bAllow);
 
 	/** Access ComponentSpaceTransforms for reading */
diff --git a/Engine/Source/Runtime/Engine/Private/Components/SkinnedMeshComponent.cpp b/Engine/Source/Runtime/Engine/Private/Components/SkinnedMeshComponent.cpp
--- a/Engine/Source/Runtime/Engine/Private/Components/SkinnedMeshComponent.cpp
+++ b/Engine/Source/Runtime/Engine/Private/Components/SkinnedMeshComponent.cpp
@@ -647,7 +647,9 @@ void USkinnedMeshComponent::AllowDuplicateTransformUpdates(bool bAllow)
 {
 	if (SceneProxy)
 	{
+PRAGMA_DISABLE_DEPRECATION_WARNINGS
 		SceneProxy->SetCanSkipRedundantTransformUpdates(!bAllow);
+PRAGMA_ENABLE_DEPRECATION_WARNINGS
 	}
 }
 
diff --git a/Engine/Source/Runtime/Engine/Private/PrimitiveSceneProxy.cpp b/Engine/Source/Runtime/Engine/Private/PrimitiveSceneProxy.cpp
--- a/Engine/Source/Runtime/Engine/Private/PrimitiveSceneProxy.cpp
+++ b/Engine/Source/Runtime/Engine/Private/PrimitiveSceneProxy.cpp
@@ -277,7 +277,10 @@ FPrimitiveSceneProxyDesc::FPrimitiveSceneProxyDesc(const UPrimitiveComponent* In
 	InitializeFromPrimitiveComponent(InComponent);
 }
 
+PRAGMA_DISABLE_DEPRECATION_WARNINGS
+// NOTE: because of: bCanSkipRedundantTransformUpdates, remove when cleaned up
 FPrimitiveSceneProxy::FPrimitiveSceneProxy(FPrimitiveSceneProxy const&) = default;
+PRAGMA_ENABLE_DEPRECATION_WARNINGS
 
 
 FPrimitiveSceneProxy::FPrimitiveSceneProxy(const UPrimitiveComponent* InComponent, FName InResourceName)
@@ -325,7 +328,9 @@ FPrimitiveSceneProxy::FPrimitiveSceneProxy(const FPrimitiveSceneProxyDesc& InPro
 ,	bSupportsParallelGDRTI(true)
 ,	bVisibleInLumenScene(false)
 ,	bOpaqueOrMasked(true)
+PRAGMA_DISABLE_DEPRECATION_WARNINGS
 ,	bCanSkipRedundantTransformUpdates(true)
+PRAGMA_ENABLE_DEPRECATION_WARNINGS
 ,	bGoodCandidateForCachedShadowmap(true)
 ,	bNeedsUnbuiltPreviewLighting(!InProxyDesc.IsPrecomputedLightingValid())
 ,	bHasValidSettingsForStaticLighting(InProxyDesc.HasValidSettingsForStaticLighting())
diff --git a/Engine/Source/Runtime/Engine/Public/PrimitiveSceneProxy.h b/Engine/Source/Runtime/Engine/Public/PrimitiveSceneProxy.h
--- a/Engine/Source/Runtime/Engine/Public/PrimitiveSceneProxy.h
+++ b/Engine/Source/Runtime/Engine/Public/PrimitiveSceneProxy.h
@@ -617,17 +617,23 @@ public:
 	/**
 	* @return true if the proxy can skip redundant transform updates where applicable.
 	*/
+	UE_DEPRECATED(5.8, "Will be removed as it is not generally thread safe.")
 	bool CanSkipRedundantTransformUpdates() const
 	{
+PRAGMA_DISABLE_DEPRECATION_WARNINGS
 		return bCanSkipRedundantTransformUpdates;
+PRAGMA_ENABLE_DEPRECATION_WARNINGS
 	}
 
 	/**
 	* Set whether the proxy can skip redundant transform updates where applicable.
 	*/
+	UE_DEPRECATED(5.8, "Will be removed as it is not generally thread safe.")
 	void SetCanSkipRedundantTransformUpdates(bool bCanSkip)
 	{
+PRAGMA_DISABLE_DEPRECATION_WARNINGS
 		bCanSkipRedundantTransformUpdates = bCanSkip;
+PRAGMA_ENABLE_DEPRECATION_WARNINGS
 	}
 
 	/**
@@ -1390,6 +1396,7 @@ protected:
 	uint8 bOpaqueOrMasked : 1;
 
 	/** Whether this component can skip redundant transform updates where applicable. */
+	UE_DEPRECATED(5.8, "Will be removed as it is not generally thread safe.")
 	uint8 bCanSkipRedundantTransformUpdates : 1;
 
 	/** Whether this proxy's mesh is unlikely to be constantly changing. */
@@ -1786,6 +1793,7 @@ private:
 	 */
 	ENGINE_API void SetTransform(FRHICommandListBase& RHICmdList, const FMatrix& InLocalToWorld, const FBoxSphereBounds& InBounds, const FBoxSphereBounds& InLocalBounds, FVector InActorPosition);
 
+	UE_DEPRECATED(5.8, "Will be removed as it is not generally thread safe.")
 	ENGINE_API bool WouldSetTransformBeRedundant_AnyThread(const FMatrix& InLocalToWorld, const FBoxSphereBounds& InBounds, const FBoxSphereBounds& InLocalBounds, const FVector& InActorPosition) const;
 
 	ENGINE_API void CreateUniformBuffer();
diff --git a/Engine/Source/Runtime/Experimental/GeometryCollectionEngine/Private/GeometryCollection/GeometryCollectionSceneProxy.cpp b/Engine/Source/Runtime/Experimental/GeometryCollectionEngine/Private/GeometryCollection/GeometryCollectionSceneProxy.cpp
--- a/Engine/Source/Runtime/Experimental/GeometryCollectionEngine/Private/GeometryCollection/GeometryCollectionSceneProxy.cpp
+++ b/Engine/Source/Runtime/Experimental/GeometryCollectionEngine/Private/GeometryCollection/GeometryCollectionSceneProxy.cpp
@@ -1606,7 +1606,9 @@ void FNaniteGeometryCollectionSceneProxy::UpdateInstanceSceneDataBuffers(const F
 
 	// set the prev by copying the last current
 	ProxyData.PrevInstanceToPrimitiveRelative = ProxyData.InstanceToPrimitiveRelative;
+PRAGMA_DISABLE_DEPRECATION_WARNINGS
 	bCanSkipRedundantTransformUpdates = false; // should we compare the transform to better decide about this ? 
+PRAGMA_ENABLE_DEPRECATION_WARNINGS
 
 	ProxyData.InstanceToPrimitiveRelative.Reset(TransformCount);
 	ProxyData.InstanceLocalBounds.Reset(TransformCount);
@@ -1650,7 +1652,9 @@ void FNaniteGeometryCollectionSceneProxy::UpdateInstanceSceneDataBuffers(const F
 	if (ProxyData.PrevInstanceToPrimitiveRelative.Num() != ProxyData.InstanceToPrimitiveRelative.Num())
 	{
 		ProxyData.PrevInstanceToPrimitiveRelative = ProxyData.InstanceToPrimitiveRelative;
+PRAGMA_DISABLE_DEPRECATION_WARNINGS
 		bCanSkipRedundantTransformUpdates = true;
+PRAGMA_ENABLE_DEPRECATION_WARNINGS
 	}
 
 	InstanceSceneDataBuffersImpl.EndWriteAccess(AccessTag);
diff --git a/Engine/Source/Runtime/Renderer/Private/RendererScene.cpp b/Engine/Source/Runtime/Renderer/Private/RendererScene.cpp
--- a/Engine/Source/Runtime/Renderer/Private/RendererScene.cpp
+++ b/Engine/Source/Runtime/Renderer/Private/RendererScene.cpp
@@ -1522,14 +1522,16 @@ void FScene::BatchAddPrimitives(TArrayView<FPrimitiveSceneDesc*> InPrimitives)
 
 static int32 GWarningOnRedundantTransformUpdate = 0;
 static FAutoConsoleVariableRef CVarWarningOnRedundantTransformUpdate(
+	FCVarDeprecationArgs{ TEXT("5.8") },
 	TEXT("r.WarningOnRedundantTransformUpdate"),
 	GWarningOnRedundantTransformUpdate,
 	TEXT("Produce a warning when UpdatePrimitiveTransform is called redundantly."),
 	ECVF_Default
 );
 
-static int32 GSkipRedundantTransformUpdate = 1;
+static int32 GSkipRedundantTransformUpdate = 0;
 static FAutoConsoleVariableRef CVarSkipRedundantTransformUpdate(
+	FCVarDeprecationArgs{ TEXT("5.8") },
 	TEXT("r.SkipRedundantTransformUpdate"),
 	GSkipRedundantTransformUpdate,
 	TEXT("Skip updates UpdatePrimitiveTransform is called redundantly, if the proxy allows it."),
@@ -1633,10 +1635,14 @@ void FScene::UpdatePrimitiveTransformInternal(T* Primitive)
 
 			bool bPerformUpdate = true;
 
+PRAGMA_DISABLE_DEPRECATION_WARNINGS
 			const bool bAllowSkip = GSkipRedundantTransformUpdate && Primitive->GetSceneProxy()->CanSkipRedundantTransformUpdates();
+PRAGMA_ENABLE_DEPRECATION_WARNINGS
 			if (bAllowSkip || GWarningOnRedundantTransformUpdate)
 			{
+PRAGMA_DISABLE_DEPRECATION_WARNINGS
 				if (Primitive->GetSceneProxy()->WouldSetTransformBeRedundant_AnyThread(
+PRAGMA_ENABLE_DEPRECATION_WARNINGS
 					UpdateParams.LocalToWorld,
 					UpdateParams.WorldBounds,
 					UpdateParams.LocalBounds,
