Added default implementations of ctor/assignment operator implementations to be able to suppress deprecation warnings.

#rnx

#Synced-CL 51370156

diff --git a/Engine/Source/Runtime/Engine/Private/PrimitiveSceneProxy.cpp b/Engine/Source/Runtime/Engine/Private/PrimitiveSceneProxy.cpp
--- a/Engine/Source/Runtime/Engine/Private/PrimitiveSceneProxy.cpp
+++ b/Engine/Source/Runtime/Engine/Private/PrimitiveSceneProxy.cpp
@@ -277,12 +277,15 @@ FPrimitiveSceneProxyDesc::FPrimitiveSceneProxyDesc(const UPrimitiveComponent* In
 	InitializeFromPrimitiveComponent(InComponent);
 }
 
+// NOTE: required because of: bCanSkipRedundantTransformUpdates, remove when cleaned up
 PRAGMA_DISABLE_DEPRECATION_WARNINGS
-// NOTE: because of: bCanSkipRedundantTransformUpdates, remove when cleaned up
+FPrimitiveSceneProxy::FPrimitiveSceneProxy() = default;
 FPrimitiveSceneProxy::FPrimitiveSceneProxy(FPrimitiveSceneProxy const&) = default;
+FPrimitiveSceneProxy::FPrimitiveSceneProxy(FPrimitiveSceneProxy&&) = default;
+FPrimitiveSceneProxy& FPrimitiveSceneProxy::operator=(const FPrimitiveSceneProxy&) = default;
+FPrimitiveSceneProxy& FPrimitiveSceneProxy::operator=(FPrimitiveSceneProxy&&) = default;
 PRAGMA_ENABLE_DEPRECATION_WARNINGS
 
-
 FPrimitiveSceneProxy::FPrimitiveSceneProxy(const UPrimitiveComponent* InComponent, FName InResourceName)
 	: FPrimitiveSceneProxy(FPrimitiveSceneProxyDesc(InComponent), InResourceName)
 {
diff --git a/Engine/Source/Runtime/Engine/Public/PrimitiveSceneProxy.h b/Engine/Source/Runtime/Engine/Public/PrimitiveSceneProxy.h
--- a/Engine/Source/Runtime/Engine/Public/PrimitiveSceneProxy.h
+++ b/Engine/Source/Runtime/Engine/Public/PrimitiveSceneProxy.h
@@ -296,9 +296,12 @@ public:
 	ENGINE_API FPrimitiveSceneProxy(const UPrimitiveComponent* InComponent, FName ResourceName = NAME_None);
 	ENGINE_API FPrimitiveSceneProxy(const FPrimitiveSceneProxyDesc& InDesc, FName ResourceName = NAME_None);
 
-	/** Copy constructor. */
+	/** Remaining constructors, and assignment operators. */
+	ENGINE_API FPrimitiveSceneProxy();
 	ENGINE_API FPrimitiveSceneProxy(FPrimitiveSceneProxy const&);
-
+	ENGINE_API FPrimitiveSceneProxy(FPrimitiveSceneProxy&&);
+	ENGINE_API FPrimitiveSceneProxy& operator=(const FPrimitiveSceneProxy&);
+	ENGINE_API FPrimitiveSceneProxy& operator=(FPrimitiveSceneProxy&&);
 	/** Virtual destructor. */
 	ENGINE_API virtual ~FPrimitiveSceneProxy();
 
