Expose a const accessor to the current UInputComponent stack on APlayerController. This can be useful when making debug tools for input ! This is only available in test/non-shipping builds.

#rnx

#Synced-CL 51092797

diff --git a/Engine/Source/Runtime/Engine/Classes/GameFramework/PlayerController.h b/Engine/Source/Runtime/Engine/Classes/GameFramework/PlayerController.h
--- a/Engine/Source/Runtime/Engine/Classes/GameFramework/PlayerController.h
+++ b/Engine/Source/Runtime/Engine/Classes/GameFramework/PlayerController.h
@@ -1682,6 +1682,13 @@ public:
 	/** Returns the list of netcondition groups we are part of. */
 	const TArray<FName>& GetNetConditionGroups() const { return NetConditionGroups; }
 
+#if !(UE_BUILD_SHIPPING || UE_BUILD_TEST)
+	// The current input component stack which is built every frame in APlayerController::BuildInputStack.
+	// This is useful for building any external debug tools where you have access to the player controller
+	// It is not recommended to access this outside of debug tooling because this is re-built every frame.
+	TConstArrayView<TWeakObjectPtr<const UInputComponent>> Debug_GetMostRecentInputStack() const { return { LastUsedInputStack }; }
+#endif
+	
 private:
 	
 	/** List of netcondition groups we are currently a member of. */
@@ -1696,6 +1703,11 @@ protected:
 
 	/** Internal.  Current stack of InputComponents. */
 	TArray< TWeakObjectPtr<UInputComponent> > CurrentInputStack;
+
+#if !(UE_BUILD_SHIPPING || UE_BUILD_TEST)
+	// Debug tracking of the most recent *complete* input stack.
+	TArray<TWeakObjectPtr<const UInputComponent>> LastUsedInputStack;
+#endif
 	
 	/** InputComponent we use when player is in Inactive state. */
 	UPROPERTY()
diff --git a/Engine/Source/Runtime/Engine/Private/PlayerController.cpp b/Engine/Source/Runtime/Engine/Private/PlayerController.cpp
--- a/Engine/Source/Runtime/Engine/Private/PlayerController.cpp
+++ b/Engine/Source/Runtime/Engine/Private/PlayerController.cpp
@@ -2752,6 +2752,10 @@ void APlayerController::ProcessPlayerInput(const float DeltaTime, const bool bGa
 {
 	static TArray<UInputComponent*> InputStack;
 
+#if !(UE_BUILD_SHIPPING || UE_BUILD_TEST)
+	LastUsedInputStack.Reset();
+#endif
+	
 	// must be called non-recursively and on the game thread
 	check(IsInGameThread() && !InputStack.Num());
 
@@ -2767,6 +2771,15 @@ void APlayerController::ProcessPlayerInput(const float DeltaTime, const bool bGa
 		PlayerInput->ProcessInputStack(InputStack, DeltaTime, bGamePaused);
 	}
 
+#if !(UE_BUILD_SHIPPING || UE_BUILD_TEST)
+	{
+		for (const UInputComponent* IC : InputStack)
+		{
+			LastUsedInputStack.Add(IC);
+		}
+	}
+#endif
+	
 	InputStack.Reset();
 }
 
