Moved some LOG to LOGF.

#rnx
#tests Tested the log locally on pc editor.

#Synced-CL 51369972

diff --git a/Engine/Source/Runtime/Engine/Private/Components/SkyAtmosphereComponent.cpp b/Engine/Source/Runtime/Engine/Private/Components/SkyAtmosphereComponent.cpp
--- a/Engine/Source/Runtime/Engine/Private/Components/SkyAtmosphereComponent.cpp
+++ b/Engine/Source/Runtime/Engine/Private/Components/SkyAtmosphereComponent.cpp
@@ -264,7 +264,7 @@ void USkyAtmosphereComponent::SendRenderTransform_Concurrent()
 	{
 		if (GSkyAtmosphereRaceConditionFatalLog)
 		{
-			UE_LOG(LogSkyAtmosphereComponent, Fatal, TEXT("SendRenderTransform_Concurrent attempted while DestroyRenderState_Concurrent is being executed"));
+			UE_LOGF(LogSkyAtmosphereComponent, Fatal, "SendRenderTransform_Concurrent attempted while DestroyRenderState_Concurrent is being executed");
 		}
 		FDebug::DumpStackTraceToLog(ELogVerbosity::Warning);
 		ensureMsgf(0, TEXT("SendRenderTransform_Concurrent attempted while DestroyRenderState_Concurrent is being executed"));
@@ -284,7 +284,7 @@ void USkyAtmosphereComponent::DestroyRenderState_Concurrent()
 	{
 		if (GSkyAtmosphereRaceConditionFatalLog)
 		{
-			UE_LOG(LogSkyAtmosphereComponent, Fatal, TEXT("DestroyRenderState_Concurrent called multiple times"));
+			UE_LOGF(LogSkyAtmosphereComponent, Fatal, "DestroyRenderState_Concurrent called multiple times");
 		}
 		FDebug::DumpStackTraceToLog(ELogVerbosity::Warning);
 		ensureMsgf(0, TEXT("DestroyRenderState_Concurrent called multiple times"));
diff --git a/Engine/Source/Runtime/Renderer/Private/SkyAtmosphereRendering.cpp b/Engine/Source/Runtime/Renderer/Private/SkyAtmosphereRendering.cpp
--- a/Engine/Source/Runtime/Renderer/Private/SkyAtmosphereRendering.cpp
+++ b/Engine/Source/Runtime/Renderer/Private/SkyAtmosphereRendering.cpp
@@ -675,7 +675,7 @@ void FScene::RemoveSkyAtmosphere(FSkyAtmosphereSceneProxy* SkyAtmosphereScenePro
 			{
 				if (GSkyAtmosphereStompFatalLog)
 				{
-					UE_LOG(LogSkyAtmosphereRendering, Fatal, TEXT("SkyAtmosphereSceneProxy's RenderSceneInfo was previously deleted. Use after free detected!"));
+					UE_LOGF(LogSkyAtmosphereRendering, Fatal, "SkyAtmosphereSceneProxy's RenderSceneInfo was previously deleted. Use after free detected!");
 				}
 				ensureMsgf(0, TEXT("SkyAtmosphereSceneProxy's RenderSceneInfo was previously deleted. Use after free detected!"));
 				return;
@@ -684,7 +684,7 @@ void FScene::RemoveSkyAtmosphere(FSkyAtmosphereSceneProxy* SkyAtmosphereScenePro
 			{
 				if (GSkyAtmosphereStompFatalLog)
 				{
-					UE_LOG(LogSkyAtmosphereRendering, Fatal, TEXT("SkyAtmosphereSceneProxy's RenderSceneInfo memory stomp detected!"));
+					UE_LOGF(LogSkyAtmosphereRendering, Fatal, "SkyAtmosphereSceneProxy's RenderSceneInfo memory stomp detected!");
 				}
 				ensureMsgf(0, TEXT("SkyAtmosphereSceneProxy's RenderSceneInfo memory stomp detected!"));
 				return;
diff --git a/Engine/Source/Runtime/Renderer/Private/VolumetricFogLightFunction.cpp b/Engine/Source/Runtime/Renderer/Private/VolumetricFogLightFunction.cpp
--- a/Engine/Source/Runtime/Renderer/Private/VolumetricFogLightFunction.cpp
+++ b/Engine/Source/Runtime/Renderer/Private/VolumetricFogLightFunction.cpp
@@ -152,8 +152,8 @@ void FSceneRenderer::RenderLightFunctionForVolumetricFog(
 		if (LiFuncResXU32 > MaxTextureResU32 || LiFuncResYU32 > MaxTextureResU32)
 		{
 #if !(UE_BUILD_SHIPPING)
-			UE_LOG(LogRenderer, Error,
-				TEXT("Invalid LightFunctionResolution %dx%d, View={ %s}, LightDirection={ %s }"),
+			UE_LOGF(LogRenderer, Error,
+				"Invalid LightFunctionResolution %dx%d, View={ %ls}, LightDirection={ %ls }",
 				LightFunctionResolution.X,
 				LightFunctionResolution.Y,
 				*View.ViewMatrices.GetOverriddenTranslatedViewMatrix().ToString(),
diff --git a/Engine/Source/Runtime/Renderer/Private/VolumetricRenderTarget.cpp b/Engine/Source/Runtime/Renderer/Private/VolumetricRenderTarget.cpp
--- a/Engine/Source/Runtime/Renderer/Private/VolumetricRenderTarget.cpp
+++ b/Engine/Source/Runtime/Renderer/Private/VolumetricRenderTarget.cpp
@@ -238,7 +238,7 @@ void FVolumetricRenderTargetViewStateData::Initialise(
 	bValid = TextureResolutionInCopy.X > 0 && TextureResolutionInCopy.Y > 0;
 	if (!bValid)
 	{
-		UE_LOG(LogRenderer, Warning, TEXT("Warning: A viewport of resolution 0x0 was specified - VolumetricCloud not rendered."));
+		UE_LOGF(LogRenderer, Warning, "Warning: A viewport of resolution 0x0 was specified - VolumetricCloud not rendered.");
 		return;
 	}
 
