Emit bUsingNewHLSLGenerator in material shader key generation by switching from Record to RecordAndEmit.

#rnx

#Synced-CL 51218577

diff --git a/Engine/Source/Runtime/Engine/Private/Materials/MaterialShader.cpp b/Engine/Source/Runtime/Engine/Private/Materials/MaterialShader.cpp
--- a/Engine/Source/Runtime/Engine/Private/Materials/MaterialShader.cpp
+++ b/Engine/Source/Runtime/Engine/Private/Materials/MaterialShader.cpp
@@ -1785,7 +1785,7 @@ void FMaterialShaderMapId::RecordAndEmit(FMaterialKeyGeneratorContext& Context)
 		Context.RecordAndEmit("ExternalCodeReferencesHash", ExternalCodeReferencesHash);
 	}
 
-	Context.Record("bUsingNewHLSLGenerator", bUsingNewHLSLGenerator);
+	Context.RecordAndEmit("bUsingNewHLSLGenerator", bUsingNewHLSLGenerator);
 
 	if (PreferredShaderCodeFormat != NAME_None)
 	{
