[PostProcessing] Replace FilmGrainTexelSize with float2 FilmGrainScale. This allows scaling the noise to different aspect ratios, like you would expect with anamorphic squeezing. Also added explicit postprocess setting to turn off DoF.

#Synced-CL 51354224

diff --git a/Engine/Plugins/Cameras/GameplayCameras/Source/GameplayCameras/Private/Compat/PostProcessUtils.cpp b/Engine/Plugins/Cameras/GameplayCameras/Source/GameplayCameras/Private/Compat/PostProcessUtils.cpp
--- a/Engine/Plugins/Cameras/GameplayCameras/Source/GameplayCameras/Private/Compat/PostProcessUtils.cpp
+++ b/Engine/Plugins/Cameras/GameplayCameras/Source/GameplayCameras/Private/Compat/PostProcessUtils.cpp
@@ -131,7 +131,7 @@ bool FPostProcessUtils::OverridePostProcessSettings(FPostProcessSettings& ThisFr
 		UE_SET_PP(FilmGrainShadowsMax);
 		UE_SET_PP(FilmGrainHighlightsMin);
 		UE_SET_PP(FilmGrainHighlightsMax);
-		UE_SET_PP(FilmGrainTexelSize);
+		UE_SET_PP(FilmGrainScale);
 		UE_SET_PP(AmbientOcclusionIntensity);
 		UE_SET_PP(AmbientOcclusionStaticFraction);
 		UE_SET_PP(AmbientOcclusionRadius);
@@ -147,9 +147,8 @@ bool FPostProcessUtils::OverridePostProcessSettings(FPostProcessSettings& ThisFr
 		UE_SET_PP(AmbientOcclusionTemporalBlendWeight);
 		UE_SET_PP(IndirectLightingColor);
 		UE_SET_PP(IndirectLightingIntensity);
-
+		UE_SET_PP(DepthOfFieldEnabled);
 		UE_SET_PP(DepthOfFieldFocalDistance);
-
 		UE_SET_PP(DepthOfFieldFstop);
 		UE_SET_PP(DepthOfFieldMinFstop);
 		UE_SET_PP(DepthOfFieldSensorWidth);
@@ -422,7 +421,7 @@ bool FPostProcessUtils::BlendPostProcessSettings(FPostProcessSettings& ThisFrom,
 		UE_LERP_PP(FilmGrainShadowsMax);
 		UE_LERP_PP(FilmGrainHighlightsMin);
 		UE_LERP_PP(FilmGrainHighlightsMax);
-		UE_LERP_PP(FilmGrainTexelSize);
+		UE_LERP_PP(FilmGrainScale);
 		UE_LERP_PP(AmbientOcclusionIntensity);
 		UE_LERP_PP(AmbientOcclusionStaticFraction);
 		UE_LERP_PP(AmbientOcclusionRadius);
@@ -451,6 +450,7 @@ bool FPostProcessUtils::BlendPostProcessSettings(FPostProcessSettings& ThisFrom,
 			}
 			bAnyOverwritten = true;
 		}
+		UE_SET_PP(DepthOfFieldEnabled)
 		UE_LERP_PP(DepthOfFieldFstop);
 		UE_LERP_PP(DepthOfFieldMinFstop);
 		UE_LERP_PP(DepthOfFieldSensorWidth);
diff --git a/Engine/Plugins/Cameras/GameplayCameras/Source/GameplayCameras/Private/Debug/PostProcessSettingsDebugBlock.cpp b/Engine/Plugins/Cameras/GameplayCameras/Source/GameplayCameras/Private/Debug/PostProcessSettingsDebugBlock.cpp
--- a/Engine/Plugins/Cameras/GameplayCameras/Source/GameplayCameras/Private/Debug/PostProcessSettingsDebugBlock.cpp
+++ b/Engine/Plugins/Cameras/GameplayCameras/Source/GameplayCameras/Private/Debug/PostProcessSettingsDebugBlock.cpp
@@ -165,7 +165,7 @@ void FPostProcessSettingsDebugBlock::OnDebugDraw(const FCameraDebugBlockDrawPara
 		UE_DRAW_PP(FilmGrainShadowsMax);
 		UE_DRAW_PP(FilmGrainHighlightsMin);
 		UE_DRAW_PP(FilmGrainHighlightsMax);
-		UE_DRAW_PP(FilmGrainTexelSize);
+		UE_DRAW_PP(FilmGrainScale);
 		UE_DRAW_PP(AmbientOcclusionIntensity);
 		UE_DRAW_PP(AmbientOcclusionStaticFraction);
 		UE_DRAW_PP(AmbientOcclusionRadius);
@@ -181,9 +181,8 @@ void FPostProcessSettingsDebugBlock::OnDebugDraw(const FCameraDebugBlockDrawPara
 		UE_DRAW_PP(AmbientOcclusionTemporalBlendWeight);
 		UE_DRAW_PP(IndirectLightingColor);
 		UE_DRAW_PP(IndirectLightingIntensity);
-
+		UE_DRAW_PP(DepthOfFieldEnabled);
 		UE_DRAW_PP(DepthOfFieldFocalDistance);
-
 		UE_DRAW_PP(DepthOfFieldFstop);
 		UE_DRAW_PP(DepthOfFieldMinFstop);
 		UE_DRAW_PP(DepthOfFieldSensorWidth);
diff --git a/Engine/Plugins/Interchange/Runtime/Source/Import/Private/Animation/InterchangeLevelSequenceHelper.cpp b/Engine/Plugins/Interchange/Runtime/Source/Import/Private/Animation/InterchangeLevelSequenceHelper.cpp
--- a/Engine/Plugins/Interchange/Runtime/Source/Import/Private/Animation/InterchangeLevelSequenceHelper.cpp
+++ b/Engine/Plugins/Interchange/Runtime/Source/Import/Private/Animation/InterchangeLevelSequenceHelper.cpp
@@ -273,7 +273,7 @@ namespace UE::Interchange::Private
 			{EInterchangePropertyTracks::CameraPostProcessSettingsFilmGrainIntensityMidtones, {UMovieSceneFloatTrack::StaticClass()->GetName(),  GET_MEMBER_NAME_STRING_CHECKED(UCameraComponent, PostProcessSettings) + FString{ TEXT(".") } + GET_MEMBER_NAME_STRING_CHECKED(FPostProcessSettings, FilmGrainIntensityMidtones), TEXT("Film Grain Intensity Midtones (Post Process Settings)")}},
 			{EInterchangePropertyTracks::CameraPostProcessSettingsFilmGrainIntensityShadows, {UMovieSceneFloatTrack::StaticClass()->GetName(),  GET_MEMBER_NAME_STRING_CHECKED(UCameraComponent, PostProcessSettings) + FString{ TEXT(".") } + GET_MEMBER_NAME_STRING_CHECKED(FPostProcessSettings, FilmGrainIntensityShadows), TEXT("Film Grain Intensity Shadows (Post Process Settings)")}},
 			{EInterchangePropertyTracks::CameraPostProcessSettingsFilmGrainShadowsMax, {UMovieSceneFloatTrack::StaticClass()->GetName(),  GET_MEMBER_NAME_STRING_CHECKED(UCameraComponent, PostProcessSettings) + FString{ TEXT(".") } + GET_MEMBER_NAME_STRING_CHECKED(FPostProcessSettings, FilmGrainShadowsMax), TEXT("Film Grain Shadows Max (Post Process Settings)")}},
-			{EInterchangePropertyTracks::CameraPostProcessSettingsFilmGrainTexelSize, {UMovieSceneFloatTrack::StaticClass()->GetName(),  GET_MEMBER_NAME_STRING_CHECKED(UCameraComponent, PostProcessSettings) + FString{ TEXT(".") } + GET_MEMBER_NAME_STRING_CHECKED(FPostProcessSettings, FilmGrainTexelSize), TEXT("Film Grain Texel Size (Post Process Settings)")}},
+			{EInterchangePropertyTracks::CameraPostProcessSettingsFilmGrainScale, {UMovieSceneFloatVectorTrack::StaticClass()->GetName(),  GET_MEMBER_NAME_STRING_CHECKED(UCameraComponent, PostProcessSettings) + FString{ TEXT(".") } + GET_MEMBER_NAME_STRING_CHECKED(FPostProcessSettings, FilmGrainScale), TEXT("Film Grain Scale (Post Process Settings)"), 2}},
 			{EInterchangePropertyTracks::CameraPostProcessSettingsFilmShoulder, {UMovieSceneFloatTrack::StaticClass()->GetName(),  GET_MEMBER_NAME_STRING_CHECKED(UCameraComponent, PostProcessSettings) + FString{ TEXT(".") } + GET_MEMBER_NAME_STRING_CHECKED(FPostProcessSettings, FilmShoulder), TEXT("Shoulder (Post Process Settings)")}},
 			{EInterchangePropertyTracks::CameraPostProcessSettingsFilmSlope, {UMovieSceneFloatTrack::StaticClass()->GetName(),  GET_MEMBER_NAME_STRING_CHECKED(UCameraComponent, PostProcessSettings) + FString{ TEXT(".") } + GET_MEMBER_NAME_STRING_CHECKED(FPostProcessSettings, FilmSlope), TEXT("Slope (Post Process Settings)")}},
 			{EInterchangePropertyTracks::CameraPostProcessSettingsFilmToe, {UMovieSceneFloatTrack::StaticClass()->GetName(),  GET_MEMBER_NAME_STRING_CHECKED(UCameraComponent, PostProcessSettings) + FString{ TEXT(".") } + GET_MEMBER_NAME_STRING_CHECKED(FPostProcessSettings, FilmToe), TEXT("Toe (Post Process Settings)")}},
diff --git a/Engine/Plugins/Interchange/Runtime/Source/Nodes/Public/InterchangeAnimationDefinitions.h b/Engine/Plugins/Interchange/Runtime/Source/Nodes/Public/InterchangeAnimationDefinitions.h
--- a/Engine/Plugins/Interchange/Runtime/Source/Nodes/Public/InterchangeAnimationDefinitions.h
+++ b/Engine/Plugins/Interchange/Runtime/Source/Nodes/Public/InterchangeAnimationDefinitions.h
@@ -249,7 +249,7 @@ enum class EInterchangePropertyTracks : int32
 	CameraPostProcessSettingsFilmGrainIntensityMidtones,
 	CameraPostProcessSettingsFilmGrainIntensityShadows,
 	CameraPostProcessSettingsFilmGrainShadowsMax,
-	CameraPostProcessSettingsFilmGrainTexelSize,
+	CameraPostProcessSettingsFilmGrainScale,
 	CameraPostProcessSettingsFilmShoulder,
 	CameraPostProcessSettingsFilmSlope,
 	CameraPostProcessSettingsFilmToe,
diff --git a/Engine/Source/Runtime/Engine/Classes/Engine/Scene.h b/Engine/Source/Runtime/Engine/Classes/Engine/Scene.h
--- a/Engine/Source/Runtime/Engine/Classes/Engine/Scene.h
+++ b/Engine/Source/Runtime/Engine/Classes/Engine/Scene.h
@@ -1062,8 +1062,13 @@ struct FPostProcessSettings
 	UPROPERTY(EditAnywhere, BlueprintReadWrite, Category = Overrides, meta = (PinHiddenByDefault, InlineEditConditionToggle))
 	uint8 bOverride_FilmGrainHighlightsMax : 1;
 
+#if WITH_EDITORONLY_DATA
+	UPROPERTY()
+	uint8 bOverride_FilmGrainTexelSize_DEPRECATED:1;
+#endif
+
 	UPROPERTY(EditAnywhere, BlueprintReadWrite, Category=Overrides, meta=(PinHiddenByDefault, InlineEditConditionToggle))
-	uint8 bOverride_FilmGrainTexelSize:1;
+	uint8 bOverride_FilmGrainScale:1;
 
 	UPROPERTY(EditAnywhere, BlueprintReadWrite, Category = Overrides, meta = (PinHiddenByDefault, InlineEditConditionToggle))
 	uint8 bOverride_FilmGrainTexture : 1;
@@ -1179,6 +1184,9 @@ struct FPostProcessSettings
 	UPROPERTY(EditAnywhere, BlueprintReadWrite, Category=Overrides, meta=(PinHiddenByDefault, InlineEditConditionToggle))
 	uint8 bOverride_ColorGradingLUT:1;
 
+	UPROPERTY(EditAnywhere, BlueprintReadWrite, Category=Overrides, meta=(PinHiddenByDefault, InlineEditConditionToggle))
+	uint8 bOverride_DepthOfFieldEnabled:1;
+
 	UPROPERTY(EditAnywhere, BlueprintReadWrite, Category=Overrides, meta=(PinHiddenByDefault, InlineEditConditionToggle))
 	uint8 bOverride_DepthOfFieldFocalDistance:1;
 
@@ -2187,9 +2195,16 @@ struct FPostProcessSettings
 	UPROPERTY(interp, BlueprintReadWrite, Category = "Film Grain", meta = (UIMin = "1.0", UIMax = "10.0", editcondition = "bOverride_FilmGrainHighlightsMax"))
 	float FilmGrainHighlightsMax;
 
+#if WITH_EDITORONLY_DATA
 	/** Controls the size of the film grain. Size of texel of FilmGrainTexture on screen. */
-	UPROPERTY(interp, BlueprintReadWrite, Category = "Film Grain", meta = (UIMin = "0.0", UIMax = "4.0", editcondition = "bOverride_FilmGrainTexelSize"))
-	float FilmGrainTexelSize;
+	UE_DEPRECATED(5.8, "Use FilmGrainScale instead")
+	UPROPERTY()
+	float FilmGrainTexelSize_DEPRECATED;
+#endif
+
+	/** Controls the size of the film grain. Size of texel of FilmGrainTexture on screen. */
+	UPROPERTY(interp, BlueprintReadWrite, Category = "Film Grain", meta = (UIMin = "0.0", UIMax = "4.0", editcondition = "bOverride_FilmGrainScale"))
+	FVector2f FilmGrainScale;
 
 	/** Defines film grain texture to use. */
 	UPROPERTY(EditAnywhere, BlueprintReadWrite, Category = "Film Grain", meta = (editcondition = "bOverride_FilmGrainTexture"))
@@ -2275,6 +2290,10 @@ struct FPostProcessSettings
 	UPROPERTY(EditAnywhere, BlueprintReadWrite, Category="Color Grading|Misc", meta=(editcondition = "bOverride_ColorGradingLUT", DisplayName = "Color Grading LUT"))
 	TObjectPtr<class UTexture> ColorGradingLUT;
 
+	/** When false, post-process depth of field will not be used. When true, DoF may be used depending on lens and scalability settings. */
+	UPROPERTY(EditAnywhere, BlueprintReadWrite, Category="Lens|Depth of Field", meta=(editcondition = "bOverride_DepthOfFieldEnabled", DisplayName = "Enabled"))
+	bool DepthOfFieldEnabled;
+
 	/** Width of the camera sensor to assume, in mm. */
 	UPROPERTY(EditAnywhere, BlueprintReadWrite, Category="Lens|Depth of Field", meta=(ForceUnits=mm, ClampMin = "0.1", UIMin="0.1", UIMax= "1000.0", editcondition = "bOverride_DepthOfFieldSensorWidth", DisplayName = "Sensor Width (mm)"))
 	float DepthOfFieldSensorWidth;
@@ -2561,6 +2580,7 @@ struct FPostProcessSettings
 	// for backwards compatibility
 	void OnAfterLoad()
 	{
+PRAGMA_DISABLE_DEPRECATION_WARNINGS
 		for(int32 i = 0, count = Blendables_DEPRECATED.Num(); i < count; ++i)
 		{
 			if(Blendables_DEPRECATED[i])
@@ -2594,6 +2614,13 @@ struct FPostProcessSettings
 			LocalExposureHighlightContrastScale = LocalExposureContrastScale_DEPRECATED;
 			LocalExposureShadowContrastScale = LocalExposureContrastScale_DEPRECATED;
 		}
+
+		if (bOverride_FilmGrainTexelSize_DEPRECATED)
+		{
+			bOverride_FilmGrainScale = true;
+			FilmGrainScale = FVector2f(FilmGrainTexelSize_DEPRECATED, FilmGrainTexelSize_DEPRECATED);
+		}
+PRAGMA_ENABLE_DEPRECATION_WARNINGS
 	}
 #endif
 
diff --git a/Engine/Source/Runtime/Engine/Private/PostProcessUtils.cpp b/Engine/Source/Runtime/Engine/Private/PostProcessUtils.cpp
--- a/Engine/Source/Runtime/Engine/Private/PostProcessUtils.cpp
+++ b/Engine/Source/Runtime/Engine/Private/PostProcessUtils.cpp
@@ -132,7 +132,7 @@ bool FPostProcessUtils::OverridePostProcessSettings(FPostProcessSettings& ThisFr
 		UE_SET_PP(FilmGrainShadowsMax);
 		UE_SET_PP(FilmGrainHighlightsMin);
 		UE_SET_PP(FilmGrainHighlightsMax);
-		UE_SET_PP(FilmGrainTexelSize);
+		UE_SET_PP(FilmGrainScale);
 		UE_SET_PP(AmbientOcclusionIntensity);
 		UE_SET_PP(AmbientOcclusionStaticFraction);
 		UE_SET_PP(AmbientOcclusionRadius);
@@ -153,6 +153,7 @@ bool FPostProcessUtils::OverridePostProcessSettings(FPostProcessSettings& ThisFr
 
 		UE_SET_PP(DepthOfFieldFstop);
 		UE_SET_PP(DepthOfFieldMinFstop);
+		UE_SET_PP(DepthOfFieldEnabled);
 		UE_SET_PP(DepthOfFieldSensorWidth);
 		UE_SET_PP(DepthOfFieldSqueezeFactor);
 		UE_SET_PP(DepthOfFieldDepthBlurRadius);
@@ -448,7 +449,7 @@ bool FPostProcessUtils::BlendPostProcessSettings(FPostProcessSettings& ThisFrom,
 		UE_LERP_PP(FilmGrainShadowsMax);
 		UE_LERP_PP(FilmGrainHighlightsMin);
 		UE_LERP_PP(FilmGrainHighlightsMax);
-		UE_LERP_PP(FilmGrainTexelSize);
+		UE_LERP_PP(FilmGrainScale);
 		UE_LERP_PP(AmbientOcclusionIntensity);
 		UE_LERP_PP(AmbientOcclusionStaticFraction);
 		UE_LERP_PP(AmbientOcclusionRadius);
@@ -477,6 +478,7 @@ bool FPostProcessUtils::BlendPostProcessSettings(FPostProcessSettings& ThisFrom,
 			}
 			bAnyOverwritten = true;
 		}
+		UE_SET_PP(DepthOfFieldEnabled)
 		UE_LERP_PP(DepthOfFieldFstop);
 		UE_LERP_PP(DepthOfFieldMinFstop);
 		UE_LERP_PP(DepthOfFieldSensorWidth);
diff --git a/Engine/Source/Runtime/Engine/Private/Scene.cpp b/Engine/Source/Runtime/Engine/Private/Scene.cpp
--- a/Engine/Source/Runtime/Engine/Private/Scene.cpp
+++ b/Engine/Source/Runtime/Engine/Private/Scene.cpp
@@ -382,7 +382,7 @@ static void DoPostProcessSettingsSanityCheck()
 		RenameMap.Add(TEXT("FilmGrainShadowsMax"), TEXT(""));
 		RenameMap.Add(TEXT("FilmGrainHighlightsMin"), TEXT(""));
 		RenameMap.Add(TEXT("FilmGrainHighlightsMax"), TEXT(""));
-		RenameMap.Add(TEXT("FilmGrainTexelSize"), TEXT(""));
+		RenameMap.Add(TEXT("FilmGrainScale"), TEXT(""));
 		RenameMap.Add(TEXT("FilmGrainTexture"), TEXT(""));
 		VerifyPostProcessingProperties(TEXT("Film"),
 			TArray<const UStruct*>({FFilmStockSettings::StaticStruct()}),
@@ -556,7 +556,12 @@ FPostProcessSettings::FPostProcessSettings()
 	FilmGrainShadowsMax = 0.09f;
 	FilmGrainHighlightsMin = 0.5f;
 	FilmGrainHighlightsMax = 1.0f;
-	FilmGrainTexelSize = 1.0f;
+#if WITH_EDITORONLY_DATA
+PRAGMA_DISABLE_DEPRECATION_WARNINGS
+	FilmGrainTexelSize_DEPRECATED = 1.0f;
+PRAGMA_ENABLE_DEPRECATION_WARNINGS
+#endif
+	FilmGrainScale = FVector2f{1.0f, 1.0f};
 
 	// next value might get overwritten by r.DefaultFeature.AmbientOcclusion
 	AmbientOcclusionIntensity = .5f;
@@ -602,6 +607,7 @@ FPostProcessSettings::FPostProcessSettings()
 	DepthOfFieldFstop = 4.0f; 
 	DepthOfFieldMinFstop = 1.2f;
 	DepthOfFieldBladeCount = FPostProcessSettings::kDefaultDepthOfFieldBladeCount;
+	DepthOfFieldEnabled = true;
 	DepthOfFieldSensorWidth = 24.576f;			// APS-C
 	DepthOfFieldSqueezeFactor = 1.0f;
 	DepthOfFieldDepthBlurAmount = 1.0f;
diff --git a/Engine/Source/Runtime/Engine/Private/SceneView.cpp b/Engine/Source/Runtime/Engine/Private/SceneView.cpp
--- a/Engine/Source/Runtime/Engine/Private/SceneView.cpp
+++ b/Engine/Source/Runtime/Engine/Private/SceneView.cpp
@@ -1701,7 +1701,7 @@ void FSceneView::OverridePostProcessSettings(const FPostProcessSettings& Src, fl
 		LERP_PP(FilmGrainShadowsMax);
 		LERP_PP(FilmGrainHighlightsMin);
 		LERP_PP(FilmGrainHighlightsMax);
-		LERP_PP(FilmGrainTexelSize);
+		LERP_PP(FilmGrainScale);
 		LERP_PP(AmbientOcclusionIntensity);
 		LERP_PP(AmbientOcclusionStaticFraction);
 		LERP_PP(AmbientOcclusionRadius);
@@ -1731,6 +1731,7 @@ void FSceneView::OverridePostProcessSettings(const FPostProcessSettings& Src, fl
 		}
 		LERP_PP(DepthOfFieldFstop);
 		LERP_PP(DepthOfFieldMinFstop);
+		SET_PP(DepthOfFieldEnabled)
 		LERP_PP(DepthOfFieldSensorWidth);
 		LERP_PP(DepthOfFieldSqueezeFactor);
 		LERP_PP(DepthOfFieldDepthBlurRadius);
diff --git a/Engine/Source/Runtime/Renderer/Private/PostProcess/DiaphragmDOF.cpp b/Engine/Source/Runtime/Renderer/Private/PostProcess/DiaphragmDOF.cpp
--- a/Engine/Source/Runtime/Renderer/Private/PostProcess/DiaphragmDOF.cpp
+++ b/Engine/Source/Runtime/Renderer/Private/PostProcess/DiaphragmDOF.cpp
@@ -1502,6 +1502,7 @@ bool DiaphragmDOF::IsEnabled(const FViewInfo& View)
 		View.Family->EngineShowFlags.DepthOfField &&
 		bDepthOfFieldRequestedByCVar &&
 		!(View.Family->EngineShowFlags.PathTracing && PathTracing::UsesReferenceDOF(View)) &&
+		View.FinalPostProcessSettings.DepthOfFieldEnabled &&
 		(View.FinalPostProcessSettings.DepthOfFieldFocalDistance > 0.f || View.FinalPostProcessSettings.DepthOfFieldDepthBlurRadius > 0.f);
 }
 
diff --git a/Engine/Source/Runtime/Renderer/Private/PostProcess/DiaphragmDOFUtils.cpp b/Engine/Source/Runtime/Renderer/Private/PostProcess/DiaphragmDOFUtils.cpp
--- a/Engine/Source/Runtime/Renderer/Private/PostProcess/DiaphragmDOFUtils.cpp
+++ b/Engine/Source/Runtime/Renderer/Private/PostProcess/DiaphragmDOFUtils.cpp
@@ -49,7 +49,8 @@ float ComputeFocalLengthFromFov(const FSceneView& View)
 FVector4f DiaphragmDOF::CircleDofHalfCoc(const FViewInfo& View)
 {
 	static const auto CVar = IConsoleManager::Get().FindTConsoleVariableDataInt(TEXT("r.DepthOfFieldQuality"));
-	bool bDepthOfField = View.Family->EngineShowFlags.DepthOfField && CVar->GetValueOnRenderThread() > 0 && View.FinalPostProcessSettings.DepthOfFieldFocalDistance > 0;
+	bool bDepthOfField = View.Family->EngineShowFlags.DepthOfField && CVar->GetValueOnRenderThread() > 0 
+		&& View.FinalPostProcessSettings.DepthOfFieldEnabled && View.FinalPostProcessSettings.DepthOfFieldFocalDistance > 0;
 
 	FVector4f Ret(0, 1, 0, 0);
 
diff --git a/Engine/Source/Runtime/Renderer/Private/PostProcess/PostProcessTonemap.cpp b/Engine/Source/Runtime/Renderer/Private/PostProcess/PostProcessTonemap.cpp
--- a/Engine/Source/Runtime/Renderer/Private/PostProcess/PostProcessTonemap.cpp
+++ b/Engine/Source/Runtime/Renderer/Private/PostProcess/PostProcessTonemap.cpp
@@ -797,9 +797,14 @@ FScreenPassTexture AddTonemapPass(FRDGBuilder& GraphBuilder, const FViewInfo& Vi
 		FVector2f TextureSize(View.FilmGrainTexture->GetSizeX(), View.FilmGrainTexture->GetSizeY());
 		FVector2f OutputSizeF(1920.0f, 1920.0f * float(OutputViewport.Rect.Height()) / float(OutputViewport.Rect.Width()));
 
+		FVector2f FilmGrainScale
+		{
+			FMath::Max(View.FinalPostProcessSettings.FilmGrainScale.X, 0.001f),
+			FMath::Max(View.FinalPostProcessSettings.FilmGrainScale.Y, 0.001f)
+		};
 		CommonParameters.FilmGrain.ScreenPosToFilmGrainTextureUV = 
 			FScreenTransform::ScreenPosToViewportUV *
-			((OutputSizeF / TextureSize) * (1.0f / View.FinalPostProcessSettings.FilmGrainTexelSize)) + RandomGrainTextureUVOffset;
+			((OutputSizeF / TextureSize) * (FVector2f(1.0f, 1.0f) / FilmGrainScale)) + RandomGrainTextureUVOffset;
 
 		CommonParameters.FilmGrain.FilmGrainTextureConstants = GraphBuilder.CreateSRV(FilmGrainConstantsBuffer);
 	}
diff --git a/Engine/Source/Runtime/Renderer/Private/PostProcess/PostProcessing.cpp b/Engine/Source/Runtime/Renderer/Private/PostProcess/PostProcessing.cpp
--- a/Engine/Source/Runtime/Renderer/Private/PostProcess/PostProcessing.cpp
+++ b/Engine/Source/Runtime/Renderer/Private/PostProcess/PostProcessing.cpp
@@ -2721,7 +2721,7 @@ void AddMobilePostProcessingPasses(FRDGBuilder& GraphBuilder, FScene* Scene, con
 	// add the passes we want to add to the graph (commenting a line means the pass is not inserted into the graph) ---------
 
 	// HQ gaussian 
-	bool bUseDof = GetMobileDepthOfFieldScale(View) > 0.0f && View.Family->EngineShowFlags.DepthOfField && !View.Family->EngineShowFlags.VisualizeDOF;
+	bool bUseDof = GetMobileDepthOfFieldScale(View) > 0.0f && View.Family->EngineShowFlags.DepthOfField && !View.Family->EngineShowFlags.VisualizeDOF && View.FinalPostProcessSettings.DepthOfFieldEnabled;
 	bool bUseMobileDof = bUseDof && !View.FinalPostProcessSettings.bMobileHQGaussian;
 
 	// Do not use the tonemapper if output texture is sRGB cause the conversion will be performed by HW.
