[shaders] include a hash of the base material path when calculating shadermap DDC keys. when editoronly data is enabled this path gets reflected into some debug data, and this is a source of non-determinism in the shadermap DDC records.

#Synced-CL 51376167

diff --git a/Engine/Source/Runtime/Engine/Private/Materials/MaterialShader.cpp b/Engine/Source/Runtime/Engine/Private/Materials/MaterialShader.cpp
--- a/Engine/Source/Runtime/Engine/Private/Materials/MaterialShader.cpp
+++ b/Engine/Source/Runtime/Engine/Private/Materials/MaterialShader.cpp
@@ -1287,12 +1287,14 @@ void FMaterialShaderMapId::Serialize(FArchive& Ar, bool bLoadingCooked)
 			Ar << bUsingNewHLSLGenerator;
 			Ar << PreferredShaderCodeFormat;
 			Ar << ExternalCodeReferencesHash;
+			Ar << BaseMaterialPathHash;
 		}
 		else
 		{
 			bUsingNewHLSLGenerator = false;
 			PreferredShaderCodeFormat = NAME_None;
 			ExternalCodeReferencesHash = FSHAHash();
+			BaseMaterialPathHash = FSHAHash();
 		}
 
 		// SUBSTRATE_TODO We do not need to serialize FSubstrateCompilationConfig for now since this is only used when debugging in the editor.
@@ -1561,7 +1563,12 @@ bool FMaterialShaderMapId::Equals(const FMaterialShaderMapId& ReferenceSet, bool
 			return false;
 		}
 
-		if( BasePropertyOverridesHash != ReferenceSet.BasePropertyOverridesHash )
+		if (BasePropertyOverridesHash != ReferenceSet.BasePropertyOverridesHash)
+		{
+			return false;
+		}
+
+		if (BaseMaterialPathHash != ReferenceSet.BaseMaterialPathHash)
 		{
 			return false;
 		}
@@ -1697,6 +1704,7 @@ void FMaterialShaderMapId::RecordAndEmit(FMaterialKeyGeneratorContext& Context)
 	}
 
 	Context.RecordAndEmit("LayoutParams", LayoutParams);
+	Context.RecordAndEmit("BaseMaterialPathHash", BaseMaterialPathHash);
 
 	// Recording of some fields necessary for IncludeUObjectData is not yet implemented; we only support emitting them
 	check(!Context.HasAllFlags(EMaterialKeyInclude::UObjectData) || !Context.IsRecording());
diff --git a/Engine/Source/Runtime/Engine/Private/Materials/MaterialShared.cpp b/Engine/Source/Runtime/Engine/Private/Materials/MaterialShared.cpp
--- a/Engine/Source/Runtime/Engine/Private/Materials/MaterialShared.cpp
+++ b/Engine/Source/Runtime/Engine/Private/Materials/MaterialShared.cpp
@@ -775,6 +775,7 @@ void FMaterial::BuildShaderMapIdOverride(const FBuildShaderMapIdArgs& Args) cons
 			GetDependentShaderAndVFTypes(OutId.LayoutParams, ShaderTypes, ShaderPipelineTypes, VFTypes);
 		}
 
+		OutId.BaseMaterialPathHash = FSHA1::HashBuffer(Args.BaseMaterialPath.GetCharArray().GetData(), Args.BaseMaterialPath.NumBytesWithoutNull());
 		OutId.Usage = GetShaderMapUsage();
 		OutId.bUsingNewHLSLGenerator = IsUsingNewHLSLGenerator();
 		OutId.PreferredShaderCodeFormat = FDataDrivenShaderPlatformInfo::GetPreferredShaderCodeFormat(Args.Platform);
diff --git a/Engine/Source/Runtime/Engine/Public/MaterialShared.h b/Engine/Source/Runtime/Engine/Public/MaterialShared.h
--- a/Engine/Source/Runtime/Engine/Public/MaterialShared.h
+++ b/Engine/Source/Runtime/Engine/Public/MaterialShared.h
@@ -1318,9 +1318,12 @@ public:
 	/** A hash of the base property overrides for this material instance. */
 	FSHAHash BasePropertyOverridesHash;
 
-	/** A hash of all external code delcarations referenced in the shader map. */
+	/** A hash of all external code declarations referenced in the shader map. */
 	FSHAHash ExternalCodeReferencesHash;
 
+	/** A hash of the base material's full path */
+	FSHAHash BaseMaterialPathHash;
+
 	/** Is the material using the new HLSL generator? */
 	bool bUsingNewHLSLGenerator;
 
@@ -2317,6 +2320,7 @@ public:
 		EShaderPlatform Platform = EShaderPlatform::SP_NumPlatforms;
 		const ITargetPlatform* TargetPlatform = nullptr;
 		EMaterialKeyInclude IncludeFlags = EMaterialKeyInclude::All;
+		FString BaseMaterialPath;
 	};
 
 	UE_DEPRECATED(5.7, "Call BuildShaderMapId instead")
@@ -2338,7 +2342,7 @@ public:
 		EMaterialKeyInclude IncludeFlags = EMaterialKeyInclude::All) const
 	{
 		BuildShaderMapIdOverride({
-			.OutId = &OutId, .Platform = GetShaderPlatform(), .TargetPlatform = TargetPlatform, .IncludeFlags = IncludeFlags }
+			.OutId = &OutId, .Platform = GetShaderPlatform(), .TargetPlatform = TargetPlatform, .IncludeFlags = IncludeFlags, .BaseMaterialPath = GetBaseMaterialPathName()}
 			);
 	}
 
@@ -2911,8 +2915,9 @@ protected:
 	/** Used to translate code for custom output attributes such as ClearCoatBottomNormal  */
 	virtual int32 CompileCustomAttribute(const FGuid& AttributeID, class FMaterialCompiler* Compiler) const {return INDEX_NONE;}
 
-	/** Useful for debugging. */
+	/** Retrieves the full path of the base UMaterial associated with this FMaterial */
 	virtual FString GetBaseMaterialPathName() const { return TEXT(""); }
+	/** Get a debug name useful for debugging, just the base material path by default */
 	virtual FString GetDebugName() const { return GetBaseMaterialPathName(); }
 
 	UE_DEPRECATED(5.7, "Please use SetQualityLevelProperties with EShaderPlatform argument and not ERHIFeatureLevel::Type")
