RemoteObjects:  The PlayerController's PlayerCameraManager can be remote and accessed outside of a transaction.  Prevent that from causing a crash by early out.

#Synced-CL 51377531

diff --git a/Engine/Source/Runtime/Engine/Private/PlayerController.cpp b/Engine/Source/Runtime/Engine/Private/PlayerController.cpp
--- a/Engine/Source/Runtime/Engine/Private/PlayerController.cpp
+++ b/Engine/Source/Runtime/Engine/Private/PlayerController.cpp
@@ -6484,7 +6484,18 @@ const FString& APlayerController::GetCurrentInputModeDebugString() const
 
 void APlayerController::UpdateCameraManager(float DeltaSeconds)
 {
-	if (PlayerCameraManager != NULL)
+#if UE_SERVER && UE_AUTORTFM && UE_WITH_REMOTE_OBJECT_HANDLE
+	if (!AutoRTFM::IsClosed())
+	{
+		if (PlayerCameraManager.GetResidence() != EResidence::Local)
+		{
+			UE_LOGF(LogPlayerController, VeryVerbose, "%s: Skipping due to PlayerCameraManager being non-local", __FUNCTION__);
+			return;
+		}
+	}
+#endif
+
+	if (PlayerCameraManager)
 	{
 		PlayerCameraManager->UpdateCamera(DeltaSeconds);
 	}
