RemoteObjects:  Fix for bIsReadyForReplication always being set to false due to migration.
#Synced-CL 51377531

diff --git a/Engine/Source/Runtime/Engine/Private/Components/ActorComponent.cpp b/Engine/Source/Runtime/Engine/Private/Components/ActorComponent.cpp
--- a/Engine/Source/Runtime/Engine/Private/Components/ActorComponent.cpp
+++ b/Engine/Source/Runtime/Engine/Private/Components/ActorComponent.cpp
@@ -752,8 +752,15 @@ void UActorComponent::PostMigrate(const struct FUObjectMigrationContext& Migrati
 
 	if (MigrationContext.MigrationSide == EObjectMigrationSide::Send)
 	{
+		const EResidence Residence = UE::RemoteObject::Handle::GetResidence(this);
 		const EObjectMigrationSendType SendType = MigrationContext.GetObjectMigrationSendType(this);
-		UE_LOG(LogActorComponent, VeryVerbose, TEXT("%hs: %s %s"), __func__, ToString(SendType), *GetName());
+		UE_LOG(LogActorComponent, VeryVerbose, TEXT("%hs: %s %s (Residence: %s)"), __func__, ToString(SendType), *GetName(), EnumToString(Residence));
+
+		if (Residence != EResidence::Local)
+		{
+			ensureMsgf(Residence == EResidence::LocalNotReady, TEXT("Trying to send a non-local object"));
+			return;
+		}
 
 		// Remove the actor and it's sub-objects from the network id cache since these objects are now considered remote. It's assumed that the
 		// network id cache will borrow any objects that are not in the cache but are considered remote.
@@ -813,7 +820,7 @@ void UActorComponent::PostMigrate(const struct FUObjectMigrationContext& Migrati
 
 		// Add the actor and it's sub-objects to the network id cache since these objects are no longer considered remote and can be derferenced
 		// through the network id cache (e.g. by RPC argument and object deserialization).
-		if (HasBegunPlay() && IsSupportedForNetworking())
+		if (HasBegunPlay() && IsSupportedForNetworking() && (RecvType != EObjectMigrationRecvType::ReturnedLoan))
 		{
 			GetWorld()->GetNetDriver()->AddObjectToNetworkIdCache(this);
 		}
@@ -822,7 +829,14 @@ void UActorComponent::PostMigrate(const struct FUObjectMigrationContext& Migrati
 		// since the replicated component/subobject lists will be reset in that function.
 		// Even though we only replicate when ownership is assigned, ResetOwnedComponents is called on every receive,
 		// so we also have to rebuild replicated subobjects on every receive.
-		bIsReadyForReplication = false;
+		if (AActor* Owner = GetOwner())
+		{
+			// Has ResetOwnedComponents been called yet?
+			if (!Owner->IsReplicatedActorComponentRegistered(this))
+			{
+				bIsReadyForReplication = false;
+			}
+		}
 
 		// We are expecting this flow from AActor, but now we must do this per-component in such a way
 		// that we are not mutating our previously deserialized data.
