only set EPSOPrecacheInMemoryMode::Cached when KeepPrecachedPSOsInMemoryForActiveMaterials is enabled for global psos otherwise they will be cached independent if KeepPrecachedPSOsInMemoryForActiveMaterials is enabled or not

diff --git a/Engine/Source/Runtime/Engine/Private/ShaderCompiler/ShaderCompiler.cpp b/Engine/Source/Runtime/Engine/Private/ShaderCompiler/ShaderCompiler.cpp
--- a/Engine/Source/Runtime/Engine/Private/ShaderCompiler/ShaderCompiler.cpp
+++ b/Engine/Source/Runtime/Engine/Private/ShaderCompiler/ShaderCompiler.cpp
@@ -5191,7 +5191,7 @@ void PrecacheComputePipelineStatesForGlobalShaders(ERHIFeatureLevel::Type Featur
 					FPSOPrecacheData PSOPrecacheData;
 					PSOPrecacheData.Type = FPSOPrecacheData::EType::Compute;
 					PSOPrecacheData.ComputeInitializer = ComputeInitializer;
-					PSOPrecacheData.InMemoryMode = PSOPrecacheSettings.bKeepGlobalComputePSOsInMemory ? EPSOPrecacheInMemoryMode::Cached : EPSOPrecacheInMemoryMode::Uncached;
+					PSOPrecacheData.InMemoryMode = (PSOPrecacheSettings.bKeepGlobalComputePSOsInMemory && KeepPrecachedPSOsInMemoryForActiveMaterials()) ? EPSOPrecacheInMemoryMode::Cached : EPSOPrecacheInMemoryMode::Uncached;
 					#if PSO_PRECACHING_VALIDATE
 					PSOPrecacheData.ResourceName = GlobalShader.GetType()->GetName();
 					PSOPrecacheData.PSOCollectorIndex = GlobalComputePSOCollectorIndex;
@@ -5272,7 +5272,7 @@ void PrecacheComputePipelineStatesForGlobalShaders(ERHIFeatureLevel::Type Featur
 	// Set if the global graphics PSOs need to be kept in memory or not
 	for (FPSOPrecacheData& PrecacheData : GlobalPSOInitializers)
 	{
-		PrecacheData.InMemoryMode = PSOPrecacheSettings.bKeepGlobalGraphicPSOsInMemory ? EPSOPrecacheInMemoryMode::Cached : EPSOPrecacheInMemoryMode::Uncached;
+		PrecacheData.InMemoryMode = (PSOPrecacheSettings.bKeepGlobalGraphicPSOsInMemory && KeepPrecachedPSOsInMemoryForActiveMaterials()) ? EPSOPrecacheInMemoryMode::Cached : EPSOPrecacheInMemoryMode::Uncached;
 	}
 
 	RequestPrecachePSOs(EPSOPrecacheType::Global, GlobalPSOInitializers);
