[Materials] Improve runtime performance of `UMaterialFunctionInterface::UpdateParameterSet()`.

Remove `bChanged` state and return true on the first match as the contract of this function claims.

#rnx

#Synced-CL 51330672

diff --git a/Engine/Source/Runtime/Engine/Public/Materials/MaterialFunctionInterface.h b/Engine/Source/Runtime/Engine/Public/Materials/MaterialFunctionInterface.h
--- a/Engine/Source/Runtime/Engine/Public/Materials/MaterialFunctionInterface.h
+++ b/Engine/Source/Runtime/Engine/Public/Materials/MaterialFunctionInterface.h
@@ -246,8 +246,6 @@ public:
 	template <typename ParameterType, typename ExpressionType>
 	bool UpdateParameterSet(ParameterType& Parameter)
 	{
-		bool bChanged = false;
-
 		if (UMaterialFunctionInterface* ParameterFunction = GetBaseFunctionInterface())
 		{
 			TArray<UMaterialFunctionInterface*> Functions;
@@ -263,15 +261,13 @@ public:
 						if (ParameterExpression->ParameterName == Parameter.ParameterInfo.Name)
 						{
 							Parameter.ExpressionGUID = ParameterExpression->ExpressionGUID;
-							bChanged = true;
-							break;
+							return true;
 						}
 					}
 				}
 			}
 		}
-
-		return bChanged;
+		return false;
 	}
 #endif // WITH_EDITOR
 
