[Input] Add some debug data to player input available in non-shipping and test builds behind the UE_PLAYER_INPUT_INCLUDE_DEBUG define.

- Add a "GetUObject" function to TInputUnifiedDelegate which makes the interface consistent with being able to figure out what object the input delegate is bound to
- Add PlayerInputDebug events for when a delegate is fired and when player input is flushed.

[RN] Add some debug data to player input available in non-shipping and test builds behind the UE_PLAYER_INPUT_INCLUDE_DEBUG define.

#Synced-CL 51092797

diff --git a/Engine/Plugins/EnhancedInput/Source/EnhancedInput/Private/EnhancedPlayerInput.cpp b/Engine/Plugins/EnhancedInput/Source/EnhancedInput/Private/EnhancedPlayerInput.cpp
--- a/Engine/Plugins/EnhancedInput/Source/EnhancedInput/Private/EnhancedPlayerInput.cpp
+++ b/Engine/Plugins/EnhancedInput/Source/EnhancedInput/Private/EnhancedPlayerInput.cpp
@@ -9,6 +9,7 @@
 #include "EnhancedInputModule.h"
 #include "EnhancedInputSubsystemInterface.h"
 #include "GameFramework/PlayerController.h"
+#include "GameFramework/PlayerInputDebugging.h"
 #include "GameFramework/WorldSettings.h"
 #include "HAL/IConsoleManager.h"
 #include "InputMappingContext.h"
@@ -826,7 +827,18 @@ bool UEnhancedPlayerInput::EvaluateInputComponentDelegates(UInputComponent* cons
 					else
 					{
 						Delegate->Execute(*ActionData);
-					}					
+					}
+
+#if UE_PLAYER_INPUT_INCLUDE_DEBUG
+					UE::Input::FPlayerInputDebugging::BroadcastPlayerInputDelegateExecuted(
+						this,
+						{
+							.InputComponent = InputComponent,
+							.ListeningObject = Delegate->GetUObject(),
+							.ActionName = GetPathNameSafe(ActionData->GetSourceAction()),
+							.InputValue = ActionData->GetValue().Get<FVector>(),
+						});
+#endif
 				}
 
 				// Keep track of any delegates which want to consume input away from other lower priority listeners
diff --git a/Engine/Source/Runtime/Engine/Classes/Components/InputComponent.h b/Engine/Source/Runtime/Engine/Classes/Components/InputComponent.h
--- a/Engine/Source/Runtime/Engine/Classes/Components/InputComponent.h
+++ b/Engine/Source/Runtime/Engine/Classes/Components/InputComponent.h
@@ -25,6 +25,27 @@ struct TInputUnifiedDelegate
 		return ((FuncDelegate.IsValid() && FuncDelegate->IsBound()) || (FuncDynDelegate.IsValid() && FuncDynDelegate->IsBound()));
 	}
 
+	/** Returns the UObject that this delegate is bound to. */
+	inline UObject* GetUObject() const
+	{
+		if (FuncDelegate.IsValid())
+		{
+			if (FuncDelegate->IsBound())
+			{
+				return FuncDelegate->GetUObject();
+			}
+		}
+		else if (FuncDynDelegate.IsValid())
+		{
+			if (FuncDynDelegate->IsBound())
+			{
+				return FuncDynDelegate->GetUObject();
+			}
+		}
+		
+		return nullptr;
+	}
+	
 	/** Returns if either the native or dynamic delegate is bound to an object */
 	inline bool IsBoundToObject(void const* Object) const
 	{
diff --git a/Engine/Source/Runtime/Engine/Classes/GameFramework/PlayerInputDebugging.h b/Engine/Source/Runtime/Engine/Classes/GameFramework/PlayerInputDebugging.h
new file mode 100644
--- /dev/null
+++ b/Engine/Source/Runtime/Engine/Classes/GameFramework/PlayerInputDebugging.h
@@ -0,0 +1,95 @@
+// Copyright Epic Games, Inc. All Rights Reserved.
+
+#pragma once
+
+#include "CoreMinimal.h"
+#include "Delegates/Delegate.h"
+#include "UObject/WeakObjectPtr.h"
+#include "UObject/WeakObjectPtrTemplates.h"
+
+#define UE_API ENGINE_API
+
+/**
+ * The UE_PLAYER_INPUT_INCLUDE_DEBUG macro will allow the engine to broadcast a debug delegate whenever
+ *
+ * a Player input event is executed. This is useful
+ * 
+ * If you do not wish to have this enabled for your build target, you can add this to your *.Target.cs file:
+ *		GlobalDefinitions.Add("UE_PLAYER_INPUT_INCLUDE_DEBUG=0");
+ */
+#ifndef UE_PLAYER_INPUT_INCLUDE_DEBUG
+	#define UE_PLAYER_INPUT_INCLUDE_DEBUG	!UE_BUILD_SHIPPING || UE_BUILD_TEST
+#endif
+
+#if UE_PLAYER_INPUT_INCLUDE_DEBUG
+
+class UInputComponent;
+class UPlayerInput;
+class UObject;
+
+namespace UE::Input
+{
+	
+/**
+ * Payload data for when a player input event has occured.
+ */
+struct FPlayerInputDebuggingArgs
+{
+	/**
+	 * The input component which has a delegate bound to this
+	 */
+	TWeakObjectPtr<const UInputComponent> InputComponent;
+
+	/**
+	 * The object which is receiving this input delegate
+	 */
+	TWeakObjectPtr<const UObject> ListeningObject;
+
+	/**
+	 * The name of the input delegate which is being fired.
+	 * For legacy action/axis mappings, this will simply be their FName found in
+	 * the project settings. For Enhanced Input actions, this will be the Object Path of the UInputAction.
+	 */
+	FString ActionName = TEXT("");
+	
+	/**
+	 * The value of the input delegate which has been fired. 
+	 *
+	 * - For legacy IE_Pressed events, the X value will be 1.
+	 * - For legacy IE_Released events, the Y value will be 1.
+	 * - For legacy "axis" events, this value will be the raw accumulated value of that axis.
+	 * Keep in mind that legacy axis events fire every tick so long as they are bound, so
+	 * applying a deadzone may be desirable for any debug tooling here.
+	 * - For Enhanced input events the evaluated input action instance will be used.
+	 */
+	FVector InputValue = FVector::ZeroVector;
+};
+
+/**
+ * A class which contains several player input debugging related delegates which may be useful
+ * for users building debugging tools for input.
+ *
+ * This is especially useful when combined with the SlateDebugging.h, because you can receive notifications
+ * of when a "raw" slate input event occurs, as well as when a player event occurs. This makes it easier to
+ * track down issues with input consumption and create event driven visualizing tools.
+ */
+class FPlayerInputDebugging
+{
+public:
+
+	/** Event fired when a player has executed an input event. */
+	DECLARE_MULTICAST_DELEGATE_TwoParams(FOnPlayerInputDelegateExecuted, const UPlayerInput*, const FPlayerInputDebuggingArgs&);
+	static UE_API FOnPlayerInputDelegateExecuted OnPlayerInputEventExecuted;
+
+	/** Event fired when the given player's input has been flushed. For example, when the viewport loses focus. */
+	DECLARE_MULTICAST_DELEGATE_OneParam(FOnPlayerInputFlushed, const UPlayerInput*);
+	static UE_API FOnPlayerInputFlushed OnPlayerInputFlushed;
+	
+	static UE_API void BroadcastPlayerInputDelegateExecuted(const UPlayerInput* PlayerInput, const FPlayerInputDebuggingArgs& EventData);
+	static UE_API void BroadcastPlayerInputFlushed(const UPlayerInput* PlayerInput);
+};
+};	// namespace UE::Input
+
+#endif	// UE_PLAYER_INPUT_INCLUDE_DEBUG
+
+#undef UE_API
\ No newline at end of file
diff --git a/Engine/Source/Runtime/Engine/Private/UserInterface/PlayerInput.cpp b/Engine/Source/Runtime/Engine/Private/UserInterface/PlayerInput.cpp
--- a/Engine/Source/Runtime/Engine/Private/UserInterface/PlayerInput.cpp
+++ b/Engine/Source/Runtime/Engine/Private/UserInterface/PlayerInput.cpp
@@ -5,6 +5,7 @@
 =============================================================================*/
 
 #include "GameFramework/PlayerInput.h"
+#include "GameFramework/PlayerInputDebugging.h"
 #include "Engine/World.h"
 #include "Misc/App.h"
 #include "GlobalRenderResources.h"
@@ -187,6 +188,10 @@ void UPlayerInput::FlushPressedKeys()
 		// Flag each key as having been flushed this frame so that we can correctly evaluate it next time
 		KeyState.bWasJustFlushed = true;
 	}
+
+#if UE_PLAYER_INPUT_INCLUDE_DEBUG
+	UE::Input::FPlayerInputDebugging::BroadcastPlayerInputFlushed(this);
+#endif
 }
 
 void UPlayerInput::FlushPressedActionBindingKeys(FName ActionName)
@@ -1552,6 +1557,17 @@ bool UPlayerInput::EvaluateInputComponentDelegates(UInputComponent* const IC, co
 		{
 			FoundChords[ChordIndex].FoundIndex = NonAxisDelegates.Num();
 			NonAxisDelegates.Add(FoundChords[ChordIndex]);
+
+#if UE_PLAYER_INPUT_INCLUDE_DEBUG
+			UE::Input::FPlayerInputDebugging::BroadcastPlayerInputDelegateExecuted(
+				this,
+				{
+					.InputComponent = IC,
+					.ListeningObject = FoundChords[ChordIndex].ActionDelegate.GetUObject(),
+					.ActionName = FoundChords[ChordIndex].SourceAction ? FoundChords[ChordIndex].SourceAction->GetActionName().ToString() : TEXT("Unknown"),
+					.InputValue = FoundChords[ChordIndex].KeyEvent == IE_Pressed ? FVector(1.0,0.0,0.0) : FVector(0.0,1.0,0.0),
+				});
+#endif
 		}
 	}
 
@@ -1571,12 +1587,34 @@ bool UPlayerInput::EvaluateInputComponentDelegates(UInputComponent* const IC, co
 					FVector* TouchedLocation = TouchEventLocations.Find(EventIndices[0]);
 					FDelegateDispatchDetails TouchInfo(EventIndices[0], NonAxisDelegates.Num(), TB.TouchDelegate, TouchedLocation != nullptr ? *TouchedLocation : FVector(-1.0f, -1.0f, 0.0f), (ETouchIndex::Type)TouchIndex);
 					NonAxisDelegates.Add(TouchInfo);
+#if UE_PLAYER_INPUT_INCLUDE_DEBUG
+					UE::Input::FPlayerInputDebugging::BroadcastPlayerInputDelegateExecuted(
+						this,
+						{
+							.InputComponent = IC,
+							.ListeningObject = TB.TouchDelegate.GetUObject(),
+							.ActionName = TouchInfo.SourceAction ? TouchInfo.SourceAction->GetActionName().ToString() : TEXT("Unknown"),
+							.InputValue = TouchInfo.TouchLocation,
+						});
+#endif
+					
 					for (int32 EventsIndex = 1; EventsIndex < EventIndices.Num(); ++EventsIndex)
 					{
 						TouchInfo.EventIndex = EventIndices[EventsIndex];
 						TouchedLocation = TouchEventLocations.Find(TouchInfo.EventIndex);
 						TouchInfo.TouchLocation = TouchedLocation != nullptr ? *TouchedLocation : FVector(-1.0f, -1.0f, 0.0f);
 						NonAxisDelegates.Add(TouchInfo);
+
+#if UE_PLAYER_INPUT_INCLUDE_DEBUG
+						UE::Input::FPlayerInputDebugging::BroadcastPlayerInputDelegateExecuted(
+							this,
+							{
+								.InputComponent = IC,
+								.ListeningObject = TB.TouchDelegate.GetUObject(),
+								.ActionName = TouchInfo.SourceAction ? TouchInfo.SourceAction->GetActionName().ToString() : TEXT("Unknown"),
+								.InputValue = TouchInfo.TouchLocation,
+							});
+#endif
 					}
 				}
 				if (TB.bConsumeInput)
@@ -1602,6 +1640,17 @@ bool UPlayerInput::EvaluateInputComponentDelegates(UInputComponent* const IC, co
 				FDelegateDispatchDetails GestureInfo(EventIndices[0], NonAxisDelegates.Num(), GB.GestureDelegate, KeyState->Value.X);
 				NonAxisDelegates.Emplace(MoveTemp(GestureInfo));
 
+#if UE_PLAYER_INPUT_INCLUDE_DEBUG
+				UE::Input::FPlayerInputDebugging::BroadcastPlayerInputDelegateExecuted(
+					this,
+					{
+						.InputComponent = IC,
+						.ListeningObject = GB.GestureDelegate.GetUObject(),
+						.ActionName = GestureInfo.SourceAction ? GestureInfo.SourceAction->GetActionName().ToString() : TEXT("Unknown"),
+						.InputValue = GestureInfo.TouchLocation,
+					});
+#endif
+
 				if (GB.bConsumeInput)
 				{
 					KeysToConsume.AddUnique(GB.GestureKey);
@@ -1623,6 +1672,15 @@ bool UPlayerInput::EvaluateInputComponentDelegates(UInputComponent* const IC, co
 		if ((bHadAnyKeys || !UE::Input::AxisEventsCanBeConsumed) && AB.AxisDelegate.IsBound())
 		{
 			AxisDelegates.Emplace(FAxisDelegateDetails(AB.AxisDelegate, AB.AxisValue));
+#if UE_PLAYER_INPUT_INCLUDE_DEBUG
+			UE::Input::FPlayerInputDebugging::BroadcastPlayerInputDelegateExecuted(
+				this,
+				{
+					.InputComponent = IC,
+					.ActionName = AB.AxisName.ToString(),
+					.InputValue = FVector(AB.AxisValue, 0.0, 0.0),
+				});
+#endif
 		}
 	}
 	for (FInputAxisKeyBinding& AxisKeyBinding : IC->AxisKeyBindings)
@@ -1651,6 +1709,16 @@ bool UPlayerInput::EvaluateInputComponentDelegates(UInputComponent* const IC, co
 		if (AxisKeyBinding.AxisDelegate.IsBound())
 		{
 			AxisDelegates.Emplace(FAxisDelegateDetails(AxisKeyBinding.AxisDelegate, AxisKeyBinding.AxisValue));
+#if UE_PLAYER_INPUT_INCLUDE_DEBUG
+			UE::Input::FPlayerInputDebugging::BroadcastPlayerInputDelegateExecuted(
+				this,
+				{
+					.InputComponent = IC,
+					.ListeningObject = AxisKeyBinding.AxisDelegate.GetUObject(),
+					.ActionName = AxisKeyBinding.AxisKey.GetFName().ToString(),
+					.InputValue = FVector(AxisKeyBinding.AxisValue, 0.0, 0.0),
+				});
+#endif
 		}
 	}
 	for (FInputVectorAxisBinding& VectorAxisBinding : IC->VectorAxisBindings)
@@ -1675,6 +1743,16 @@ bool UPlayerInput::EvaluateInputComponentDelegates(UInputComponent* const IC, co
 		if (VectorAxisBinding.AxisDelegate.IsBound())
 		{
 			VectorAxisDelegates.Emplace(FVectorAxisDelegateDetails(VectorAxisBinding.AxisDelegate, VectorAxisBinding.AxisValue));
+#if UE_PLAYER_INPUT_INCLUDE_DEBUG
+			UE::Input::FPlayerInputDebugging::BroadcastPlayerInputDelegateExecuted(
+				this,
+				{
+					.InputComponent = IC,
+					.ListeningObject = VectorAxisBinding.AxisDelegate.GetUObject(),
+					.ActionName = VectorAxisBinding.AxisKey.GetFName().ToString(),
+					.InputValue = VectorAxisBinding.AxisValue,
+				});
+#endif
 		}
 	}
 
diff --git a/Engine/Source/Runtime/Engine/Private/UserInterface/PlayerInputDebugging.cpp b/Engine/Source/Runtime/Engine/Private/UserInterface/PlayerInputDebugging.cpp
new file mode 100644
--- /dev/null
+++ b/Engine/Source/Runtime/Engine/Private/UserInterface/PlayerInputDebugging.cpp
@@ -0,0 +1,32 @@
+// Copyright Epic Games, Inc. All Rights Reserved.
+
+#include "GameFramework/PlayerInputDebugging.h"
+
+#if UE_PLAYER_INPUT_INCLUDE_DEBUG
+
+#include "GameFramework/PlayerInput.h"
+
+namespace UE::Input
+{
+	
+FPlayerInputDebugging::FOnPlayerInputFlushed FPlayerInputDebugging::OnPlayerInputFlushed;
+FPlayerInputDebugging::FOnPlayerInputDelegateExecuted FPlayerInputDebugging::OnPlayerInputEventExecuted;
+
+void FPlayerInputDebugging::BroadcastPlayerInputDelegateExecuted(const UPlayerInput* PlayerInput, const FPlayerInputDebuggingArgs& EventData)
+{
+	if (OnPlayerInputEventExecuted.IsBound())
+	{
+		OnPlayerInputEventExecuted.Broadcast(PlayerInput, EventData);
+	}
+}
+
+void FPlayerInputDebugging::BroadcastPlayerInputFlushed(const UPlayerInput* PlayerInput)
+{
+	if (OnPlayerInputFlushed.IsBound())
+	{
+		OnPlayerInputFlushed.Broadcast(PlayerInput);
+	}
+}
+}
+
+#endif	// #if UE_PLAYER_INPUT_INCLUDE_DEBUG
