Refactor ActorDescContainer delegate handling to reduce global delegate subscriptions:
- Moved delegate registration from UActorDescContainerInstance to UActorDescContainer.
- Container now subscribes to global delegates once and broadcasts to local events.
- Instances subscribe to container events instead of global delegates directly.
- Instances don't register any events when cooking anymore (same logic as containers).
This reduces global delegate list pollution and improves unregister performance by reducing the number of linear searches required (1 per container vs N per instance).
Also removed deprecated 5.4 method declarations and added profiler instrumentation for events registration code.

#Synced-CL 51337651
#Synced-CL 51337651

diff --git a/Engine/Source/Runtime/Engine/Private/WorldPartition/ActorDescContainer.cpp b/Engine/Source/Runtime/Engine/Private/WorldPartition/ActorDescContainer.cpp
--- a/Engine/Source/Runtime/Engine/Private/WorldPartition/ActorDescContainer.cpp
+++ b/Engine/Source/Runtime/Engine/Private/WorldPartition/ActorDescContainer.cpp
@@ -570,6 +570,11 @@ void UActorDescContainer::OnObjectPreSave(UObject* Object, FObjectPreSaveContext
 	}
 }
 
+void UActorDescContainer::OnObjectsReplaced(const TMap<UObject*, UObject*>& InOldToNewObjectMap)
+{
+	OnObjectsReplacedEvent.Broadcast(InOldToNewObjectMap);
+}
+
 void UActorDescContainer::OnPackageDeleted(UPackage* Package)
 {
 	if (const AActor* Actor = AActor::FindActorInPackage(Package))
@@ -597,6 +602,11 @@ void UActorDescContainer::OnPackageDeleted(UPackage* Package)
 	}
 }
 
+void UActorDescContainer::OnEditorActorReplaced(AActor* InOldActor, AActor* InNewActor)
+{
+	OnEditorActorReplacedEvent.Broadcast(InOldActor, InNewActor);
+}
+
 void UActorDescContainer::OnClassDescriptorUpdated(const FWorldPartitionActorDesc* InClassDesc)
 {
 	FWorldPartitionClassDescRegistry& ClassDescRegistry = FWorldPartitionClassDescRegistry::Get();
@@ -671,8 +681,12 @@ bool UActorDescContainer::ShouldRegisterDelegates() const
 
 void UActorDescContainer::RegisterEditorDelegates()
 {
+	TRACE_CPUPROFILER_EVENT_SCOPE(UActorDescContainer::RegisterEditorDelegates);
+
 	FCoreUObjectDelegates::OnObjectPreSave.AddUObject(this, &UActorDescContainer::OnObjectPreSave);
+	FCoreUObjectDelegates::OnObjectsReplaced.AddUObject(this, &UActorDescContainer::OnObjectsReplaced);
 	FEditorDelegates::OnPackageDeleted.AddUObject(this, &UActorDescContainer::OnPackageDeleted);
+	FEditorDelegates::OnEditorActorReplaced.AddUObject(this, &UActorDescContainer::OnEditorActorReplaced);
 
 	FWorldPartitionClassDescRegistry& ClassDescRegistry = FWorldPartitionClassDescRegistry::Get();
 	ClassDescRegistry.OnClassDescriptorUpdated().AddUObject(this, &UActorDescContainer::OnClassDescriptorUpdated);
@@ -682,7 +696,13 @@ void UActorDescContainer::RegisterEditorDelegates()
 
 void UActorDescContainer::UnregisterEditorDelegates()
 {
+	TRACE_CPUPROFILER_EVENT_SCOPE(UActorDescContainer::UnregisterEditorDelegates);
+
 	FCoreUObjectDelegates::OnObjectPreSave.RemoveAll(this);
+	FCoreUObjectDelegates::OnObjectsReplaced.RemoveAll(this);
+	FEditorDelegates::OnPackageDeleted.RemoveAll(this);
+	FEditorDelegates::OnEditorActorReplaced.RemoveAll(this);
+
 	FEditorDelegates::OnPackageDeleted.RemoveAll(this);
 
 	FWorldPartitionClassDescRegistry& ClassDescRegistry = FWorldPartitionClassDescRegistry::Get();
diff --git a/Engine/Source/Runtime/Engine/Private/WorldPartition/ActorDescContainerInstance.cpp b/Engine/Source/Runtime/Engine/Private/WorldPartition/ActorDescContainerInstance.cpp
--- a/Engine/Source/Runtime/Engine/Private/WorldPartition/ActorDescContainerInstance.cpp
+++ b/Engine/Source/Runtime/Engine/Private/WorldPartition/ActorDescContainerInstance.cpp
@@ -239,19 +239,20 @@ bool UActorDescContainerInstance::ShouldRegisterDelegates() const
 	// No World Partition means we are a ChildContainerInstance created for StreamingGeneration and need to listen to some events for updates
 	const UWorldPartition* WorldPartition = GetOuterWorldPartition();
 	const UWorld* OwningWorld = WorldPartition ? WorldPartition->GetWorld() : nullptr;
-
-	return !OwningWorld || !OwningWorld->IsGameWorld();
+	return (!OwningWorld || !OwningWorld->IsGameWorld()) && !IsRunningCookCommandlet();
 }
 
 void UActorDescContainerInstance::RegisterDelegates()
 {
+	TRACE_CPUPROFILER_EVENT_SCOPE(UActorDescContainerInstance::RegisterDelegates);
+
 	check(Container);
 
 	// Only listen to Object replaced events on ContainerInstance that have a direct World Partition outer (Loaded Container Instances: Main World or Loaded Level Instances)
 	if (GetOuterWorldPartition())
 	{
-		FCoreUObjectDelegates::OnObjectsReplaced.AddUObject(this, &UActorDescContainerInstance::OnObjectsReplaced);
-		FEditorDelegates::OnEditorActorReplaced.AddUObject(this, &UActorDescContainerInstance::OnEditorActorReplaced);
+		Container->OnObjectsReplacedEvent.AddUObject(this, &UActorDescContainerInstance::OnObjectsReplaced);
+		Container->OnEditorActorReplacedEvent.AddUObject(this, &UActorDescContainerInstance::OnEditorActorReplaced);
 	}
 
 	// Only listen to this event if we have registered Child Container Instances
@@ -277,10 +278,11 @@ void UActorDescContainerInstance::RegisterDelegates()
 
 void UActorDescContainerInstance::UnregisterDelegates()
 {
-	FCoreUObjectDelegates::OnObjectsReplaced.RemoveAll(this);
-	FEditorDelegates::OnEditorActorReplaced.RemoveAll(this);
+	TRACE_CPUPROFILER_EVENT_SCOPE(UActorDescContainerInstance::UnregisterDelegates);
 
 	check(Container);
+	Container->OnObjectsReplacedEvent.RemoveAll(this);
+	Container->OnEditorActorReplacedEvent.RemoveAll(this);
 	Container->OnActorDescAddedEvent.RemoveAll(this);
 	Container->OnActorDescRemovedEvent.RemoveAll(this);
 	Container->OnActorDescPreUpdateEvent.RemoveAll(this);
diff --git a/Engine/Source/Runtime/Engine/Public/WorldPartition/ActorDescContainer.h b/Engine/Source/Runtime/Engine/Public/WorldPartition/ActorDescContainer.h
--- a/Engine/Source/Runtime/Engine/Public/WorldPartition/ActorDescContainer.h
+++ b/Engine/Source/Runtime/Engine/Public/WorldPartition/ActorDescContainer.h
@@ -57,7 +57,9 @@ public:
 	bool IsInitialized() const { return bContainerInitialized; }
 
 	ENGINE_API void OnObjectPreSave(UObject* Object, FObjectPreSaveContext SaveContext);
+	ENGINE_API void OnObjectsReplaced(const TMap<UObject*, UObject*>& InOldToNewObjectMap);
 	ENGINE_API void OnPackageDeleted(UPackage* Package);
+	ENGINE_API void OnEditorActorReplaced(AActor* InOldActor, AActor* InNewActor);
 	ENGINE_API void OnClassDescriptorUpdated(const FWorldPartitionActorDesc* InClassDesc);
 
 	virtual FString GetContainerName() const { return ContainerPackageName.ToString(); }
@@ -95,6 +97,15 @@ public:
 	DECLARE_EVENT_OneParam(UActorDescContainer, FActorDescUpdatedEvent, FWorldPartitionActorDesc*);
 	FActorDescUpdatedEvent OnActorDescUpdatedEvent;
 
+	// On objects replaced event, for container instances
+	using FOnObjectsReplacedEventType = const TMap<UObject*, UObject*>&;
+	DECLARE_EVENT_OneParam(UActorDescContainer, FOnObjectsReplacedEvent, FOnObjectsReplacedEventType);
+	FOnObjectsReplacedEvent OnObjectsReplacedEvent;
+
+	// On editor object replaced event, for container instances
+	DECLARE_EVENT_TwoParams(UActorDescContainer, FOnEditorActorReplaced, AActor*, AActor*);
+	FOnEditorActorReplaced OnEditorActorReplacedEvent;
+
 	DECLARE_MULTICAST_DELEGATE_OneParam(FActorDescContainerInitializeDelegate, UActorDescContainer*);
 	static ENGINE_API FActorDescContainerInitializeDelegate OnActorDescContainerInitialized;
 	
@@ -125,35 +136,6 @@ public:
 
 	TArray<FAssetData> InvalidActors;
 
-	//~ Begin Deprecation
-	UE_DEPRECATED(5.4, "UActorDescContainer::Update is deprecated.")
-	void Update() {}
-
-	UE_DEPRECATED(5.4, "Use UActorDescContainerInstance::GetInstancingContext instead")
-	const FLinkerInstancingContext* GetInstancingContext() const { return nullptr; }
-	
-	UE_DEPRECATED(5.4, "Use UActorDescContainerInstance::GetInstanceTransform instead")
-	FTransform GetInstanceTransform() const { return FTransform::Identity; }
-
-	UE_DEPRECATED(5.4, "Use UActorDescContainerInstance::LoadAllActors instead")
-	void LoadAllActors(TArray<FWorldPartitionReference>& OutReferences) {}
-
-	UE_DEPRECATED(5.4, "Use UActorDescContainerInstance::OnObjectsReplaced instead")
-	void OnObjectsReplaced(const TMap<UObject*, UObject*>& OldToNewObjectMap) {}
-
-	UE_DEPRECATED(5.4, "UActorDescContainer::IsTemplateContainer is deprecated")
-	bool IsTemplateContainer() const { return true;}
-
-	UE_DEPRECATED(5.4, "UActorDescContainer::IsMainWorldPartitionContainer is deprecated")
-	bool IsMainPartitionContainer() const { return false; }
-	
-	UE_DEPRECATED(5.4, "Use UActorDescContainerInstance::GetWorldPartition instead")
-	UWorldPartition* GetWorldPartition() const { return nullptr; }
-
-	UE_DEPRECATED(5.4, "Use version without UWorld parameter")
-	void RegisterActorDescriptor(FWorldPartitionActorDesc* ActorDesc, UWorld* InWorldContext) {}
-	//~ End Deprecation
-
 	void ForEachChildActorDesc(const FGuid& InActorGuid, TFunctionRef<void(const FGuid&)> InFunc);
 
 protected:
