Make DataflowDebugDraw use the PriorityGroup state for DrawLine (matching DrawPoint behavior)
+ render medial axis skeleton lines w/ screen priority
#Synced-CL 51328354

diff --git a/Engine/Plugins/Experimental/GeometryCollectionPlugin/Source/GeometryCollectionNodes/Private/Dataflow/GeometryCollectionSamplingNodes.cpp b/Engine/Plugins/Experimental/GeometryCollectionPlugin/Source/GeometryCollectionNodes/Private/Dataflow/GeometryCollectionSamplingNodes.cpp
--- a/Engine/Plugins/Experimental/GeometryCollectionPlugin/Source/GeometryCollectionNodes/Private/Dataflow/GeometryCollectionSamplingNodes.cpp
+++ b/Engine/Plugins/Experimental/GeometryCollectionPlugin/Source/GeometryCollectionNodes/Private/Dataflow/GeometryCollectionSamplingNodes.cpp
@@ -29,7 +29,7 @@ namespace UE::Dataflow
 	{
 		using namespace UE::Geometry;
 
-		DataflowRenderingInterface.SetWorldPriority();
+		DataflowRenderingInterface.SetForegroundPriority();
 		DataflowRenderingInterface.SetColor(FLinearColor::Green);
 		for (int32 ClusterIdx = 0; ClusterIdx < ToDrawSkel.ClusterNeighbors.Num(); ++ClusterIdx)
 		{
diff --git a/Engine/Source/Runtime/Engine/Private/DebugRenderSceneProxy.cpp b/Engine/Source/Runtime/Engine/Private/DebugRenderSceneProxy.cpp
--- a/Engine/Source/Runtime/Engine/Private/DebugRenderSceneProxy.cpp
+++ b/Engine/Source/Runtime/Engine/Private/DebugRenderSceneProxy.cpp
@@ -466,7 +466,7 @@ FMaterialRenderProxy* FDebugRenderSceneProxy::FMaterialCache::operator[](FLinear
 
 void FDebugRenderSceneProxy::FDebugLine::Draw(FPrimitiveDrawInterface* PDI) const
 {
-	PDI->DrawLine(Start, End, Color, SDPG_World, Thickness, 0, Thickness > 0);
+	PDI->DrawLine(Start, End, Color, PriorityGroup, Thickness, 0, Thickness > 0);
 }
 
 void FDebugRenderSceneProxy::FWireStar::Draw(FPrimitiveDrawInterface* PDI) const
diff --git a/Engine/Source/Runtime/Engine/Public/DebugRenderSceneProxy.h b/Engine/Source/Runtime/Engine/Public/DebugRenderSceneProxy.h
--- a/Engine/Source/Runtime/Engine/Public/DebugRenderSceneProxy.h
+++ b/Engine/Source/Runtime/Engine/Public/DebugRenderSceneProxy.h
@@ -144,17 +144,19 @@ public:
 	/** Struct to hold info about lines to render. */
 	struct FDebugLine
 	{
-		FDebugLine(const FVector &InStart, const FVector &InEnd, const FColor &InColor, const float InThickness = 0.f) 
+		FDebugLine(const FVector &InStart, const FVector &InEnd, const FColor &InColor, const float InThickness = 0.f, ESceneDepthPriorityGroup InPriorityGroup = ESceneDepthPriorityGroup::SDPG_World)
 			: Start(InStart)
 			, End(InEnd)
 			, Color(InColor)
-			, Thickness(InThickness) 
+			, Thickness(InThickness)
+			, PriorityGroup(InPriorityGroup)
 		{}
 
 		FVector Start;
 		FVector End;
 		FColor Color;
 		float Thickness;
+		ESceneDepthPriorityGroup PriorityGroup;
 
 		ENGINE_API void Draw(FPrimitiveDrawInterface* PDI) const;
 	};
diff --git a/Engine/Source/Runtime/Experimental/Dataflow/Engine/Private/Dataflow/DataflowDebugDraw.cpp b/Engine/Source/Runtime/Experimental/Dataflow/Engine/Private/Dataflow/DataflowDebugDraw.cpp
--- a/Engine/Source/Runtime/Experimental/Dataflow/Engine/Private/Dataflow/DataflowDebugDraw.cpp
+++ b/Engine/Source/Runtime/Experimental/Dataflow/Engine/Private/Dataflow/DataflowDebugDraw.cpp
@@ -135,7 +135,7 @@ void FDataflowDebugDraw::DrawPoint(const FVector& Position)
 
 void FDataflowDebugDraw::DrawLine(const FVector& Start, const FVector& End) const
 {
-	DebugRenderSceneProxy->Lines.Add(FDebugRenderSceneProxy::FDebugLine(Start, End, ColorWithTranslucency.ToFColor(true), LineWidth));
+	DebugRenderSceneProxy->Lines.Add(FDebugRenderSceneProxy::FDebugLine(Start, End, ColorWithTranslucency.ToFColor(true), LineWidth, PriorityGroup));
 }
 
 void FDataflowDebugDraw::DrawMesh(const IDebugDrawMesh& Mesh) const
