Fixing issue where FSkeletalMeshObjectStatic objects would try to render at a LOD level that wasn't loaded yet.

diff --git a/Engine/Source/Runtime/Engine/Private/SkeletalRenderStatic.cpp b/Engine/Source/Runtime/Engine/Private/SkeletalRenderStatic.cpp
--- a/Engine/Source/Runtime/Engine/Private/SkeletalRenderStatic.cpp
+++ b/Engine/Source/Runtime/Engine/Private/SkeletalRenderStatic.cpp
@@ -91,7 +91,8 @@ TConstArrayView<FMatrix44f> FSkeletalMeshObjectStatic::GetReferenceToLocalMatric
 
 int32 FSkeletalMeshObjectStatic::GetLOD() const
 {
-	return FMath::Max<int32>(WorkingMinDesiredLODLevel, MinLODLevel);
+	// WorkingMinDesiredLODLevel and MinLODLevel could be a LOD that's not loaded, so we need to clamp it to the first loaded LOD in the render data.
+	return FMath::Max3<int32>(WorkingMinDesiredLODLevel, MinLODLevel, SkeletalMeshRenderData->CurrentFirstLODIdx);
 }
 
 void FSkeletalMeshObjectStatic::FSkeletalMeshObjectLOD::InitResources(const FSkelMeshComponentLODInfo* CompLODInfo)
