Added support for getting correct viewport local cursor position in cases where viewport size is fixed and different from rendered size in Slate

#Synced-CL 51386868

diff --git a/Engine/Source/Editor/LevelEditor/Private/SLevelViewport.cpp b/Engine/Source/Editor/LevelEditor/Private/SLevelViewport.cpp
--- a/Engine/Source/Editor/LevelEditor/Private/SLevelViewport.cpp
+++ b/Engine/Source/Editor/LevelEditor/Private/SLevelViewport.cpp
@@ -965,7 +965,8 @@ bool SLevelViewport::HandleDragObjects(const FGeometry& MyGeometry, const FDragD
 		LevelViewportClient->GetViewportDimensions(ViewportOrigin, ViewportSize);
 
 		// Save off the local mouse position from the drop point for potential use later (with Drag Drop context menu)
-		CachedOnDropLocalMousePos = FVector2f(MyGeometry.AbsoluteToLocal( DragDropEvent.GetScreenSpacePosition() ) * MyGeometry.Scale).IntPoint();
+		const FVector2D ToViewportPixelScale = UE::Slate::Viewport::GetViewportLocalToPixelScale(MyGeometry, ViewportSize);
+		CachedOnDropLocalMousePos = FVector2f(MyGeometry.AbsoluteToLocal( DragDropEvent.GetScreenSpacePosition() ) * ToViewportPixelScale).IntPoint();
 		CachedOnDropLocalMousePos -= ViewportOrigin;
 	}
 
diff --git a/Engine/Source/Runtime/Engine/Private/Slate/SceneViewport.cpp b/Engine/Source/Runtime/Engine/Private/Slate/SceneViewport.cpp
--- a/Engine/Source/Runtime/Engine/Private/Slate/SceneViewport.cpp
+++ b/Engine/Source/Runtime/Engine/Private/Slate/SceneViewport.cpp
@@ -124,8 +124,9 @@ void FSceneViewport::ShowCursor( bool bVisible )
 	{
 		if( bIsSoftwareCursorVisible )
 		{
-			const int32 ClampedMouseX = FMath::Clamp<int32>(SoftwareCursorPosition.X / CachedGeometry.Scale, 0, SizeX);
-			const int32 ClampedMouseY = FMath::Clamp<int32>(SoftwareCursorPosition.Y / CachedGeometry.Scale, 0, SizeY);
+			const FVector2D ToViewportPixelScale = UE::Slate::Viewport::GetViewportLocalToPixelScale(CachedGeometry, FVector2D(SizeX, SizeY));
+			const int32 ClampedMouseX = FMath::Clamp<int32>(SoftwareCursorPosition.X * ToViewportPixelScale.X, 0, SizeX);
+			const int32 ClampedMouseY = FMath::Clamp<int32>(SoftwareCursorPosition.Y * ToViewportPixelScale.Y, 0, SizeY);
 
 			CurrentReplyState.SetMousePos( CachedGeometry.LocalToAbsolute( FVector2D(ClampedMouseX, ClampedMouseY) ).IntPoint() );
 		}
@@ -229,14 +230,17 @@ void FSceneViewport::UpdateCachedCursorPos( const FGeometry& InGeometry, const F
 	if (UserId == FSlateApplication::SlateAppPrimaryPlatformUser)
 	{
 		FVector2D LocalPixelMousePos = InGeometry.AbsoluteToLocal(InMouseEvent.GetScreenSpacePosition());
-		LocalPixelMousePos.X = FMath::Clamp(LocalPixelMousePos.X * CachedGeometry.Scale, (double)TNumericLimits<int32>::Min(), (double)TNumericLimits<int32>::Max());
-		LocalPixelMousePos.Y = FMath::Clamp(LocalPixelMousePos.Y * CachedGeometry.Scale, (double)TNumericLimits<int32>::Min(), (double)TNumericLimits<int32>::Max());
+
+		const FVector2D ToViewportPixelScale = UE::Slate::Viewport::GetViewportLocalToPixelScale(InGeometry, FVector2D(SizeX, SizeY));
+
+		LocalPixelMousePos.X = FMath::Clamp(LocalPixelMousePos.X * ToViewportPixelScale.X, (double)TNumericLimits<int32>::Min(), (double)TNumericLimits<int32>::Max());
+		LocalPixelMousePos.Y = FMath::Clamp(LocalPixelMousePos.Y * ToViewportPixelScale.Y, (double)TNumericLimits<int32>::Min(), (double)TNumericLimits<int32>::Max());
 		
 		const FIntPoint NewPos = LocalPixelMousePos.IntPoint();
 		const FIntPoint Delta = NewPos - CachedCursorPos;
 		
-		const int32 PixelWidth = FMath::RoundToInt32(FMath::Clamp(InGeometry.Size.X * CachedGeometry.Scale, (double)TNumericLimits<int32>::Min(), (double)TNumericLimits<int32>::Max()));
-		const int32 PixelHeight = FMath::RoundToInt32(FMath::Clamp(InGeometry.Size.Y * CachedGeometry.Scale, (double)TNumericLimits<int32>::Min(), (double)TNumericLimits<int32>::Max()));
+		const int32 PixelWidth = FMath::RoundToInt32(FMath::Clamp(InGeometry.Size.X * ToViewportPixelScale.X, (double)TNumericLimits<int32>::Min(), (double)TNumericLimits<int32>::Max()));
+		const int32 PixelHeight = FMath::RoundToInt32(FMath::Clamp(InGeometry.Size.Y * ToViewportPixelScale.Y, (double)TNumericLimits<int32>::Min(), (double)TNumericLimits<int32>::Max()));
 		
 		CachedCursorPos = NewPos;
 		
@@ -250,7 +254,7 @@ void FSceneViewport::UpdateCachedCursorPos( const FGeometry& InGeometry, const F
 			// When the mouse is outside the viewport or on the edge, rely on straight deltas.
 			// This distinction was made due to seeing duplicate OnMouseMove deltas being applied.
 			// If that problem can be ameliorated, the above branch should no longer be necessary.
-			const FIntPoint UnboundDelta = (InMouseEvent.GetCursorDelta() * InGeometry.Scale).IntPoint();
+			const FIntPoint UnboundDelta = (InMouseEvent.GetCursorDelta() * ToViewportPixelScale).IntPoint();
 			CachedUnboundedCursorPos += UnboundDelta;
 		}
 	}
diff --git a/Engine/Source/Runtime/Slate/Private/Widgets/SViewport.cpp b/Engine/Source/Runtime/Slate/Private/Widgets/SViewport.cpp
--- a/Engine/Source/Runtime/Slate/Private/Widgets/SViewport.cpp
+++ b/Engine/Source/Runtime/Slate/Private/Widgets/SViewport.cpp
@@ -194,7 +194,8 @@ int32 SViewport::OnPaint( const FPaintArgs& Args, const FGeometry& AllottedGeome
 		}
 		const FVector2f CursorSize = Brush->ImageSize / AllottedGeometry.Scale;
 
-		FVector2f CursorPositionLocalSpace = UE::Slate::CastToVector2f(ViewportInterfacePin->GetSoftwareCursorPosition()) / AllottedGeometry.Scale;
+		const FVector2D ToViewportPixelScale = UE::Slate::Viewport::GetViewportLocalToPixelScale(AllottedGeometry, ViewportInterfacePin->GetSize());
+		FVector2f CursorPositionLocalSpace = UE::Slate::CastToVector2f(ViewportInterfacePin->GetSoftwareCursorPosition() / ToViewportPixelScale);
 
 		LayerId++;
 		FSlateDrawElement::MakeBox(
@@ -445,3 +446,16 @@ FNavigationReply SViewport::OnNavigation(const FGeometry& MyGeometry, const FNav
 {
 	return ViewportInterface.IsValid() ? ViewportInterface.Pin()->OnNavigation(MyGeometry, InNavigationEvent) : FNavigationReply::Stop();
 }
+
+FVector2D UE::Slate::Viewport::GetViewportLocalToPixelScale(const FGeometry& InGeometry, const FVector2D& InViewportSize)
+{
+	const FDeprecateVector2DResult SlateSize = InGeometry.GetAbsoluteSize();
+
+	// Check for invalid sizes, return identity scalar in such a case
+	if (SlateSize.IsNearlyZero() || InViewportSize.IsNearlyZero())
+	{
+		return FVector2D(1.0f, 1.0f);
+	}
+
+	return InGeometry.Scale * (InViewportSize / SlateSize);
+}
diff --git a/Engine/Source/Runtime/Slate/Public/Widgets/SViewport.h b/Engine/Source/Runtime/Slate/Public/Widgets/SViewport.h
--- a/Engine/Source/Runtime/Slate/Public/Widgets/SViewport.h
+++ b/Engine/Source/Runtime/Slate/Public/Widgets/SViewport.h
@@ -292,3 +292,14 @@ private:
 	/** True if the viewport texture has pre-multiplied alpha */
 	bool bPreMultipliedAlpha;
 };
+
+namespace UE::Slate::Viewport
+{
+	/**
+	 * Calculates the scaling factor needed to transform from local Slate geometry to viewport pixel geometry
+	 * @param InGeometry The viewport's Slate geometry
+	 * @param InViewportSize The viewport's render size
+	 * @return A non-uniform scaling factor to convert from local Slate coordinates to viewport pixel coordinates
+	 */
+	SLATE_API FVector2D GetViewportLocalToPixelScale(const FGeometry& InGeometry, const FVector2D& InViewportSize);
+}
\ No newline at end of file
