Add scopes around DeviceProfileManager

#Synced-CL 51322414

#Synced-CL 51322414

diff --git a/Engine/Source/Runtime/Engine/Private/DeviceProfiles/DeviceProfileManager.cpp b/Engine/Source/Runtime/Engine/Private/DeviceProfiles/DeviceProfileManager.cpp
--- a/Engine/Source/Runtime/Engine/Private/DeviceProfiles/DeviceProfileManager.cpp
+++ b/Engine/Source/Runtime/Engine/Private/DeviceProfiles/DeviceProfileManager.cpp
@@ -1218,6 +1218,8 @@ void UDeviceProfileManager::GetProfileConfigFiles(OUT TArray<FString>& OutConfig
 
 void UDeviceProfileManager::LoadProfiles()
 {
+	TRACE_CPUPROFILER_EVENT_SCOPE(UDeviceProfileManager::LoadProfiles);
+
 	if( !HasAnyFlags( RF_ClassDefaultObject ) )
 	{
 		UE::ConfigAccessTracking::FIgnoreScope IgnoreConfigAccessTracking = UE::ConfigAccessTracking::FIgnoreScope();
@@ -1231,6 +1233,8 @@ void UDeviceProfileManager::LoadProfiles()
 		// go over all the platforms we find, starting with the current platform
 		for (int32 PlatformIndex = 0; PlatformIndex <= PlatformsToLoad.Num(); PlatformIndex++)
 		{
+			TRACE_CPUPROFILER_EVENT_SCOPE(UDeviceProfileManager::FindOrLoadPlatformConfig);
+
 			// which platform's set of ini files should we load from?
 			FString ConfigLoadPlatform = PlatformIndex == 0 ? FString(FPlatformProperties::IniPlatformName()) : PlatformsToLoad[PlatformIndex - 1].ToString();
 
