SpringMath: Added some tests for spring functions. Fixed bug with StartingDistance function.

#Synced-CL 51457304

diff --git a/Engine/Source/Runtime/Engine/Public/Animation/SpringMath.h b/Engine/Source/Runtime/Engine/Public/Animation/SpringMath.h
--- a/Engine/Source/Runtime/Engine/Public/Animation/SpringMath.h
+++ b/Engine/Source/Runtime/Engine/Public/Animation/SpringMath.h
@@ -67,12 +67,13 @@ private:
 		return FMath::FindDeltaAngleRadians(B, A) + B;
 	}
 
-	/** Computes the lower (-1) branch of the LambertW function, which is used for solving equations of the form x*e^x */
+	/** Computes the lower (-1) branch of the LambertW function, which is used for solving equations of the form x*e^x - input will be clamped to range [-1/e, 0], output will be in range [-1, -inf] */
 	static inline float LambertWLower(const float X)
 	{
-		if (X >= 0.0f) { return -UE_MAX_FLT; }
-		
-		const float Redx = 0.625529587f * 0.588108778f + X;
+		// Clamp to valid domain
+		const float XValid = FMath::Clamp(X, -1.0f / UE_EULERS_NUMBER + UE_SMALL_NUMBER, -UE_SMALL_NUMBER);
+
+		const float Redx = 0.625529587f * 0.588108778f + XValid;
 		
 		// Polynomial approximation for when close to -(exp(-1))
 		if (Redx <= 0.09765625f)
@@ -98,19 +99,19 @@ private:
 			const float C2 = -2.87322998e-2f;
 			const float C3 = +7.07275391e-1f;
 
-			const float S = FMath::Log2(-X) * -UE_LN2 - 1.0f;
+			const float S = FMath::Log2(-XValid) * -UE_LN2 - 1.0f;
 			float T = FMath::Sqrt(S);
 			float W = -1.0f - S - (1.0f / (FMath::Exp2(C2 * T) * (C1 * T) + C3 / T + C0));
 
-			if (X > -0x1.0p-116f)
+			if (XValid > -0x1.0p-116f)
 			{
 				const float E = FMath::Exp2(W / UE_LN2 + 32.0f);
-				W = -(W * E - 0x1.0p32f * X) / (W * E + E) + W;
+				W = -(W * E - 0x1.0p32f * XValid) / (W * E + E) + W;
 			}
 			else
 			{
 				T = W / (1.0f + W);
-				W = T * FMath::Loge(X / W) + T;
+				W = T * FMath::Loge(XValid / W) + T;
 			}
 
 			return W;
@@ -671,42 +672,46 @@ public:
 		}
 	}
 
-	/** Estimates the time at which a character driven by a critical spring will reach zero acceleration */
+	/** Estimates the time in the future at which a character driven by a critical spring will reach zero acceleration. */
 	UE_EXPERIMENTAL(5.8, "SpringMath is experimental")
-	static inline float SpringCharacterStoppingInflectionTime(const float InitialVelocity, const float InitialAcceleration, const float SmoothingTime)
+	static inline TOptional<float> SpringCharacterStoppingInflectionTime(const float InitialVelocity, const float InitialAcceleration, const float SmoothingTime)
 	{
-		const float V = FMath::Abs(InitialVelocity);
+		const float V = InitialVelocity;
+		const float A = InitialAcceleration;
 		const float Y = SmoothingTimeToDamping(SmoothingTime) / 2.0f;
+		const float Denom = A + V * Y;
+		
+		if (FMath::Abs(Denom) < UE_SMALL_NUMBER) { return TOptional<float>(); }
 
-		return -V / (InitialAcceleration + V * Y);
+		return -V / Denom > 0.0f ? -V / Denom : TOptional<float>();
 	}
 
 	/** Estimates an approximate stopping distance from a smoothing time for a character using a simple damped spring */
 	UE_EXPERIMENTAL(5.8, "SpringMath is experimental")
 	static inline float SpringCharacterStoppingDistance(const float InitialVelocity, const float InitialAcceleration, const float SmoothingTime)
 	{
-		const float V = FMath::Abs(InitialVelocity);
+		const float V = InitialVelocity;
 		const float A = InitialAcceleration;
 		const float Y = SmoothingTimeToDamping(SmoothingTime) / 2.0f;
 
 		if (FMath::Abs(A) < UE_SMALL_NUMBER)
 		{
-			return 2.0f * V / FMath::Max(Y, UE_SMALL_NUMBER);
+			return 2.0f * FMath::Abs(V) / FMath::Max(Y, UE_SMALL_NUMBER);
 		}
 
-		const float I0 = (V * Y * 2 + A) / FMath::Max(FMath::Square(Y), UE_SMALL_NUMBER);
+		const float I0 = (V * Y * 2 + A) / FMath::Max(Y * Y, UE_SMALL_NUMBER);
 
 		PRAGMA_DISABLE_EXPERIMENTAL_WARNINGS
-		const float T = SpringCharacterStoppingInflectionTime(V, A, SmoothingTime);
+		const TOptional<float> T = SpringCharacterStoppingInflectionTime(V, A, SmoothingTime);
 		PRAGMA_ENABLE_EXPERIMENTAL_WARNINGS
 
-		if (!FMath::IsFinite(T) || T <= 0.0f)
+		if (!T.IsSet())
 		{
 			return FMath::Abs(I0);
 		}
 		else
 		{
-			const float IT = (FMath::Exp(-Y * T) / FMath::Max(FMath::Square(Y), UE_SMALL_NUMBER)) * (T * A * Y + V * Y * (T * Y + 2.0f) + A);
+			const float IT = (FMath::Exp(-Y * T.GetValue()) / FMath::Max(Y * Y, UE_SMALL_NUMBER)) * (T.GetValue() * A * Y + V * Y * (T.GetValue() * Y + 2.0f) + A);
 			return FMath::Abs(I0 - IT) + FMath::Abs(IT);
 		}
 	}
@@ -715,11 +720,9 @@ public:
 	UE_EXPERIMENTAL(5.8, "SpringMath is experimental")
 	static inline float SpringCharacterStoppingTime(const float InitialVelocity, const float SmoothingTime, const float VelocityThreshold = 10.0f)
 	{
-		if (FMath::Abs(VelocityThreshold) >= FMath::Abs(InitialVelocity)) { return 0.0f; }
-
 		const float Y = SmoothingTimeToDamping(SmoothingTime) / 2.0f;
 
-		const float VRatio = FMath::Clamp(FMath::Abs(VelocityThreshold) / FMath::Max(FMath::Abs(InitialVelocity), UE_SMALL_NUMBER), UE_SMALL_NUMBER, 1.0f - UE_SMALL_NUMBER);
+		const float VRatio = FMath::Abs(VelocityThreshold) / FMath::Max(FMath::Abs(InitialVelocity), UE_SMALL_NUMBER);
 
 		return -(LambertWLower(-VRatio / UE_EULERS_NUMBER) + 1.0f) / FMath::Max(Y, UE_SMALL_NUMBER);
 	}
@@ -735,8 +738,6 @@ public:
 	UE_EXPERIMENTAL(5.8, "SpringMath is experimental")
 	static inline float SpringCharacterSmoothingTimeFromStoppingTime(const float InitialVelocity, const float StoppingTime, const float VelocityThreshold = 10.0f)
 	{
-		if (FMath::Abs(VelocityThreshold) >= FMath::Abs(InitialVelocity)) { return 0.0f; }
-
 		const float VRatio = FMath::Clamp(FMath::Abs(VelocityThreshold) / FMath::Max(FMath::Abs(InitialVelocity), UE_SMALL_NUMBER), UE_SMALL_NUMBER, 1.0f - UE_SMALL_NUMBER);
 
 		return DampingToSmoothingTime(2.0f * -(LambertWLower(-VRatio / UE_EULERS_NUMBER) + 1.0f) / FMath::Max(StoppingTime, UE_SMALL_NUMBER));
@@ -746,10 +747,12 @@ public:
 	UE_EXPERIMENTAL(5.8, "SpringMath is experimental")
 	static inline float SpringCharacterStartingTime(const float TargetVelocity, const float SmoothingTime, const float VelocityThreshold = 10.0f)
 	{
+		const float Y = SmoothingTimeToDamping(SmoothingTime) / 2.0f;
+
 		// Starting time is the same as if the character was moving at the target velocity and coming to a stop
-		PRAGMA_DISABLE_EXPERIMENTAL_WARNINGS
-		return SpringCharacterStoppingTime(TargetVelocity, SmoothingTime, VelocityThreshold);
-		PRAGMA_ENABLE_EXPERIMENTAL_WARNINGS
+		const float VRatio = FMath::Abs(VelocityThreshold) / FMath::Max(FMath::Abs(TargetVelocity), UE_SMALL_NUMBER);
+
+		return -(LambertWLower(-VRatio / UE_EULERS_NUMBER) + 1.0f) / FMath::Max(Y, UE_SMALL_NUMBER);
 	}
 
 	/** Estimates an approximate starting distance (the distance traveled until the velocity gets within the VelocityThreshold of the TargetVelocity) from a smoothing time for a character using a simple damped spring */
@@ -759,10 +762,10 @@ public:
 		const float V = FMath::Abs(TargetVelocity);
 		const float Y = SmoothingTimeToDamping(SmoothingTime) / 2.0f;
 		PRAGMA_DISABLE_EXPERIMENTAL_WARNINGS
-		const float T = SpringCharacterStartingTime(TargetVelocity, SmoothingTime, VelocityThreshold);
+		const float T = SpringCharacterStartingTime(V, SmoothingTime, VelocityThreshold);
 		PRAGMA_ENABLE_EXPERIMENTAL_WARNINGS
 
-		return FMath::Exp(-Y * T) * ((2.0f * V) / FMath::Max(Y, UE_SMALL_NUMBER) + V * T) + (2.0f * V) / FMath::Max(Y, UE_SMALL_NUMBER) + V * T;
+		return FMath::Exp(-Y * T) * ((2.0f * V) / FMath::Max(Y, UE_SMALL_NUMBER) + V * T) - (2.0f * V) / FMath::Max(Y, UE_SMALL_NUMBER) + V * T;
 	}
 
 	/** Estimates an approximate spring smoothing time from a starting time and target velocity for a character using a simple damped spring. */
@@ -779,24 +782,25 @@ public:
 	UE_EXPERIMENTAL(5.8, "SpringMath is experimental")
 		static inline float SpringCharacterSmoothingTimeFromStartingDistance(const float TargetVelocity, const float StartingDistance, const float VelocityThreshold = 10.0f)
 	{
-		if (FMath::Abs(VelocityThreshold) >= FMath::Abs(TargetVelocity)) { return 0.0f; }
-
-		const float V = FMath::Abs(TargetVelocity);
-		const float VRatio = FMath::Clamp(FMath::Abs(VelocityThreshold) / FMath::Max(FMath::Abs(TargetVelocity), UE_SMALL_NUMBER), UE_SMALL_NUMBER, 1.0f - UE_SMALL_NUMBER);
+		const float VRatio = FMath::Abs(VelocityThreshold) / FMath::Max(FMath::Abs(TargetVelocity), UE_SMALL_NUMBER);
 
 		const float W = -(LambertWLower(-VRatio / UE_EULERS_NUMBER) + 1.0f);
 
+		const float V = FMath::Abs(TargetVelocity);
 		return DampingToSmoothingTime((2.0f * FMath::Exp(-W) * V * (2.0f + W) + 2.0f * V * (W - 2.0f)) / FMath::Max(StartingDistance, UE_SMALL_NUMBER));
 	}
 
-	/** Estimates the time at which the maximum acceleration will occur */
+	/** Estimates the time at which the maximum acceleration will occur. */
 	UE_EXPERIMENTAL(5.8, "SpringMath is experimental")
-	static inline float SpringCharacterMaximumAccelerationTime(const float InitialVelocity, const float InitialAcceleration, const float TargetVelocity, const float SmoothingTime)
+	static inline TOptional<float> SpringCharacterMaximumAccelerationTime(const float InitialVelocity, const float InitialAcceleration, const float TargetVelocity, const float SmoothingTime)
 	{
 		const float Y = SmoothingTimeToDamping(SmoothingTime) / 2.0f;
 		const float J0 = InitialVelocity - TargetVelocity;
 		const float J1 = InitialAcceleration + J0 * Y;
-		return (J1 + InitialAcceleration) / (J1 * Y);
+		const float Denom = J1 * Y;
+		if (FMath::Abs(Denom) < UE_SMALL_NUMBER) { return TOptional<float>(); }
+
+		return (J1 + InitialAcceleration) / Denom > 0.0f ? (J1 + InitialAcceleration) / Denom : TOptional<float>();
 	}
 
 	/** Estimates the maximum acceleration for a character */
@@ -807,16 +811,16 @@ public:
 		const float J0 = InitialVelocity - TargetVelocity;
 		const float J1 = InitialAcceleration + J0 * Y;
 		PRAGMA_DISABLE_EXPERIMENTAL_WARNINGS
-		const float T = SpringCharacterMaximumAccelerationTime(InitialVelocity, InitialAcceleration, TargetVelocity, SmoothingTime);
+		const TOptional<float> T = SpringCharacterMaximumAccelerationTime(InitialVelocity, InitialAcceleration, TargetVelocity, SmoothingTime);
 		PRAGMA_ENABLE_EXPERIMENTAL_WARNINGS
 
-		if (FMath::IsFinite(T) && T > 0.0f)
+		if (T.IsSet())
 		{
-			return FMath::InvExpApprox(Y * T) * (InitialAcceleration - J1 * Y * T);
+			return FMath::Max(FMath::Abs(InitialAcceleration), FMath::Abs(FMath::Exp(-Y * T.GetValue()) * (InitialAcceleration - J1 * Y * T.GetValue())));
 		}
 		else
 		{
-			return InitialAcceleration;
+			return FMath::Abs(InitialAcceleration);
 		}
 	}
 
@@ -824,17 +828,20 @@ public:
 	UE_EXPERIMENTAL(5.8, "SpringMath is experimental")
 	static inline float SpringCharacterSmoothingTimeFromMaximumAcceleration(const float InitialVelocity, const float TargetVelocity, const float MaximumAcceleration)
 	{
-		return DampingToSmoothingTime(2.0f * ((MaximumAcceleration / FMath::Max(FMath::Abs(TargetVelocity - InitialVelocity), UE_SMALL_NUMBER)) * UE_EULERS_NUMBER));
+		return DampingToSmoothingTime(2.0f * ((FMath::Abs(MaximumAcceleration) / FMath::Max(FMath::Abs(TargetVelocity - InitialVelocity), UE_SMALL_NUMBER)) * UE_EULERS_NUMBER));
 	}
 
-	/** Estimates the time at which the maximum angular velocity will occur */
+	/** Estimates the time in the future at which the maximum angular velocity will occur. */
 	UE_EXPERIMENTAL(5.8, "SpringMath is experimental")
-	static inline float SpringCharacterMaximumAngularVelocityTime(const float InitialAngle, const float InitialAngularVelocity, const float TargetAngle, const float SmoothingTime)
+	static inline TOptional<float> SpringCharacterMaximumAngularVelocityTime(const float InitialAngle, const float InitialAngularVelocity, const float TargetAngle, const float SmoothingTime)
 	{
 		const float Y = SmoothingTimeToDamping(SmoothingTime) / 2.0f;
 		const float J0 = FMath::FindDeltaAngleRadians(TargetAngle, InitialAngle);
 		const float J1 = InitialAngularVelocity + J0 * Y;
-		return (J1 + InitialAngularVelocity) / (J1 * Y);
+		const float Denom = J1 * Y;
+		if (FMath::Abs(Denom) < UE_SMALL_NUMBER) { return TOptional<float>(); }
+
+		return (J1 + InitialAngularVelocity) / Denom > 0.0f ? (J1 + InitialAngularVelocity) / Denom : TOptional<float>();
 	}
 
 	/** Estimates the maximum angular velocity for a character */
@@ -845,24 +852,24 @@ public:
 		const float J0 = FMath::FindDeltaAngleRadians(TargetAngle, InitialAngle);
 		const float J1 = InitialAngularVelocity + J0 * Y;
 		PRAGMA_DISABLE_EXPERIMENTAL_WARNINGS
-		const float T = SpringCharacterMaximumAngularVelocityTime(InitialAngle, InitialAngularVelocity, TargetAngle, SmoothingTime);
+		const TOptional<float> T = SpringCharacterMaximumAngularVelocityTime(InitialAngle, InitialAngularVelocity, TargetAngle, SmoothingTime);
 		PRAGMA_ENABLE_EXPERIMENTAL_WARNINGS
 
-		if (FMath::IsFinite(T) && T > 0.0f)
+		if (T.IsSet())
 		{
-			return FMath::InvExpApprox(Y * T) * (InitialAngularVelocity - J1 * Y * T);
+			return FMath::Max(FMath::Abs(InitialAngularVelocity), FMath::Abs(FMath::Exp(-Y * T.GetValue()) * (InitialAngularVelocity - J1 * Y * T.GetValue())));
 		}
 		else
 		{
-			return InitialAngularVelocity;
+			return FMath::Abs(InitialAngularVelocity);
 		}
 	}
 
 	/** Estimates an approximate spring smoothing time from the maximum angular velocity */
 	UE_EXPERIMENTAL(5.8, "SpringMath is experimental")
-		static inline float SpringCharacterSmoothingTimeFromMaximumAngularVelocity(const float InitialAngle, const float TargetAngle, const float MaximumAngularVelocity)
+	static inline float SpringCharacterSmoothingTimeFromMaximumAngularVelocity(const float InitialAngle, const float TargetAngle, const float MaximumAngularVelocity)
 	{
-		return DampingToSmoothingTime(2.0f * ((MaximumAngularVelocity / FMath::Max(FMath::Abs(FMath::FindDeltaAngleRadians(InitialAngle, TargetAngle)), UE_SMALL_NUMBER)) * UE_EULERS_NUMBER));
+		return DampingToSmoothingTime(2.0f * ((FMath::Abs(MaximumAngularVelocity) / FMath::Max(FMath::Abs(FMath::FindDeltaAngleRadians(InitialAngle, TargetAngle)), UE_SMALL_NUMBER)) * UE_EULERS_NUMBER));
 	}
 
 	/** Update a position representing a character given a target velocity using a velocity spring.
