New material translator work:
- Refactored the way material parameters are handled by the new translator. Instead of making material parameters directly handled by the emitter, the emitter now provides 3 general APIs PrimitiveUniform(), TextureUniform() and VirtualTextureUniform(). On top of these simple APIs, the frontend compiles all the various expressions that provide a texture with or without parameters and the other parameter types. This is much better than before becauase most of the Parameter-related complexity now lives in the "frontend" (MaterialExpressions) rather than in the translator, which limits itself to only exporting three general purpose APIs on top of which material parameters (and not only parameters)  can be implemented.
- Due to the above, lots of old parameter/textures handling code in the new translator has been removed or simplified.
- Implemented CustomPrimitiveData()
- Factored out some material specific emission rutines like emitting TextureCoordinates and CustomPrimitiveData to its own file MaterialExpressionToMIRCommon.h/cpp so that it can be reused by licensees to build upon Material specific features without moving them inside the translator keeping a nice separation between frontend and backend.
- Renamed MaterialExpressionsIR.cpp to MaterialExpressionsToMIR.cpp to stay consistent with MaterialExpressionToMIRCommon
- Implemented MateiralExpressionCurveAtlasRowParameter
- Moved GetSamplerTypeForTexture and VerifySamplerType utility static functions out of MaterialExpressionTextureBase to the independent MaterialExpressionsUtils.h since they're really not related to material expressions and are used by both translators and others.

#rnx
#Synced-CL 51416398

#Synced-CL 51416398

diff --git a/Engine/Plugins/Editor/UVEditor/Source/UVEditor/Private/UVEditorBackgroundPreview.cpp b/Engine/Plugins/Editor/UVEditor/Source/UVEditor/Private/UVEditorBackgroundPreview.cpp
--- a/Engine/Plugins/Editor/UVEditor/Source/UVEditor/Private/UVEditorBackgroundPreview.cpp
+++ b/Engine/Plugins/Editor/UVEditor/Source/UVEditor/Private/UVEditorBackgroundPreview.cpp
@@ -10,7 +10,7 @@
 #include "Async/Async.h"
 #include "UVEditorUXSettings.h"
 #include "UDIMUtilities.h"
-#include "Materials/MaterialExpressionTextureBase.h"
+#include "Materials/MaterialExpressionUtils.h"
 
 #include UE_INLINE_GENERATED_CPP_BY_NAME(UVEditorBackgroundPreview)
 
@@ -101,7 +101,7 @@ void UUVEditorBackgroundPreview::UpdateBackground()
 					BackgroundMaterial = UMaterialInstanceDynamic::Create(VirtualTextureMaterial, this);
 				}
 
-				EMaterialSamplerType SamplerRequiredForTexture = UMaterialExpressionTextureBase::GetSamplerTypeForTexture(Settings->SourceTexture.Get());
+				EMaterialSamplerType SamplerRequiredForTexture = MaterialExpressionUtils::GetSamplerTypeForTexture(Settings->SourceTexture.Get());
 
 				switch (SamplerRequiredForTexture)
 				{
diff --git a/Engine/Plugins/Enterprise/DatasmithImporter/Source/DatasmithImporter/Private/DatasmithMaterialExpressions.cpp b/Engine/Plugins/Enterprise/DatasmithImporter/Source/DatasmithImporter/Private/DatasmithMaterialExpressions.cpp
--- a/Engine/Plugins/Enterprise/DatasmithImporter/Source/DatasmithImporter/Private/DatasmithMaterialExpressions.cpp
+++ b/Engine/Plugins/Enterprise/DatasmithImporter/Source/DatasmithImporter/Private/DatasmithMaterialExpressions.cpp
@@ -63,6 +63,7 @@
 #include "Materials/MaterialExpressionWorldPosition.h"
 #include "Materials/MaterialFunction.h"
 #include "Materials/MaterialInstanceConstant.h"
+#include "Materials/MaterialExpressionUtils.h"
 
 #include "Math/UnrealMathUtility.h"
 #include "Misc/FileHelper.h"
@@ -1145,7 +1146,7 @@ UMaterialExpression* FDatasmithMaterialExpressions::AddTextureExpression(UTextur
 
 		ConnectToSlot(ToBeConnected, UnrealMaterial, Slot, UV.OutputChannel);
 
-		TextureExpression->SamplerType = UMaterialExpressionTextureBase::GetSamplerTypeForTexture( Texture );
+		TextureExpression->SamplerType = MaterialExpressionUtils::GetSamplerTypeForTexture( Texture );
 	}
 
 	return ToBeConnected;
@@ -2801,7 +2802,7 @@ UMaterialExpression* FDatasmithMaterialExpressions::CreateTextureExpression( IDa
 	{
 		TextureExpression->Group = DatasmithTextureExpression.GetGroupName();
 		TextureExpression->Texture = Texture;
-		TextureExpression->SamplerType = UMaterialExpressionTextureBase::GetSamplerTypeForTexture( Texture );
+		TextureExpression->SamplerType = MaterialExpressionUtils::GetSamplerTypeForTexture( Texture );
 	}
 
 	return TextureExpression;
diff --git a/Engine/Plugins/Enterprise/MDLImporter/Source/MDLImporter/Private/generator/GeneratorMaterialExpressions.cpp b/Engine/Plugins/Enterprise/MDLImporter/Source/MDLImporter/Private/generator/GeneratorMaterialExpressions.cpp
--- a/Engine/Plugins/Enterprise/MDLImporter/Source/MDLImporter/Private/generator/GeneratorMaterialExpressions.cpp
+++ b/Engine/Plugins/Enterprise/MDLImporter/Source/MDLImporter/Private/generator/GeneratorMaterialExpressions.cpp
@@ -71,6 +71,7 @@
 #include "Materials/MaterialExpressionVectorParameter.h"
 #include "Materials/MaterialExpressionVertexNormalWS.h"
 #include "Materials/MaterialExpressionWorldPosition.h"
+#include "Materials/MaterialExpressionUtils.h"
 #include "generator/MaterialExpressionConnection.h"
 
 namespace Generator
@@ -1627,7 +1628,7 @@ namespace Generator
 		if (Texture)
 		{
 			Expression->Texture = Texture;
-			Expression->SamplerType = UMaterialExpressionTextureBase::GetSamplerTypeForTexture(Texture);
+			Expression->SamplerType = MaterialExpressionUtils::GetSamplerTypeForTexture(Texture);
 		}
 
 		return Expression;
diff --git a/Engine/Plugins/Experimental/VirtualHeightfieldMesh/Source/VirtualHeightfieldMesh/Private/HeightfieldMinMaxTextureMaterialExpression.cpp b/Engine/Plugins/Experimental/VirtualHeightfieldMesh/Source/VirtualHeightfieldMesh/Private/HeightfieldMinMaxTextureMaterialExpression.cpp
--- a/Engine/Plugins/Experimental/VirtualHeightfieldMesh/Source/VirtualHeightfieldMesh/Private/HeightfieldMinMaxTextureMaterialExpression.cpp
+++ b/Engine/Plugins/Experimental/VirtualHeightfieldMesh/Source/VirtualHeightfieldMesh/Private/HeightfieldMinMaxTextureMaterialExpression.cpp
@@ -6,7 +6,7 @@
 #include "Engine/Texture2D.h"
 #include "HeightfieldMinMaxTexture.h"
 #include "MaterialCompiler.h"
-#include "Materials/MaterialExpressionTextureBase.h"
+#include "Materials/MaterialExpressionUtils.h"
 
 #include UE_INLINE_GENERATED_CPP_BY_NAME(HeightfieldMinMaxTextureMaterialExpression)
 
@@ -39,7 +39,7 @@ void UMaterialExpressionHeightfieldMinMaxTexture::PostEditChangeProperty(FProper
 	{
 		if (MinMaxTexture != nullptr && MinMaxTexture->Texture != nullptr)
 		{
-			SamplerType = UMaterialExpressionTextureBase::GetSamplerTypeForTexture(MinMaxTexture->Texture);
+			SamplerType = MaterialExpressionUtils::GetSamplerTypeForTexture(MinMaxTexture->Texture);
 
 			FEditorSupportDelegates::ForcePropertyWindowRebuild.Broadcast(this);
 		}
diff --git a/Engine/Plugins/Interchange/Runtime/Source/Import/Private/MaterialX/MaterialExpressions/MaterialExpressionTextureSampleParameterBlur.cpp b/Engine/Plugins/Interchange/Runtime/Source/Import/Private/MaterialX/MaterialExpressions/MaterialExpressionTextureSampleParameterBlur.cpp
--- a/Engine/Plugins/Interchange/Runtime/Source/Import/Private/MaterialX/MaterialExpressions/MaterialExpressionTextureSampleParameterBlur.cpp
+++ b/Engine/Plugins/Interchange/Runtime/Source/Import/Private/MaterialX/MaterialExpressions/MaterialExpressionTextureSampleParameterBlur.cpp
@@ -7,6 +7,7 @@
 #include "Materials/MaterialExpressionCustom.h"
 #include "Materials/MaterialExpressionMultiply.h"
 #include "Materials/MaterialExpressionTextureCoordinate.h"
+#include "Materials/MaterialExpressionUtils.h"
 #include "MaterialCompiler.h"
 
 #include UE_INLINE_GENERATED_CPP_BY_NAME(MaterialExpressionTextureSampleParameterBlur)
@@ -129,7 +130,7 @@ int32 UMaterialExpressionMaterialXTextureSampleParameterBlur::Compile(FMaterialC
 		return CompilerError(Compiler, *ErrorMessage);
 	}
 
-	if(FString SamplerTypeError; !VerifySamplerType(Compiler->GetShaderPlatform(), Compiler->GetTargetPlatform(), Texture, SamplerType, SamplerTypeError))
+	if(FString SamplerTypeError; !MaterialExpressionUtils::VerifySamplerType(Compiler->GetShaderPlatform(), Compiler->GetTargetPlatform(), Texture, SamplerType, SamplerTypeError))
 	{
 		return Compiler->Errorf(TEXT("%s"), *SamplerTypeError);
 	}
diff --git a/Engine/Plugins/Runtime/USDCore/Source/USDUtilities/Private/USDShadeConversion.cpp b/Engine/Plugins/Runtime/USDCore/Source/USDUtilities/Private/USDShadeConversion.cpp
--- a/Engine/Plugins/Runtime/USDCore/Source/USDUtilities/Private/USDShadeConversion.cpp
+++ b/Engine/Plugins/Runtime/USDCore/Source/USDUtilities/Private/USDShadeConversion.cpp
@@ -34,6 +34,7 @@
 #include "Materials/MaterialExpressionVertexColor.h"
 #include "Materials/MaterialInstanceConstant.h"
 #include "Materials/MaterialInstanceDynamic.h"
+#include "Materials/MaterialExpressionUtils.h"
 #include "MaterialShared.h"
 #include "Math/TransformCalculus2D.h"
 #include "Misc/FileHelper.h"
@@ -1101,7 +1102,7 @@ namespace UE::UsdShadeConversion::Private
 				UMaterialEditingLibrary::CreateMaterialExpression(&Material, UMaterialExpressionTextureSample::StaticClass())
 			);
 			Expression->Texture = TextureValue.Texture;
-			Expression->SamplerType = UMaterialExpressionTextureBase::GetSamplerTypeForTexture(TextureValue.Texture);
+			Expression->SamplerType = MaterialExpressionUtils::GetSamplerTypeForTexture(TextureValue.Texture);
 
 			return Expression;
 		}
diff --git a/Engine/Source/Editor/MaterialEditor/Private/MaterialEditor.cpp b/Engine/Source/Editor/MaterialEditor/Private/MaterialEditor.cpp
--- a/Engine/Source/Editor/MaterialEditor/Private/MaterialEditor.cpp
+++ b/Engine/Source/Editor/MaterialEditor/Private/MaterialEditor.cpp
@@ -164,6 +164,7 @@
 #include "Materials/MaterialExpressionNamedReroute.h"
 #include "Materials/MaterialExpressionReroute.h"
 #include "Materials/MaterialExpressionSubstrate.h"
+#include "Materials/MaterialExpressionUtils.h"
 
 #include "MaterialStats.h"
 #include "MaterialEditorTabs.h"
@@ -8034,7 +8035,7 @@ void FMaterialEditor::FixupMismatchedSamplerExpressions()
 		if (UMaterialExpressionTextureBase* TextureBase = Cast<UMaterialExpressionTextureBase>(MaterialExpression))
 		{
 			FString SamplerTypeError;
-			if (!UMaterialExpressionTextureBase::VerifySamplerType(GEditor->PreviewPlatform.GetShaderPlatform(), nullptr, TextureBase->Texture, TextureBase->SamplerType, SamplerTypeError))
+			if (!MaterialExpressionUtils::VerifySamplerType(GEditor->PreviewPlatform.GetShaderPlatform(), nullptr, TextureBase->Texture, TextureBase->SamplerType, SamplerTypeError))
 			{
 				ExpressionsToFix.Add(TextureBase);
 			}
diff --git a/Engine/Source/Editor/MaterialEditor/Private/MaterialInstanceEditor.cpp b/Engine/Source/Editor/MaterialEditor/Private/MaterialInstanceEditor.cpp
--- a/Engine/Source/Editor/MaterialEditor/Private/MaterialInstanceEditor.cpp
+++ b/Engine/Source/Editor/MaterialEditor/Private/MaterialInstanceEditor.cpp
@@ -64,6 +64,7 @@
 #include "Materials/MaterialFunctionMaterialLayer.h"
 #include "Materials/MaterialFunctionMaterialLayerBlend.h"
 #include "Materials/MaterialInstanceBasePropertyOverrides.h"
+#include "Materials/MaterialExpressionUtils.h"
 #include "VT/RuntimeVirtualTexture.h"
 #include "SparseVolumeTexture/SparseVolumeTexture.h"
 #include "Widgets/Input/SButton.h"
@@ -1532,7 +1533,7 @@ void FMaterialInstanceEditor::RefreshOnScreenTextureWarnings()
 						MaterialEditorInstance->SourceInstance->GetTextureParameterValue(TextureParameterValue->ParameterInfo, Texture);
 						if (Texture)
 						{
-							EMaterialSamplerType SamplerType = UMaterialExpressionTextureBase::GetSamplerTypeForTexture(Texture);
+							EMaterialSamplerType SamplerType = MaterialExpressionUtils::GetSamplerTypeForTexture(Texture);
 							UMaterialExpressionTextureSampleParameter* Expression = nullptr;
 
 							TObjectPtr<UMaterialExpression>* MapExpressions = MaterialGuids.Find(TextureParameterValue->ExpressionId);
diff --git a/Engine/Source/Runtime/Engine/Private/Materials/HLSLMaterialTranslator.cpp b/Engine/Source/Runtime/Engine/Private/Materials/HLSLMaterialTranslator.cpp
--- a/Engine/Source/Runtime/Engine/Private/Materials/HLSLMaterialTranslator.cpp
+++ b/Engine/Source/Runtime/Engine/Private/Materials/HLSLMaterialTranslator.cpp
@@ -544,6 +544,20 @@ static inline const TCHAR * GetFloatZeroVector(uint32 NumComponents)
 	}
 }
 
+static FString FormatUnsupportedMaterialDomainError(const FMaterialExternalCodeDeclaration& InExternalCode, const FName& AssetPathName)
+{
+	TStringBuilder<1024> SupportedMaterialDomainsString;
+	for (EMaterialDomain Domain : InExternalCode.Domains)
+	{
+		if (SupportedMaterialDomainsString.Len() > 0)
+		{
+			SupportedMaterialDomainsString.Append(TEXT(", "));
+		}
+		SupportedMaterialDomainsString.Append(*MaterialDomainString(Domain));
+	}
+	return FString::Printf(TEXT("External code '%s' (asset: %s) is only available in the following material domains: %s"), *InExternalCode.Name.ToString(), *AssetPathName.ToString(), *SupportedMaterialDomainsString);
+}
+
 void FHLSLMaterialTranslator::AppendVersion(FShaderKeyGenerator& KeyGen, EShaderPlatform Platform)
 {
 	static const FGuid DDCVersion = FDevSystemGuids::GetSystemGuid(FDevSystemGuids::Get().MaterialTranslationDDCVersion);
@@ -3576,7 +3590,7 @@ int32 FHLSLMaterialTranslator::AddExternalCodeChunk(const FMaterialExternalCodeD
 		checkf(Material != nullptr, TEXT("Missing material while emitting external HLSL code '%s'"), *InExternalCode.Name.ToString());
 		if (InExternalCode.Domains.Find(Material->GetMaterialDomain()) == INDEX_NONE)
 		{
-			return Error(*MaterialExpressionUtils::FormatUnsupportedMaterialDomainError(InExternalCode, Material->GetAssetPath()));
+			return Error(*FormatUnsupportedMaterialDomainError(InExternalCode, Material->GetAssetPath()));
 		}
 	}
 
diff --git a/Engine/Source/Runtime/Engine/Private/Materials/MaterialExpressions.cpp b/Engine/Source/Runtime/Engine/Private/Materials/MaterialExpressions.cpp
--- a/Engine/Source/Runtime/Engine/Private/Materials/MaterialExpressions.cpp
+++ b/Engine/Source/Runtime/Engine/Private/Materials/MaterialExpressions.cpp
@@ -36,6 +36,7 @@
 #include "Materials/MaterialFunctionInstance.h"
 #include "Materials/MaterialInstanceSupport.h"
 #include "Materials/MaterialParameterCollection.h"
+#include "Materials/MaterialExpressionUtils.h"
 #include "Engine/TextureCube.h"
 #include "Engine/Texture2DArray.h"
 #include "Engine/TextureCubeArray.h"
@@ -2539,45 +2540,15 @@ void UMaterialExpressionTextureBase::AutoSetSampleType()
 {
 	if (Texture)
 	{
-		SamplerType = GetSamplerTypeForTexture(Texture);
+		SamplerType = MaterialExpressionUtils::GetSamplerTypeForTexture(Texture);
 	}
 }
 
 EMaterialSamplerType UMaterialExpressionTextureBase::GetSamplerTypeForTexture(const UTexture* Texture, bool ForceNoVT)
 {
-	if (Texture)
-	{
-		if (Texture->GetMaterialType() == MCT_TextureExternal)
-		{
-			return SAMPLERTYPE_External;
-		}
-		else if (Texture->LODGroup == TEXTUREGROUP_8BitData || Texture->LODGroup == TEXTUREGROUP_16BitData)
-		{
-			return SAMPLERTYPE_Data;
-		}
-			
-		const bool bVirtual = ForceNoVT ? false : Texture->GetMaterialType() == MCT_TextureVirtual;
-
-		switch (Texture->CompressionSettings)
-		{
-			case TC_Normalmap:
-				return bVirtual ? SAMPLERTYPE_VirtualNormal : SAMPLERTYPE_Normal;
-			case TC_Grayscale:
-				return Texture->SRGB	? (bVirtual ?  SAMPLERTYPE_VirtualGrayscale : SAMPLERTYPE_Grayscale)
-										: (bVirtual ? SAMPLERTYPE_VirtualLinearGrayscale : SAMPLERTYPE_LinearGrayscale);
-			case TC_Alpha:
-				return bVirtual ?  SAMPLERTYPE_VirtualAlpha : SAMPLERTYPE_Alpha;
-			case TC_Masks:
-				return bVirtual ?  SAMPLERTYPE_VirtualMasks : SAMPLERTYPE_Masks;
-			case TC_DistanceFieldFont:
-				return SAMPLERTYPE_DistanceFieldFont;
-			default:
-				return Texture->SRGB	? (bVirtual ? SAMPLERTYPE_VirtualColor : SAMPLERTYPE_Color) 
-										: (bVirtual ? SAMPLERTYPE_VirtualLinearColor : SAMPLERTYPE_LinearColor);
-		}
-	}
-	return SAMPLERTYPE_Color;
+	return MaterialExpressionUtils::GetSamplerTypeForTexture(Texture, ForceNoVT);
 }
+
 #endif // WITH_EDITOR
 
 UObject* UMaterialExpressionTextureBase::GetReferencedTexture() const
@@ -2775,65 +2746,14 @@ FName UMaterialExpressionTextureSample::GetInputName(int32 InputIndex) const
 }
 #undef IF_INPUT_RETURN
 
-bool UMaterialExpressionTextureBase::VerifySamplerType(
-	const FString& TexturePathName,
-	EMaterialSamplerType CorrectSamplerType,
-	bool bSRGB,
-	EMaterialSamplerType SamplerType,
-	FString& OutErrorMessage)
+bool UMaterialExpressionTextureBase::VerifySamplerType(const FString& TexturePathName, EMaterialSamplerType CorrectSamplerType, bool bSRGB, EMaterialSamplerType SamplerType, FString& OutErrorMessage)
 {
-	if (SamplerType != CorrectSamplerType)
-	{
-		UEnum* SamplerTypeEnum = UMaterialInterface::GetSamplerTypeEnum();
-		check(SamplerTypeEnum);
-
-		const FString SamplerTypeDisplayName = SamplerTypeEnum->GetDisplayNameTextByValue(SamplerType).ToString();
-		const FString TextureTypeDisplayName = SamplerTypeEnum->GetDisplayNameTextByValue(CorrectSamplerType).ToString();
-
-		OutErrorMessage = FString::Printf(TEXT("Sampler type is %s, should be %s for %s"),
-			*SamplerTypeDisplayName,
-			*TextureTypeDisplayName,
-			*TexturePathName);
-
-		return false;
-	}
-
-	if ((SamplerType == SAMPLERTYPE_Normal || SamplerType == SAMPLERTYPE_Masks) && bSRGB)
-	{
-		UEnum* SamplerTypeEnum = UMaterialInterface::GetSamplerTypeEnum();
-		check(SamplerTypeEnum);
-
-		const FString SamplerTypeDisplayName = SamplerTypeEnum->GetDisplayNameTextByValue(SamplerType).ToString();
-
-		OutErrorMessage = FString::Printf(TEXT("To use '%s' as sampler type, SRGB must be disabled for %s"),
-			*SamplerTypeDisplayName,
-			*TexturePathName);
-
-		return false;
-	}
-
-	return true;
+	return MaterialExpressionUtils::VerifySamplerType(TexturePathName, CorrectSamplerType, bSRGB, SamplerType, OutErrorMessage);
 }
 
-bool UMaterialExpressionTextureBase::VerifySamplerType(
-	EShaderPlatform ShaderPlatform,
-	const ITargetPlatform* TargetPlatform,
-	const UTexture* Texture,
-	EMaterialSamplerType SamplerType,
-	FString& OutErrorMessage)
+bool UMaterialExpressionTextureBase::VerifySamplerType(EShaderPlatform ShaderPlatform, const ITargetPlatform* TargetPlatform, const UTexture* Texture, EMaterialSamplerType SamplerType, FString& OutErrorMessage)
 {
-	if ( Texture )
-	{
-		EMaterialSamplerType CorrectSamplerType = UMaterialExpressionTextureBase::GetSamplerTypeForTexture( Texture );
-		bool bIsVirtualTextured = IsVirtualSamplerType(SamplerType);
-		if (bIsVirtualTextured && !UseVirtualTexturing(ShaderPlatform))
-		{
-			SamplerType = UMaterialExpressionTextureBase::GetSamplerTypeForTexture(Texture, !bIsVirtualTextured);
-		}
-		
-		return VerifySamplerType(Texture->GetPathName(), CorrectSamplerType, Texture->SRGB, SamplerType, OutErrorMessage);
-	}
-	return true;
+	return MaterialExpressionUtils::VerifySamplerType(ShaderPlatform, TargetPlatform, Texture, SamplerType, OutErrorMessage);
 }
 
 int32 UMaterialExpressionTextureSample::Compile(class FMaterialCompiler* Compiler, int32 OutputIndex)
@@ -2963,7 +2883,7 @@ int32 UMaterialExpressionTextureSample::Compile(class FMaterialCompiler* Compile
 		}
 
 		FString SamplerTypeError;
-		if (EffectiveTexture && VerifySamplerType(Compiler->GetShaderPlatform(), Compiler->GetTargetPlatform(), EffectiveTexture, EffectiveSamplerType, SamplerTypeError))
+		if (EffectiveTexture && MaterialExpressionUtils::VerifySamplerType(Compiler->GetShaderPlatform(), Compiler->GetTargetPlatform(), EffectiveTexture, EffectiveSamplerType, SamplerTypeError))
 		{
 			if (TOptional<int32> MissingError = CheckForMissingUVWInput(TextureCodeIndex))
 			{
@@ -4083,7 +4003,7 @@ int32 UMaterialExpressionTextureSampleParameter::Compile(class FMaterialCompiler
 	}
 
 	FString SamplerTypeError;
-	if (!VerifySamplerType(Compiler->GetShaderPlatform(), Compiler->GetTargetPlatform(), Texture, SamplerType, SamplerTypeError))
+	if (!MaterialExpressionUtils::VerifySamplerType(Compiler->GetShaderPlatform(), Compiler->GetTargetPlatform(), Texture, SamplerType, SamplerTypeError))
 	{
 		return Compiler->Errorf(TEXT("%s"), *SamplerTypeError);
 	}
@@ -4300,7 +4220,7 @@ int32 UMaterialExpressionTextureObject::CompilePreview(class FMaterialCompiler*
 		return CompilerError(Compiler, TEXT("Requires valid texture"));
 	}
 
-	return CompileTextureSample(Compiler, Texture, Compiler->TextureCoordinate(0, false, false), UMaterialExpressionTextureBase::GetSamplerTypeForTexture( Texture ));
+	return CompileTextureSample(Compiler, Texture, Compiler->TextureCoordinate(0, false, false), MaterialExpressionUtils::GetSamplerTypeForTexture( Texture ));
 }
 
 EMaterialValueType UMaterialExpressionTextureObject::GetOutputValueType(int32 OutputIndex)
@@ -4749,7 +4669,7 @@ int32 UMaterialExpressionTextureSampleParameterSubUV::Compile(class FMaterialCom
 	}
 
 	FString SamplerTypeError;
-	if (!VerifySamplerType(Compiler->GetShaderPlatform(), Compiler->GetTargetPlatform(), Texture, SamplerType, SamplerTypeError))
+	if (! MaterialExpressionUtils::VerifySamplerType(Compiler->GetShaderPlatform(), Compiler->GetTargetPlatform(), Texture, SamplerType, SamplerTypeError))
 	{
 		return Compiler->Errorf(TEXT("%s"), *SamplerTypeError);
 	}
@@ -10329,7 +10249,7 @@ int32 UMaterialExpressionParticleSubUV::Compile(class FMaterialCompiler* Compile
 	if (TextureToCompile)
 	{
 		FString SamplerTypeError;
-		if (!VerifySamplerType(Compiler->GetShaderPlatform(), Compiler->GetTargetPlatform(), TextureToCompile, SamplerTypeToUse, SamplerTypeError))
+		if (!MaterialExpressionUtils::VerifySamplerType(Compiler->GetShaderPlatform(), Compiler->GetTargetPlatform(), TextureToCompile, SamplerTypeToUse, SamplerTypeError))
 		{
 			return Compiler->Errorf(TEXT("%s"), *SamplerTypeError);
 		}
@@ -12089,7 +12009,7 @@ TTuple<bool, UTexture*, EMaterialSamplerType, FString> UMaterialExpressionFontSa
 	}
 
 	FString SamplerTypeError;
-	if (!UMaterialExpressionTextureBase::VerifySamplerType(ShaderPlatform, TargetPlatform, Texture, ExpectedSamplerType, SamplerTypeError))
+	if (!MaterialExpressionUtils::VerifySamplerType(ShaderPlatform, TargetPlatform, Texture, ExpectedSamplerType, SamplerTypeError))
 	{
 		return { false, nullptr, SAMPLERTYPE_MAX, SamplerTypeError };
 	}
@@ -19620,7 +19540,7 @@ int32 UMaterialExpressionAntialiasedTextureMask::Compile(class FMaterialCompiler
 	}
 
 	FString SamplerTypeError;
-	if (!VerifySamplerType(Compiler->GetShaderPlatform(), Compiler->GetTargetPlatform(), Texture, SamplerType, SamplerTypeError))
+	if (! MaterialExpressionUtils::VerifySamplerType(Compiler->GetShaderPlatform(), Compiler->GetTargetPlatform(), Texture, SamplerType, SamplerTypeError))
 	{
 		return Compiler->Errorf(TEXT("%s"), *SamplerTypeError);
 	}
diff --git a/Engine/Source/Runtime/Engine/Private/Materials/MaterialIR.cpp b/Engine/Source/Runtime/Engine/Private/Materials/MaterialIR.cpp
--- a/Engine/Source/Runtime/Engine/Private/Materials/MaterialIR.cpp
+++ b/Engine/Source/Runtime/Engine/Private/Materials/MaterialIR.cpp
@@ -24,9 +24,9 @@ const TCHAR* LexToString(EValueKind Kind)
 		case VK_Builtin: return TEXT("Builtin");
 		case VK_MaterialParameterCollection: return TEXT("MaterialParameterCollection");
 		case VK_ShadingModel: return TEXT("ShadingModel");
-		case VK_TextureObject: return TEXT("TextureObject");
-		case VK_RuntimeVirtualTextureObject: return TEXT("RuntimeVirtualTextureObject");
-		case VK_UniformParameter: return TEXT("UniformParameter");
+		case VK_PrimitiveUniform: return TEXT("PrimitiveUniform");
+		case VK_TextureUniform: return TEXT("TextureUniform");
+		case VK_VirtualTextureUniform: return TEXT("VirtualTextureUniform");
 		case VK_GetVertexInterpolator: return TEXT("GetVertexInterpolator");
 		case VK_Nop: return TEXT("Nop");
 		case VK_Composite: return TEXT("Composite");
@@ -103,9 +103,9 @@ uint32 FValue::GetSizeInBytes() const
 		case VK_Builtin: return sizeof(FBuiltin);
 		case VK_MaterialParameterCollection: return sizeof(FMaterialParameterCollection);
 		case VK_ShadingModel: return sizeof(FShadingModel);
-		case VK_TextureObject: return sizeof(FTextureObject);
-		case VK_RuntimeVirtualTextureObject: return sizeof(FRuntimeVirtualTextureObject);
-		case VK_UniformParameter: return sizeof(FUniformParameter);
+		case VK_PrimitiveUniform: return sizeof(FPrimitiveUniform);
+		case VK_TextureUniform: return sizeof(FTextureUniform);
+		case VK_VirtualTextureUniform: return sizeof(FVirtualTextureUniform);
 		case VK_GetVertexInterpolator: return sizeof(FGetVertexInterpolator);
 		case VK_Nop: return sizeof(FNop);
 		case VK_Composite: return sizeof(FComposite) + sizeof(FValue*) * static_cast<const FComposite*>(this)->GetComponents().Num();
@@ -197,13 +197,13 @@ TConstArrayView<FValue*> FValue::GetUses() const
 		case VK_TextureRead:
 		{
 			auto This = static_cast<const FTextureRead*>(this);
-			return { &This->TextureObject, FTextureRead::NumStaticUses };
+			return { &This->TextureUniform, FTextureRead::NumStaticUses };
 		}
 
 		case VK_VTPageTableRead:
 		{
 			auto This = static_cast<const FVTPageTableRead*>(this);
-			return { &This->TextureObject, FVTPageTableRead::NumStaticUses };
+			return { &This->VirtualTextureUniform, FVTPageTableRead::NumStaticUses };
 		}
 
 		case VK_Extern:
@@ -455,6 +455,19 @@ bool FValue::AreAllNearlyOne() const
 	return false;
 }
 
+bool FValue::IsConstant() const
+{
+	if (As<FConstant>())
+	{
+		return true;
+	}
+	if (const FComposite* AsComposite = As<FComposite>())
+	{
+		return AsComposite->AreComponentsConstant();
+	}
+	return false;
+}
+
 bool FValue::Equals(const FValue* Other) const
 {
 	if (this == Other)
@@ -481,8 +494,8 @@ bool FValue::EqualsConstant(float TestValue) const
 	{
 		switch (ValueConstant->Type.GetPrimitive().ScalarKind)
 		{
-			case MIR::EScalarKind::Bool:	return ValueConstant->Boolean == !!TestValue;
-			case MIR::EScalarKind::Int:		return ValueConstant->Integer == (TInteger)TestValue;
+			case MIR::EScalarKind::Boolean:	return ValueConstant->Boolean == !!TestValue;
+			case MIR::EScalarKind::Integer:	return ValueConstant->Integer == (TInteger)TestValue;
 			case MIR::EScalarKind::Float:	return ValueConstant->Float == TestValue;
 			case MIR::EScalarKind::Double:	return ValueConstant->Double == TestValue;
 			default: UE_MIR_UNREACHABLE();
@@ -519,30 +532,30 @@ bool FValue::EqualsConstant(FVector4f TestValue) const
 
 UObject* FValue::GetTextureObject() const
 {
-	if (const FTextureObject* TextureObject = As<FTextureObject>())
+	if (const FTextureUniform* TextureObject = As<FTextureUniform>())
 	{
 		return TextureObject->Texture;
 	}
-	if (const FRuntimeVirtualTextureObject* RVTextureObject = As<FRuntimeVirtualTextureObject>())
+	if (const FVirtualTextureUniform* RVTextureObject = As<FVirtualTextureUniform>())
 	{
-		return RVTextureObject->RVTexture;
+		return RVTextureObject->VirtualTexture;
 	}
 	return nullptr;
 }
 
-int32 FValue::GetUniformParameterIndex() const
+int32 FValue::GetUniformIndex() const
 {
-	if (const MIR::FTextureObject* TextureObject = As<MIR::FTextureObject>())
+	if (const MIR::FPrimitiveUniform* Uniform = As<MIR::FPrimitiveUniform>())
 	{
-		return TextureObject->Analysis_UniformParameterIndex;
+		return Uniform->Analysis_UniformIndex;
 	}
-	if (const MIR::FRuntimeVirtualTextureObject* RVTextureObject = As<MIR::FRuntimeVirtualTextureObject>())
+	if (const MIR::FTextureUniform* Uniform = As<MIR::FTextureUniform>())
 	{
-		return RVTextureObject->Analysis_UniformParameterIndex;
+		return Uniform->Analysis_UniformIndex;
 	}
-	if (const MIR::FUniformParameter* UniformParameter = As<MIR::FUniformParameter>())
+	if (const MIR::FVirtualTextureUniform* Uniform = As<MIR::FVirtualTextureUniform>())
 	{
-		return UniformParameter->Analysis_UniformParameterIndex;
+		return Uniform->Analysis_UniformIndex;
 	}
 	return INDEX_NONE;
 }
diff --git a/Engine/Source/Runtime/Engine/Private/Materials/MaterialIRDebug.cpp b/Engine/Source/Runtime/Engine/Private/Materials/MaterialIRDebug.cpp
--- a/Engine/Source/Runtime/Engine/Private/Materials/MaterialIRDebug.cpp
+++ b/Engine/Source/Runtime/Engine/Private/Materials/MaterialIRDebug.cpp
@@ -117,8 +117,8 @@ struct FDebugDumpIRUseGraphState
 		{
 			switch (Constant->Type.GetPrimitive().ScalarKind)
 			{
-				case EScalarKind::Bool:  Out.Append(Constant->Boolean ? TEXT("true") : TEXT("false")); break;
-				case EScalarKind::Int:	 Out.Appendf(TEXT("%lld"), Constant->Integer); break;
+				case EScalarKind::Boolean:  Out.Append(Constant->Boolean ? TEXT("true") : TEXT("false")); break;
+				case EScalarKind::Integer:	 Out.Appendf(TEXT("%lld"), Constant->Integer); break;
 				case EScalarKind::Float: Out.Appendf(TEXT("%f"), Constant->Float); break;
 				case EScalarKind::Double:   Out.Appendf(TEXT("%f"), Constant->Double); break;
 				default: UE_MIR_UNREACHABLE();
@@ -225,7 +225,7 @@ struct FDebugDumpIRState
 	uint32 InstrIdCounter{};
 
 	// Array of encountered parameters. Used later on to generate the a recap of all referenced parameters.
-	TArray<TPair<uint32, const FValue*>> ReferencedParameters{};
+	TArray<TPair<uint32, const FValue*>> ReferencedUniforms{};
 	
 	// Array of encountered user HLSL functions.
 	TArray<const MIR::FFunctionHLSL*> ReferencedFunctionHLSLs{};
@@ -410,16 +410,20 @@ struct FDebugDumpIRState
 		{
 			Out.Appendf(TEXT("[ExternalInput \"%s\"]"), LexToString(Value->As<FBuiltin>()->Id));
 		}
-		else if (const FTextureObject* TextureObject = Value->As<FTextureObject>())
+		else if (const FPrimitiveUniform* PrimitiveUniform = Value->As<FPrimitiveUniform>())
 		{
-			Out.Appendf(TEXT("[TextureObject #%d SamplerType=\"%s\"]"), TextureObject->Analysis_UniformParameterIndex, *StaticEnum<EMaterialSamplerType>()->GetDisplayNameTextByValue(TextureObject->SamplerType).ToString());
-			ReferencedParameters.AddUnique({ TextureObject->Analysis_UniformParameterIndex, TextureObject });
+			Out.Appendf(TEXT("[UniformIndex #%d]"), PrimitiveUniform->Analysis_UniformIndex);
+			ReferencedUniforms.AddUnique({ PrimitiveUniform->Analysis_UniformIndex, PrimitiveUniform });
 		}
-		else if (const FUniformParameter* UniformParameter = Value->As<FUniformParameter>())
+		else if (const FTextureUniform* TextureUniform = Value->As<FTextureUniform>())
 		{
-			FName ParameterName = Module->GetParameterInfo(UniformParameter->ParameterIdInModule).Name;
-			Out.Appendf(TEXT("[Parameter #%d \"%s\"]"), UniformParameter->Analysis_UniformParameterIndex, *ParameterName.ToString());
-			ReferencedParameters.AddUnique({ UniformParameter->Analysis_UniformParameterIndex, UniformParameter });
+			Out.Appendf(TEXT("[UniformIndex #%d SamplerType=\"%s\"]"), TextureUniform->Analysis_UniformIndex, *StaticEnum<EMaterialSamplerType>()->GetDisplayNameTextByValue(TextureUniform->SamplerType).ToString());
+			ReferencedUniforms.AddUnique({ TextureUniform->Analysis_UniformIndex, TextureUniform });
+		}
+		else if (const FVirtualTextureUniform* VirtualTextureUniform = Value->As<FVirtualTextureUniform>())
+		{
+			Out.Appendf(TEXT("[UniformIndex #%d]"), VirtualTextureUniform->Analysis_UniformIndex);
+			ReferencedUniforms.AddUnique({ VirtualTextureUniform->Analysis_UniformIndex, VirtualTextureUniform });
 		}
 		else
 		{
@@ -458,36 +462,32 @@ struct FDebugDumpIRState
 	// Prints the parameter recap section.
 	void DumpReferencedParameters()
 	{
-		if (ReferencedParameters.IsEmpty())
+		if (ReferencedUniforms.IsEmpty())
 		{
 			return;
 		}
 
 		// Dump the list of referenced parameters.
-		ReferencedParameters.Sort([](const auto& Param1, const auto& Param2)
+		ReferencedUniforms.Sort([](const auto& Param1, const auto& Param2)
 		{
 			return Param1.Key < Param2.Key;
 		});
 
 		Out.Append(TEXT("\n; Referenced material parameters\n"));
-		for (const auto& ParamPair : ReferencedParameters)
+		for (const auto& ParamPair : ReferencedUniforms)
 		{
 			const FValue* Value = ParamPair.Value;
-			if (const FUniformParameter* Param = Value->As<FUniformParameter>())
+			if (const FPrimitiveUniform* PrimitiveUniform = Value->As<FPrimitiveUniform>())
+			{
+				Out.Appendf(TEXT("#%d = Name=\"%s\" Type=\"%s\"\n"), PrimitiveUniform->Analysis_UniformIndex, *PrimitiveUniform->Name.ToString(), *PrimitiveUniform->Type.GetSpelling());
+			}
+			else if (const FTextureUniform* TextureUniform = Value->As<FTextureUniform>())
 			{
-				FName ParameterName = Module->GetParameterInfo(Param->ParameterIdInModule).Name;
-				EMaterialParameterType ParameterType = Module->GetParameterMetadata(Param->ParameterIdInModule).Value.Type;
-				Out.Appendf(TEXT("#%d = Name=\"%s\" Type=\"%s\"\n"),
-					Param->Analysis_UniformParameterIndex,
-					*ParameterName.ToString(),
-					*MaterialParameterTypeToString(ParameterType));
+				Out.Appendf(TEXT("#%d = Name=\"%s\" Type=\"%s\"\n"), TextureUniform->Analysis_UniformIndex, *TextureUniform->Name.ToString(), TEXT("Texture"));
 			}
-			else if (const FTextureObject* TextureObject = Value->As<FTextureObject>())
+			else if (const FVirtualTextureUniform* VirtualTextureUniform = Value->As<FVirtualTextureUniform>())
 			{
-				Out.Appendf(TEXT("#%d = Name=\"%s\" Type=\"%s\"\n"),
-					TextureObject->Analysis_UniformParameterIndex,
-					*TextureObject->Texture->GetName(),
-					TEXT("Texture"));
+				Out.Appendf(TEXT("#%d = Name=\"%s\" Type=\"%s\"\n"), VirtualTextureUniform->Analysis_UniformIndex, *VirtualTextureUniform->Name.ToString(), TEXT("VirtualTexture"));
 			}
 		}
 	}
diff --git a/Engine/Source/Runtime/Engine/Private/Materials/MaterialIREmitter.cpp b/Engine/Source/Runtime/Engine/Private/Materials/MaterialIREmitter.cpp
--- a/Engine/Source/Runtime/Engine/Private/Materials/MaterialIREmitter.cpp
+++ b/Engine/Source/Runtime/Engine/Private/Materials/MaterialIREmitter.cpp
@@ -2,6 +2,7 @@
 
 #include "Materials/MaterialIREmitter.h"
 #include "Materials/Material.h"
+#include "Materials/MaterialExpressionUtils.h"
 #include "Materials/MaterialIRModule.h"
 #include "Materials/MaterialIRModuleBuilder.h"
 #include "Materials/MaterialIRInternal.h"
@@ -495,15 +496,15 @@ FType FEmitter::GetMaterialAggregateAttributeType(const UMaterialAggregate* Aggr
 
 	switch (Aggregate->Attributes[AttributeIndex].Type)
 	{
-		case EMaterialAggregateAttributeType::Bool1: return FType::MakeVector(EScalarKind::Bool, 1);
-		case EMaterialAggregateAttributeType::Bool2: return FType::MakeVector(EScalarKind::Bool, 2);
-		case EMaterialAggregateAttributeType::Bool3: return FType::MakeVector(EScalarKind::Bool, 3);
-		case EMaterialAggregateAttributeType::Bool4: return FType::MakeVector(EScalarKind::Bool, 4);
+		case EMaterialAggregateAttributeType::Bool1: return FType::MakeVector(EScalarKind::Boolean, 1);
+		case EMaterialAggregateAttributeType::Bool2: return FType::MakeVector(EScalarKind::Boolean, 2);
+		case EMaterialAggregateAttributeType::Bool3: return FType::MakeVector(EScalarKind::Boolean, 3);
+		case EMaterialAggregateAttributeType::Bool4: return FType::MakeVector(EScalarKind::Boolean, 4);
 		case EMaterialAggregateAttributeType::ShadingModel: return FType::MakeIntScalar();
-		case EMaterialAggregateAttributeType::UInt1: return FType::MakeVector(EScalarKind::Int, 1);
-		case EMaterialAggregateAttributeType::UInt2: return FType::MakeVector(EScalarKind::Int, 2);
-		case EMaterialAggregateAttributeType::UInt3: return FType::MakeVector(EScalarKind::Int, 3);
-		case EMaterialAggregateAttributeType::UInt4: return FType::MakeVector(EScalarKind::Int, 4);
+		case EMaterialAggregateAttributeType::UInt1: return FType::MakeVector(EScalarKind::Integer, 1);
+		case EMaterialAggregateAttributeType::UInt2: return FType::MakeVector(EScalarKind::Integer, 2);
+		case EMaterialAggregateAttributeType::UInt3: return FType::MakeVector(EScalarKind::Integer, 3);
+		case EMaterialAggregateAttributeType::UInt4: return FType::MakeVector(EScalarKind::Integer, 4);
 		case EMaterialAggregateAttributeType::Float1: return FType::MakeVector(EScalarKind::Float, 1);
 		case EMaterialAggregateAttributeType::Float2: return FType::MakeVector(EScalarKind::Float, 2);
 		case EMaterialAggregateAttributeType::Float3: return FType::MakeVector(EScalarKind::Float, 3);
@@ -914,8 +915,8 @@ FValueRef FEmitter::ConstantZero(EScalarKind Kind)
 {
 	switch (Kind)
 	{
-		case EScalarKind::Bool: return ConstantFalse();
-		case EScalarKind::Int: return ConstantInt(0);
+		case EScalarKind::Boolean: return ConstantFalse();
+		case EScalarKind::Integer: return ConstantInt(0);
 		case EScalarKind::Float: return ConstantFloat(0.0f);
 		case EScalarKind::Double: return ConstantDouble(0.0);
 		default: UE_MIR_UNREACHABLE();
@@ -926,8 +927,8 @@ FValueRef FEmitter::ConstantOne(EScalarKind Kind)
 {
 	switch (Kind)
 	{
-		case EScalarKind::Bool: return ConstantTrue();
-		case EScalarKind::Int: return ConstantInt(1);
+		case EScalarKind::Boolean: return ConstantTrue();
+		case EScalarKind::Integer: return ConstantInt(1);
 		case EScalarKind::Float: return ConstantFloat(1.0f);
 		case EScalarKind::Double: return ConstantDouble(1.0);
 		default: UE_MIR_UNREACHABLE();
@@ -938,8 +939,8 @@ FValueRef FEmitter::ConstantScalar(EScalarKind Kind, TDouble Value)
 {
 	switch (Kind)
 	{
-		case EScalarKind::Bool: return ConstantBool(Value != 0.0f);
-		case EScalarKind::Int: return ConstantInt(TInteger(Value));
+		case EScalarKind::Boolean: return ConstantBool(Value != 0.0f);
+		case EScalarKind::Integer: return ConstantInt(TInteger(Value));
 		case EScalarKind::Float: return ConstantFloat(TFloat(Value));
 		case EScalarKind::Double: return ConstantDouble(Value);
 		default: UE_MIR_UNREACHABLE();
@@ -987,7 +988,7 @@ FValueRef FEmitter::ConstantBool4(bool InX, bool InY, bool InZ, bool InW)
 
 FValueRef FEmitter::ConstantInt(TInteger InX)
 {
-	FConstant Scalar = MakePrototype<FConstant>(FType::MakeScalar(EScalarKind::Int));
+	FConstant Scalar = MakePrototype<FConstant>(FType::MakeScalar(EScalarKind::Integer));
 	Scalar.Integer = InX;
 	return EmitPrototype(*this, Scalar);
 }
@@ -1054,6 +1055,30 @@ FValueRef FEmitter::ConstantDouble(TDouble InX)
 	return EmitPrototype(*this, Scalar);
 }
 
+FValueRef FEmitter::ConstantDouble2(UE::Math::TVector2<TDouble> InValue)
+{
+	FValueRef X = ConstantDouble(InValue.X);
+	FValueRef Y = ConstantDouble(InValue.Y);
+	return Vector2(X, Y);
+}
+
+FValueRef FEmitter::ConstantDouble3(UE::Math::TVector<TDouble> InValue)
+{
+	FValueRef X = ConstantDouble(InValue.X);
+	FValueRef Y = ConstantDouble(InValue.Y);
+	FValueRef Z = ConstantDouble(InValue.Z);
+	return Vector3(X, Y, Z);
+}
+
+FValueRef FEmitter::ConstantDouble4(UE::Math::TVector4<TDouble> InValue)
+{
+	FValueRef X = ConstantDouble(InValue.X);
+	FValueRef Y = ConstantDouble(InValue.Y);
+	FValueRef Z = ConstantDouble(InValue.Z);
+	FValueRef W = ConstantDouble(InValue.W);
+	return Vector4(X, Y, Z, W);
+}
+
 /*--------------------- Other non-instruction values emission ---------------------*/
 
 FValueRef FEmitter::Poison()
@@ -1088,20 +1113,53 @@ FValueRef FEmitter::MaterialParameterCollection(UMaterialParameterCollection* Co
 	return EmitPrototype(*this, Prototype);
 }
 
-FValueRef FEmitter::TextureObject(UTexture* Texture, EMaterialSamplerType SamplerType)
+FValueRef FEmitter::NamedPrimitiveUniform(FName Name, FValueRef DefaultConstant)
+{
+	if (!DefaultConstant->Type.AsPrimitive())
+	{
+		Error(TEXT("Default value must have primitive type."));
+		return Poison();
+	}
+
+	FPrimitive Primitive = DefaultConstant->Type.GetPrimitive();
+	if (Primitive.IsMatrix())
+	{
+		Error(TEXT("Matrix primitive uniforms are unsupported."));
+		return Poison();
+	}
+
+	if (!DefaultConstant->IsConstant())
+	{
+		Error(TEXT("Default value is expected to be constant."));
+		return Poison();
+	}
+
+	FPrimitiveUniform Proto = MakePrototype<FPrimitiveUniform>(DefaultConstant->Type);
+	Proto.Name = FScriptName{ Name };
+	Proto.DefaultConstant = DefaultConstant; 
+	Proto.Analysis_UniformIndex = -1;
+
+	return EmitPrototype(*this, Proto);
+}
+
+FValueRef FEmitter::TextureUniform(FName Name, UTexture* Texture, EMaterialSamplerType SamplerType)
 {
 	check(Texture);
+
+	// @massimo.tristano move this to ExpressionsIR.cpp
 	FString SamplerTypeError;
-	if (!UMaterialExpressionTextureBase::VerifySamplerType(GetShaderPlatform(), GetTargetPlatform(), Texture, SamplerType, SamplerTypeError))
+	if (!MaterialExpressionUtils::VerifySamplerType(GetShaderPlatform(), GetTargetPlatform(), Texture, SamplerType, SamplerTypeError))
 	{
 		Errorf(TEXT("%s"), *SamplerTypeError);
 		return Poison();
 	}
 
-	FTextureObject Proto = MakePrototype<FTextureObject>(FType::MakeTexture());
+	FTextureUniform Proto = MakePrototype<FTextureUniform>(FType::MakeTexture());
+	Proto.Name = FScriptName{ Name };
 	Proto.Texture = Texture;
 	Proto.SamplerType = SamplerType;
-	Proto.Analysis_UniformParameterIndex = INDEX_NONE;
+	Proto.Analysis_UniformIndex = INDEX_NONE;
+	
 	return EmitPrototype(*this, Proto);
 }
 
@@ -1137,200 +1195,20 @@ static EMaterialSamplerType DemoteSamplerTypeFromVTtoST(EMaterialSamplerType InS
 	return InSamplerType;
 }
 
-static bool IsVirtualTexture(FValueRef Texture)
+FValueRef FEmitter::VirtualTextureUniform(FName Name, URuntimeVirtualTexture* VirtualTexture, int32 VTLayerIndex, int32 VTPageTableIndex)
 {
-	if (FUniformParameter* UniformParameter = Texture->As<FUniformParameter>())
-	{
-		return IsVirtualSamplerType(UniformParameter->SamplerType);
-	}
-	if (FTextureObject* TextureObject = Texture->As<FTextureObject>())
-	{
-		return (TextureObject->Texture->GetMaterialType() & MCT_TextureVirtual) != 0;
-	}
-	if (Texture->IsA(VK_RuntimeVirtualTextureObject))
-	{
-		return true;
-	}
-	return false;
-}
+	check(VirtualTexture);
+	check(VTLayerIndex >= INT16_MIN && VTLayerIndex <= INT16_MAX);
+	check(VTPageTableIndex >= 0 && VTPageTableIndex <= UINT16_MAX);
 
-// Returns this value's texture sampler type if it has one (SAMPLERTYPE_MAX otherwise).
-static EMaterialSamplerType GetValueMaterialSamplerType(FValueRef Value)
-{
-	if (MIR::FTextureObject* TextureObject = MIR::As<MIR::FTextureObject>(Value))
-	{
-		return TextureObject->SamplerType;
-	}
-	if (MIR::FRuntimeVirtualTextureObject* RVTextureObject = MIR::As<MIR::FRuntimeVirtualTextureObject>(Value))
-	{
-		return RVTextureObject->SamplerType;
-	}
-	if (MIR::FUniformParameter* UniformParameter = MIR::As<MIR::FUniformParameter>(Value))
-	{
-		return UniformParameter->SamplerType;
-	}
-	return SAMPLERTYPE_MAX;
-}
-
-static EMaterialSamplerType MapSamplerTypeForTexture(FValueRef InTexture, EMaterialSamplerType InSamplerType)
-{
-	// Can't sample with virtual texturing if input texture is not a virtual texture
-	if (InSamplerType == SAMPLERTYPE_MAX)
-	{
-		InSamplerType = GetValueMaterialSamplerType(InTexture);
-	}
-	return IsVirtualTexture(InTexture) ? PromoteSamplerTypeFromSTtoVT(InSamplerType) : DemoteSamplerTypeFromVTtoST(InSamplerType);
-}
-
-static FValueRef VTPageTableLoadFromSamplerSource(
-	FEmitter& Em, FValueRef Texture, const FTextureSampleAttributes& BaseAttributes, FValueRef TexCoord,
-	FValueRef TexCoordDdx = {}, FValueRef TexCoordDdy = {}, ETextureMipValueMode MipValueMode = TMVM_None, FValueRef MipValue = {})
-{
-	// Cast input texture to UTexture. If it's a URuntimeVirtualTexture, we accept the cast to be null when passed to GetTextureAddressForSamplerSource().
-	TextureAddress StaticAddressX = TA_Wrap;
-	TextureAddress StaticAddressY = TA_Wrap;
-	TextureAddress StaticAddressZ = TA_Wrap;
-	UE::MaterialTranslatorUtils::GetTextureAddressForSamplerSource(Cast<UTexture>(Texture->GetTextureObject()), BaseAttributes.SamplerSourceMode, StaticAddressX, StaticAddressY, StaticAddressZ);
-
-	return Em.VTPageTableLoad(Texture, StaticAddressX, StaticAddressY, TexCoord, TexCoordDdx, TexCoordDdy, BaseAttributes.bEnableFeedback, BaseAttributes.bIsAdaptive, MipValueMode, MipValue);
-}
-
-FValueRef FEmitter::RuntimeVirtualTextureObject(URuntimeVirtualTexture* RVTexture, EMaterialSamplerType SamplerType, int32 VTLayerIndex, int32 VTPageTableIndex)
-{
-	check(RVTexture);
-	check(IsVirtualSamplerType(SamplerType));
-
-	FRuntimeVirtualTextureObject Prototype = MakePrototype<FRuntimeVirtualTextureObject>(FType::MakeRuntimeVirtualTexture());
-	Prototype.RVTexture = RVTexture;
-	Prototype.SamplerType = PromoteSamplerTypeFromSTtoVT(SamplerType);
-	Prototype.VTLayerIndex = VTLayerIndex;
-	Prototype.VTPageTableIndex = VTPageTableIndex;
-	Prototype.Analysis_UniformParameterIndex = INDEX_NONE;
-	return EmitPrototype(*this, Prototype);
-}
+	FVirtualTextureUniform Proto = MakePrototype<FVirtualTextureUniform>(FType::MakeRuntimeVirtualTexture());
+	Proto.Name = FScriptName{ Name };
+	Proto.VirtualTexture = VirtualTexture;
+	Proto.VTLayerIndex = VTLayerIndex;
+	Proto.VTPageTableIndex = VTPageTableIndex;
+	Proto.Analysis_UniformIndex = INDEX_NONE;
 
-FValueRef FEmitter::Parameter(FName Name, FMaterialParameterMetadata& Metadata, EMaterialSamplerType SamplerType, int32 VTLayerIndex, int32 VTPageTableIndex)
-{
-	// Helper local function that register a parameter (info and metadata) to the module, and returns some uint32 ID.
-	auto RegisterParameter = [] (FMaterialIRModule* InModule, FMaterialParameterInfo InInfo, const FMaterialParameterMetadata& InMetadata) -> uint32
-	{
-		uint32 Id;
-		if (MIR::Find(InModule->ParameterInfoToId, InInfo, Id))
-		{
-			check(InModule->ParameterIdToData[Id].Get<FMaterialParameterMetadata>().Value == InMetadata.Value);
-			return Id;
-		}
-
-		Id = InModule->ParameterIdToData.Num();
-		InModule->ParameterInfoToId.Add(InInfo, Id);
-		InModule->ParameterIdToData.Push({ InInfo, InMetadata });
-		return Id;
-	};
-
-	auto MakeUniformParameter = [this, &RegisterParameter](
-		FType Type, FMaterialParameterInfo InInfo, const FMaterialParameterMetadata& InMetadata,
-		EMaterialSamplerType InSamplerType = SAMPLERTYPE_Color, int32 VTLayerIndex = INDEX_NONE, int32 VTPageTableIndex = INDEX_NONE) -> FValueRef
-	{
-		FUniformParameter Prototype = MakePrototype<FUniformParameter>(MoveTemp(Type));
-		Prototype.ParameterIdInModule = RegisterParameter(Module, MoveTemp(InInfo), InMetadata);
-		Prototype.SamplerType = InSamplerType;
-		Prototype.VTLayerIndex = VTLayerIndex;
-		Prototype.VTPageTableIndex = VTPageTableIndex;
-		Prototype.Analysis_UniformParameterIndex = INDEX_NONE;
-		
-		return EmitPrototype(*this, Prototype);
-	};
-
-	FMaterialParameterInfo Info{ Name };
-
-	switch (Metadata.Value.Type)
-	{
-		case EMaterialParameterType::Scalar:
-		{
-			if (Metadata.PrimitiveDataIndex != INDEX_NONE)
-			{
-				return CustomPrimitiveData(Metadata.PrimitiveDataIndex);
-			}
-			return MakeUniformParameter(FType::MakeFloatScalar(), Info, Metadata);
-		}
-
-		case EMaterialParameterType::Vector:
-		case EMaterialParameterType::DoubleVector:
-		{
-			if (Metadata.PrimitiveDataIndex != INDEX_NONE)
-			{
-				FValue* X = CustomPrimitiveData(Metadata.PrimitiveDataIndex + 0);
-				FValue* Y = CustomPrimitiveData(Metadata.PrimitiveDataIndex + 1);
-				FValue* Z = CustomPrimitiveData(Metadata.PrimitiveDataIndex + 2);
-				FValue* W = CustomPrimitiveData(Metadata.PrimitiveDataIndex + 3);
-				return Vector4(X, Y, Z, W);
-			}
-
-			const EScalarKind ScalarKind = Metadata.Value.Type == EMaterialParameterType::Vector ? EScalarKind::Float : EScalarKind::Double;
-			return MakeUniformParameter(FType::MakeVector(ScalarKind, 4), Info, Metadata);
-		}
-
-		case EMaterialParameterType::Texture:
-		case EMaterialParameterType::Font:
-		{
-			return MakeUniformParameter(FType::MakeTexture(), Info, Metadata, SamplerType);
-		}
-
-		case EMaterialParameterType::RuntimeVirtualTexture:
-		{
-			return MakeUniformParameter(FType::MakeRuntimeVirtualTexture(), Info, Metadata, PromoteSamplerTypeFromSTtoVT(SamplerType), VTLayerIndex, VTPageTableIndex);
-		}
-
-		case EMaterialParameterType::StaticSwitch:
-		{
-			// Apply eventual parameter override
-			for (const FStaticSwitchParameter& Param : StaticParameterSet->GetRuntime().StaticSwitchParameters)
-			{
-				if (Param.IsOverride() && Param.ParameterInfo.Name == Name)
-				{
-					Metadata.Value.Bool[0] = Param.Value;
-					break;
-				}
-			}
-			return ConstantBool(Metadata.Value.Bool[0]);
-		}
-
-		case EMaterialParameterType::StaticComponentMask:
-		{
-			// Apply eventual parameter override
-			for (const FStaticComponentMaskParameter& Param : StaticParameterSet->EditorOnly.StaticComponentMaskParameters)
-			{
-				if (Param.IsOverride() && Param.ParameterInfo.Name == Name)
-				{
-					Metadata.Value.Bool[0] = Param.R;
-					Metadata.Value.Bool[1] = Param.G;
-					Metadata.Value.Bool[2] = Param.B;
-					Metadata.Value.Bool[3] = Param.A;
-					break;
-				}
-			}
-
-			return ConstantBool4(Metadata.Value.Bool[0], Metadata.Value.Bool[1], Metadata.Value.Bool[2], Metadata.Value.Bool[3]);
-		}
-
-		default:
-			UE_MIR_TODO();
-	}
-}
-
-UObject* FEmitter::GetTextureFromValue(MIR::FValueRef Texture)
-{
-	if (!Texture)
-	{
-		return nullptr;
-	}
-	// Handle case for uniform parameters here because their metadata is stored in the IR module, which the emitter has access to.
-	if (FUniformParameter* UniformParameter = Texture->As<FUniformParameter>())
-	{
-		const FMaterialParameterMetadata& ParameterMetadata = Module->GetParameterMetadata(UniformParameter->ParameterIdInModule);
-		return ParameterMetadata.Value.AsTextureObject();
-	}
-	return Texture->GetTextureObject();
+	return EmitPrototype(*this, Proto);
 }
 
 // For now, only a small subset of opcodes are supported: TextureSize, TexelSize, and RuntimeVirtualTextureUniform.
@@ -1343,7 +1221,7 @@ FValueRef FEmitter::PreshaderParameter(FType Type, EPreshader2Opcode Opcode, FVa
 		TEXT("Preshader opcode (0x%X) not supported for parameters in new material translator"), (int32)Opcode
 	);
 
-	UObject* SourceParameterTexture = GetTextureFromValue(SourceParameter);
+	UObject* SourceParameterTexture = SourceParameter->GetTextureObject();
 	if (!SourceParameterTexture)
 	{
 		Error(TEXT("Missing default texture from source parameter"));
@@ -1355,6 +1233,7 @@ FValueRef FEmitter::PreshaderParameter(FType Type, EPreshader2Opcode Opcode, FVa
 	Prototype.Opcode = Opcode;
 	Prototype.TextureIndex = MaterialInterface->GetDefaultTextureIdx(SourceParameterTexture);
 	Prototype.Payload = MoveTemp(Payload);
+
 	return EmitPrototype(*this, Prototype);
 }
 
@@ -2168,7 +2047,7 @@ static FValueRef TryFoldOperatorScalar(FEmitter& Emitter, EOperator Op, FValueRe
 		bool Result;
 		switch (PrimitiveType->ScalarKind)
 		{
-			case EScalarKind::Int:
+			case EScalarKind::Integer:
 				Result = FoldComparisonOperatorScalar<TInteger>(Op, AConstant->Integer, BConstant->Integer);
 				break;
 
@@ -2189,13 +2068,13 @@ static FValueRef TryFoldOperatorScalar(FEmitter& Emitter, EOperator Op, FValueRe
 	{
 		switch (PrimitiveType->ScalarKind)
 		{
-			case EScalarKind::Bool:
+			case EScalarKind::Boolean:
 			{
 				bool Result = FoldScalarOperator<TInteger>(Emitter, Op, AConstant->Boolean, BConstant ? BConstant->Boolean : 0, 0) & 0x1;
 				return Emitter.ConstantBool(Result);
 			}
 
-			case EScalarKind::Int:
+			case EScalarKind::Integer:
 			{
 				TInteger Result = FoldScalarOperator<TInteger>(Emitter, Op, AConstant->Integer, BConstant ? BConstant->Integer : 0, CConstant ? CConstant->Integer : 0);
 				return Emitter.ConstantInt(Result);
@@ -2555,7 +2434,7 @@ static FType ValidateOperatorAndGetResultType(FEmitter& Emitter, EOperator Op, F
 		int32 MaxNumComponents = FMath::Max3(FirstArgumentType.GetPrimitive().NumComponents(), BPrimitiveType.NumComponents(), CPrimitiveType.NumComponents());
 
 		// Cast the first argument (the boolean condition) to a bool vector of the maximum number of components.
-		Arguments[0] = Emitter.Cast(Arguments[0], FType::MakeVector(EScalarKind::Bool, MaxNumComponents));
+		Arguments[0] = Emitter.Cast(Arguments[0], FType::MakeVector(EScalarKind::Boolean, MaxNumComponents));
 
 		// Compute the common type between the second and third argument types with a number of components equal to the max of all three.
 		FType CommonTypeBetweenSecondAndThirdArguments = Emitter.TryGetCommonType(
@@ -2584,7 +2463,7 @@ static FType ValidateOperatorAndGetResultType(FEmitter& Emitter, EOperator Op, F
 			}
 		}
 
-		EScalarKind ArgumentsCommonScalarKind = EScalarKind::Bool;
+		EScalarKind ArgumentsCommonScalarKind = EScalarKind::Boolean;
 		if (bRequiresArgumentsCommonScalarKind)
 		{
 			for (int32 i = 0; Arguments[i]; ++i)
@@ -2640,7 +2519,7 @@ static FType ValidateOperatorAndGetResultType(FEmitter& Emitter, EOperator Op, F
 			return Arguments[0]->Type;
 
 		case OR_BooleanWithFirstArgumentDimensions:
-			return FType::MakePrimitive(EScalarKind::Bool, FirstArgumentPrimitiveType.NumRows, FirstArgumentPrimitiveType.NumColumns);
+			return FType::MakePrimitive(EScalarKind::Boolean, FirstArgumentPrimitiveType.NumRows, FirstArgumentPrimitiveType.NumColumns);
 
  		case OR_FirstArgumentTypeToScalarLWC:
 			return FirstArgumentPrimitiveType.ToScalar();
@@ -3222,22 +3101,22 @@ static FValue* CastConstant(FEmitter& Emitter, FConstant* Constant, EScalarKind
 
 	switch (ConstantScalarKind)
 	{
-		case EScalarKind::Bool:
+		case EScalarKind::Boolean:
 		{
 			switch (TargetKind)
 			{
-				case EScalarKind::Int: return Emitter.ConstantInt(Constant->Boolean ? 1 : 0);
+				case EScalarKind::Integer: return Emitter.ConstantInt(Constant->Boolean ? 1 : 0);
 				case EScalarKind::Float: return Emitter.ConstantFloat(Constant->Boolean ? 1.0f : 0.0f);
 				case EScalarKind::Double: return Emitter.ConstantDouble(Constant->Boolean ? 1.0 : 0.0);
 				default: UE_MIR_UNREACHABLE();
 			}
 		}
 
-		case EScalarKind::Int:
+		case EScalarKind::Integer:
 		{
 			switch (TargetKind)
 			{
-				case EScalarKind::Bool: return Emitter.ConstantBool(Constant->Integer != 0);
+				case EScalarKind::Boolean: return Emitter.ConstantBool(Constant->Integer != 0);
 				case EScalarKind::Float: return Emitter.ConstantFloat((TFloat)Constant->Integer);
 				case EScalarKind::Double: return Emitter.ConstantDouble(Constant->Integer);
 				default: UE_MIR_UNREACHABLE();
@@ -3248,8 +3127,8 @@ static FValue* CastConstant(FEmitter& Emitter, FConstant* Constant, EScalarKind
 		{
 			switch (TargetKind)
 			{
-				case EScalarKind::Bool: return Emitter.ConstantBool(Constant->Float != 0.0f);
-				case EScalarKind::Int: return Emitter.ConstantInt((int32)Constant->Float);
+				case EScalarKind::Boolean: return Emitter.ConstantBool(Constant->Float != 0.0f);
+				case EScalarKind::Integer: return Emitter.ConstantInt((int32)Constant->Float);
 				case EScalarKind::Double: return Emitter.ConstantDouble(Constant->Float);
 				default: UE_MIR_UNREACHABLE();
 			}
@@ -3259,8 +3138,8 @@ static FValue* CastConstant(FEmitter& Emitter, FConstant* Constant, EScalarKind
 		{
 			switch (TargetKind)
 			{
-				case EScalarKind::Bool: return Emitter.ConstantBool(Constant->Double != 0.0);
-				case EScalarKind::Int: return Emitter.ConstantInt((int32)Constant->Double);
+				case EScalarKind::Boolean: return Emitter.ConstantBool(Constant->Double != 0.0);
+				case EScalarKind::Integer: return Emitter.ConstantInt((int32)Constant->Double);
 				case EScalarKind::Float: return Emitter.ConstantFloat((float)Constant->Double);
 				default: UE_MIR_UNREACHABLE();
 			}
@@ -3452,12 +3331,12 @@ FValueRef FEmitter::CastToScalarKind(FValueRef Value, EScalarKind ToScalarKind)
 
 FValueRef FEmitter::CastToBoolKind(FValueRef Value)
 {
-	return CastToScalarKind(Value, EScalarKind::Bool);
+	return CastToScalarKind(Value, EScalarKind::Boolean);
 }
 
 FValueRef FEmitter::CastToIntKind(FValueRef Value)
 {
-	return CastToScalarKind(Value, EScalarKind::Int);
+	return CastToScalarKind(Value, EScalarKind::Integer);
 }
 
 FValueRef FEmitter::CastToFloatKind(FValueRef Value)
@@ -3467,12 +3346,12 @@ FValueRef FEmitter::CastToFloatKind(FValueRef Value)
 
 FValueRef FEmitter::CastToBool(FValueRef Value, int NumColumns)
 {
-	return Cast(Value, FType::MakeVector(EScalarKind::Bool, NumColumns));
+	return Cast(Value, FType::MakeVector(EScalarKind::Boolean, NumColumns));
 }
 
 FValueRef FEmitter::CastToInt(FValueRef Value, int NumColumns)
 {
-	return Cast(Value, FType::MakeVector(EScalarKind::Int, NumColumns));
+	return Cast(Value, FType::MakeVector(EScalarKind::Integer, NumColumns));
 }
 
 FValueRef FEmitter::CastToFloat(FValueRef Value, int NumColumns)
@@ -3513,10 +3392,79 @@ FValueRef FEmitter::Nop(FValueRef Arg)
 	return EmitPrototype(*this, Prototype);
 }
 
+static EMaterialSamplerType MapSamplerTypeForTexture(FValueRef TextureValue, EMaterialSamplerType InSamplerType)
+{
+	// Can't sample with virtual texturing if input texture is not a virtual texture
+	if (InSamplerType == SAMPLERTYPE_MAX)
+	{
+		if (MIR::FTextureUniform* TextureUniform = MIR::As<MIR::FTextureUniform>(TextureValue))
+		{
+			InSamplerType = TextureUniform->SamplerType;
+		}
+		else if (MIR::As<MIR::FVirtualTextureUniform>(TextureValue))
+		{
+			InSamplerType = SAMPLERTYPE_VirtualMasks;
+		}
+	}
+	return As<FVirtualTextureUniform>(TextureValue) ? PromoteSamplerTypeFromSTtoVT(InSamplerType) : DemoteSamplerTypeFromVTtoST(InSamplerType);
+}
+
+static FValueRef VTPageTableLoadFromSamplerSource(FEmitter& Em, FValueRef Texture, const FTextureSampleAttributes& BaseAttributes, FValueRef TexCoord,
+	FValueRef TexCoordDdx = {}, FValueRef TexCoordDdy = {}, ETextureMipValueMode MipValueMode = TMVM_None, FValueRef MipValue = {})
+{
+	// Cast input texture to UTexture. If it's a URuntimeVirtualTexture, we accept the cast to be null when passed to GetTextureAddressForSamplerSource().
+	TextureAddress StaticAddressX = TA_Wrap;
+	TextureAddress StaticAddressY = TA_Wrap;
+	TextureAddress StaticAddressZ = TA_Wrap;
+	UE::MaterialTranslatorUtils::GetTextureAddressForSamplerSource(Cast<UTexture>(Texture->GetTextureObject()), BaseAttributes.SamplerSourceMode, StaticAddressX, StaticAddressY, StaticAddressZ);
+
+	// Emits the virtual texture page load instruction. VT stack and layer indices are initializes during IR analysis
+	FVTPageTableRead Prototype = MakePrototype<FVTPageTableRead>(FType::MakeVTPageTableResult());
+	Prototype.VirtualTextureUniform = Texture;
+	Prototype.TexCoord = TexCoord;
+	Prototype.TexCoordDdx = TexCoordDdx;
+	Prototype.TexCoordDdy = TexCoordDdy;
+	Prototype.MipValue = MipValue;
+	Prototype.AddressU = StaticAddressX;
+	Prototype.AddressV = StaticAddressY;
+	Prototype.MipValueMode = MipValueMode;
+	Prototype.bEnableFeedback = BaseAttributes.bEnableFeedback;
+	Prototype.bIsAdaptive = BaseAttributes.bIsAdaptive;
+
+	return EmitPrototype(Em, Prototype);
+}
+
+// Cast the given texture coordinate to the type necessary for sampling given texture.
+static void CastTexCoord(FEmitter& Em, FValueRef Texture, FValueRef& TexCoord)
+{
+	if (IsAnyNotValid(Texture, TexCoord))
+	{
+		return;
+	}
+
+	// @massimo.tristano Instead of looking inside Texture's object, add various FTypes to match all texture types (right now we have a single "Texture" to capture them all).
+	if (const UTexture* TextureObject = Cast<UTexture>(Texture->GetTextureObject()))
+	{
+		EMaterialValueType MaterialValueType = TextureObject->GetMaterialType();
+		if (MaterialValueType == MCT_TextureCube || MaterialValueType == MCT_Texture2DArray || MaterialValueType == MCT_VolumeTexture)
+		{
+			// These texture types require a float3 texture coordinate
+			TexCoord = Em.CastToFloat(TexCoord, 3);
+			return;
+		}
+	}
+
+	// By default, cast TexCoord to a float2
+	TexCoord = Em.CastToFloat(TexCoord, 2);
+}
+
 FValueRef FEmitter::TextureGather(FValueRef Texture, FValueRef TexCoord, ETextureReadMode GatherMode, const FTextureSampleAttributes& BaseAttributes)
 {
 	check(GatherMode >= ETextureReadMode::GatherRed && GatherMode <= ETextureReadMode::GatherAlpha);
 
+	// Cast the texture coordinate to the type necessary for sampling given texture.
+	CastTexCoord(*this, Texture, TexCoord);
+	
 	if (IsAnyNotValid(Texture, TexCoord))
 	{
 		return Poison();
@@ -3525,7 +3473,7 @@ FValueRef FEmitter::TextureGather(FValueRef Texture, FValueRef TexCoord, ETextur
 	const EMaterialSamplerType SamplerType = MapSamplerTypeForTexture(Texture, BaseAttributes.SamplerType);
 
 	FTextureRead Prototype = MakePrototype<FTextureRead>(FType::MakeFloatVector(4));
-	Prototype.TextureObject = Texture;
+	Prototype.TextureUniform = Texture;
 	Prototype.TexCoord = TexCoord;
 	Prototype.Mode = GatherMode;
 	Prototype.SamplerSourceMode = BaseAttributes.SamplerSourceMode;
@@ -3541,6 +3489,9 @@ FValueRef FEmitter::TextureGather(FValueRef Texture, FValueRef TexCoord, ETextur
 
 FValueRef FEmitter::TextureSample(FValueRef Texture, FValueRef TexCoord, bool bAutomaticViewMipBias, const FTextureSampleAttributes& BaseAttributes)
 {
+	// Cast the texture coordinate to the type necessary for sampling given texture.
+	CastTexCoord(*this, Texture, TexCoord);
+	
 	if (IsAnyNotValid(Texture, TexCoord))
 	{
 		return Poison();
@@ -3549,7 +3500,7 @@ FValueRef FEmitter::TextureSample(FValueRef Texture, FValueRef TexCoord, bool bA
 	const EMaterialSamplerType SamplerType = MapSamplerTypeForTexture(Texture, BaseAttributes.SamplerType);
 
 	FTextureRead PrototypePixel = MakePrototype<FTextureRead>(FType::MakeFloatVector(4));
-	PrototypePixel.TextureObject = Texture;
+	PrototypePixel.TextureUniform = Texture;
 	PrototypePixel.TexCoord = TexCoord;
 	PrototypePixel.Mode = ETextureReadMode::MipAuto;
 	PrototypePixel.SamplerSourceMode = BaseAttributes.SamplerSourceMode;
@@ -3577,9 +3528,9 @@ FValueRef FEmitter::TextureSample(FValueRef Texture, FValueRef TexCoord, bool bA
 
 	if (IsVirtualSamplerType(SamplerType))
 	{
-		PrototypePixel.VTPageTable = VTPageTableLoadFromSamplerSource(*this, Texture, BaseAttributes, PrototypePixel.TexCoord);
+		PrototypePixel.VTPageTable   = VTPageTableLoadFromSamplerSource(*this, Texture, BaseAttributes, PrototypePixel.TexCoord);
 		PrototypeCompute.VTPageTable = VTPageTableLoadFromSamplerSource(*this, Texture, BaseAttributes, PrototypeCompute.TexCoord, PrototypeCompute.TexCoordDdx, PrototypeCompute.TexCoordDdy);
-		PrototypeVertex.VTPageTable = VTPageTableLoadFromSamplerSource(*this, Texture, BaseAttributes, PrototypeVertex.TexCoord);
+		PrototypeVertex.VTPageTable  = VTPageTableLoadFromSamplerSource(*this, Texture, BaseAttributes, PrototypeVertex.TexCoord);
 	}
 
 	FStageSwitch StageSwitch = MakePrototype<FStageSwitch>(PrototypePixel.Type);
@@ -3592,7 +3543,10 @@ FValueRef FEmitter::TextureSample(FValueRef Texture, FValueRef TexCoord, bool bA
 
 FValueRef FEmitter::TextureSampleLevel(FValueRef Texture, FValueRef TexCoord, FValueRef MipLevel, bool bAutomaticViewMipBias, const FTextureSampleAttributes& BaseAttributes)
 {
-	if (IsAnyNotValid(Texture, TexCoord, MipLevel))
+	// Cast the texture coordinate to the type necessary for sampling given texture.
+	CastTexCoord(*this, Texture, TexCoord);
+	
+	if (IsAnyNotValid(Texture, TexCoord))
 	{
 		return Poison();
 	}
@@ -3600,7 +3554,7 @@ FValueRef FEmitter::TextureSampleLevel(FValueRef Texture, FValueRef TexCoord, FV
 	const EMaterialSamplerType SamplerType = MapSamplerTypeForTexture(Texture, BaseAttributes.SamplerType);
 
 	FTextureRead Prototype = MakePrototype<FTextureRead>(FType::MakeFloatVector(4));
-	Prototype.TextureObject = Texture;
+	Prototype.TextureUniform = Texture;
 	Prototype.TexCoord = TexCoord;
 	Prototype.MipValue = MipLevel;
 	Prototype.Mode = ETextureReadMode::MipLevel;
@@ -3622,7 +3576,10 @@ FValueRef FEmitter::TextureSampleLevel(FValueRef Texture, FValueRef TexCoord, FV
 
 FValueRef FEmitter::TextureSampleBias(FValueRef Texture, FValueRef TexCoord, FValueRef MipBias, bool bAutomaticViewMipBias, const FTextureSampleAttributes& BaseAttributes)
 {
-	if (IsAnyNotValid(Texture, TexCoord, MipBias))
+	// Cast the texture coordinate to the type necessary for sampling given texture.
+	CastTexCoord(*this, Texture, TexCoord);
+	
+	if (IsAnyNotValid(Texture, TexCoord))
 	{
 		return Poison();
 	}
@@ -3635,7 +3592,7 @@ FValueRef FEmitter::TextureSampleBias(FValueRef Texture, FValueRef TexCoord, FVa
 	const EMaterialSamplerType SamplerType = MapSamplerTypeForTexture(Texture, BaseAttributes.SamplerType);
 
 	FTextureRead PrototypePixel = MakePrototype<FTextureRead>(FType::MakeFloatVector(4));
-	PrototypePixel.TextureObject = Texture;
+	PrototypePixel.TextureUniform = Texture;
 	PrototypePixel.TexCoord = TexCoord;
 	PrototypePixel.MipValue = MipBias;
 	PrototypePixel.Mode = ETextureReadMode::MipBias;
@@ -3660,7 +3617,7 @@ FValueRef FEmitter::TextureSampleBias(FValueRef Texture, FValueRef TexCoord, FVa
 
 	if (IsVirtualSamplerType(SamplerType))
 	{
-		PrototypePixel.VTPageTable = VTPageTableLoadFromSamplerSource(*this, Texture, BaseAttributes, PrototypePixel.TexCoord, {}, {}, TMVM_MipBias, PrototypePixel.MipValue);
+		PrototypePixel.VTPageTable   = VTPageTableLoadFromSamplerSource(*this, Texture, BaseAttributes, PrototypePixel.TexCoord, {}, {}, TMVM_MipBias, PrototypePixel.MipValue);
 		PrototypeCompute.VTPageTable = VTPageTableLoadFromSamplerSource(*this, Texture, BaseAttributes, PrototypeCompute.TexCoord, PrototypeCompute.TexCoordDdx, PrototypeCompute.TexCoordDdy, TMVM_MipBias, PrototypeCompute.MipValue);
 	}
 
@@ -3674,7 +3631,10 @@ FValueRef FEmitter::TextureSampleBias(FValueRef Texture, FValueRef TexCoord, FVa
 
 FValueRef FEmitter::TextureSampleGrad(FValueRef Texture, FValueRef TexCoord, FValueRef TexCoordDdx, FValueRef TexCoordDdy, bool bAutomaticViewMipBias, const FTextureSampleAttributes& BaseAttributes)
 {
-	if (IsAnyNotValid(Texture, TexCoord, TexCoordDdx, TexCoordDdy))
+	// Cast the texture coordinate to the type necessary for sampling given texture.
+	CastTexCoord(*this, Texture, TexCoord);
+	
+	if (IsAnyNotValid(Texture, TexCoord))
 	{
 		return Poison();
 	}
@@ -3682,7 +3642,7 @@ FValueRef FEmitter::TextureSampleGrad(FValueRef Texture, FValueRef TexCoord, FVa
 	const EMaterialSamplerType SamplerType = MapSamplerTypeForTexture(Texture, BaseAttributes.SamplerType);
 
 	FTextureRead Prototype = MakePrototype<FTextureRead>(FType::MakeFloatVector(4));
-	Prototype.TextureObject = Texture;
+	Prototype.TextureUniform = Texture;
 	Prototype.TexCoord = TexCoord;
 	Prototype.TexCoordDdx = TexCoordDdx;
 	Prototype.TexCoordDdy = TexCoordDdy;
@@ -3705,25 +3665,6 @@ FValueRef FEmitter::TextureSampleGrad(FValueRef Texture, FValueRef TexCoord, FVa
 	return EmitPrototype(*this, Prototype);
 }
 
-FValueRef FEmitter::VTPageTableLoad(
-	FValueRef Texture, TextureAddress AddressU, TextureAddress AddressV, FValueRef TexCoord,
-	FValueRef TexCoordDdx, FValueRef TexCoordDdy, bool bEnableFeedback, bool bIsAdaptive, ETextureMipValueMode MipValueMode, FValueRef MipValue)
-{
-	// VT stack and layer indices are initializes during IR analysis
-	FVTPageTableRead Prototype = MakePrototype<FVTPageTableRead>(FType::MakeVTPageTableResult());
-	Prototype.TextureObject = Texture;
-	Prototype.TexCoord = TexCoord;
-	Prototype.TexCoordDdx = TexCoordDdx;
-	Prototype.TexCoordDdy = TexCoordDdy;
-	Prototype.MipValue = MipValue;
-	Prototype.AddressU = AddressU;
-	Prototype.AddressV = AddressV;
-	Prototype.MipValueMode = MipValueMode;
-	Prototype.bEnableFeedback = bEnableFeedback;
-	Prototype.bIsAdaptive = bIsAdaptive;
-	return EmitPrototype(*this, Prototype);
-}
-
 FValueRef FEmitter::PartialDerivative(FValueRef Value, EDerivativeAxis Axis)
 {
 	// Any operation on poison arguments is a poison.
@@ -4141,7 +4082,7 @@ FValueRef FEmitter::AnalyticalPartialDerivative(FValueRef Value, EDerivativeAxis
 
 		// Finally, these values are uniform (constant) or their differential is not provided, thus their value is always zero.
 		case VK_Constant:
-		case VK_UniformParameter:
+		case VK_PrimitiveUniform:
 		case VK_PreshaderParameter:
 		case VK_ExternDifferential:
 			return Cast(ConstantZero(Value->Type.GetPrimitive().ScalarKind), Value->Type);
diff --git a/Engine/Source/Runtime/Engine/Private/Materials/MaterialIRInternal.cpp b/Engine/Source/Runtime/Engine/Private/Materials/MaterialIRInternal.cpp
--- a/Engine/Source/Runtime/Engine/Private/Materials/MaterialIRInternal.cpp
+++ b/Engine/Source/Runtime/Engine/Private/Materials/MaterialIRInternal.cpp
@@ -28,19 +28,6 @@ EMaterialValueType GetTextureMaterialValueType(const UObject* TextureObject)
 	return MCT_Unknown;
 }
 
-EMaterialTextureParameterType TextureMaterialValueTypeToParameterType(EMaterialValueType Type)
-{
-	switch (Type)
-	{
-		case MCT_Texture2D: return EMaterialTextureParameterType::Standard2D;
-		case MCT_Texture2DArray: return EMaterialTextureParameterType::Array2D;
-		case MCT_TextureCube: return EMaterialTextureParameterType::Cube;
-		case MCT_TextureCubeArray: return EMaterialTextureParameterType::ArrayCube;
-		case MCT_VolumeTexture: return EMaterialTextureParameterType::Volume;
-		case MCT_TextureVirtual: return EMaterialTextureParameterType::Virtual;
-		default: UE_MIR_UNREACHABLE();
-	}
-}
 
 } // namespace MIR::Internal
 
diff --git a/Engine/Source/Runtime/Engine/Private/Materials/MaterialIRInternal.h b/Engine/Source/Runtime/Engine/Private/Materials/MaterialIRInternal.h
--- a/Engine/Source/Runtime/Engine/Private/Materials/MaterialIRInternal.h
+++ b/Engine/Source/Runtime/Engine/Private/Materials/MaterialIRInternal.h
@@ -11,9 +11,6 @@ namespace MIR::Internal {
 // Returns the material value tpye of the specified UTexture or URuntimeVirtualTexture.
 EMaterialValueType GetTextureMaterialValueType(const UObject* TextureObject);
 
-// Maps texture EMaterialValueType values to EMaterialTextureParameterType.
-EMaterialTextureParameterType TextureMaterialValueTypeToParameterType(EMaterialValueType Type);
-
 // @massimo.tristano expose an internal "Context" struct between the builder and the emitter, and move these methods
 // to the context itself
 // Returns the value flowing into given expression input (previously set through `BindValueToExpressionInput`).
diff --git a/Engine/Source/Runtime/Engine/Private/Materials/MaterialIRModuleBuilder.cpp b/Engine/Source/Runtime/Engine/Private/Materials/MaterialIRModuleBuilder.cpp
--- a/Engine/Source/Runtime/Engine/Private/Materials/MaterialIRModuleBuilder.cpp
+++ b/Engine/Source/Runtime/Engine/Private/Materials/MaterialIRModuleBuilder.cpp
@@ -792,14 +792,13 @@ struct FMaterialIRModuleBuilderImpl
 
 			// The value being set to this material output is now valid.
 			check(PropertyValue);
-
+			 
 			// Add support for lerp to selection color for PC development builds.
 			if (Property == MP_EmissiveColor &&
 				Builder->MaterialInterface->GetMaterial()->MaterialDomain != MD_Volume &&
 				UE::MaterialTranslatorUtils::IsDevelopmentFeatureEnabled(NAME_SelectionColor, Module->GetShaderPlatform(), Builder->MaterialInterface->GetMaterial()))
 			{
-				FMaterialParameterMetadata ParameterMetadata{ FMaterialParameterValue(FLinearColor::Transparent) };
-				MIR::FValue* SelectionColor = Emitter.Parameter(NAME_SelectionColor, ParameterMetadata);
+				MIR::FValue* SelectionColor = Emitter.NamedPrimitiveUniform(NAME_SelectionColor, Emitter.ConstantFloat4(UE::Math::TVector4<float>{ 0.f, 0.f, 0.f, 0.f }));
 				PropertyValue = Emitter.Lerp(PropertyValue, Emitter.Swizzle(SelectionColor, MIR::FSwizzleMask::XYZ()), Emitter.Subscript(SelectionColor, 3));
 			}
 
diff --git a/Engine/Source/Runtime/Engine/Private/Materials/MaterialIRToHLSLTranslator.cpp b/Engine/Source/Runtime/Engine/Private/Materials/MaterialIRToHLSLTranslator.cpp
--- a/Engine/Source/Runtime/Engine/Private/Materials/MaterialIRToHLSLTranslator.cpp
+++ b/Engine/Source/Runtime/Engine/Private/Materials/MaterialIRToHLSLTranslator.cpp
@@ -217,8 +217,8 @@ static FPrinter& operator<<(FPrinter& Printer, const MIR::FType& Type)
 		{
 			switch (PrimitiveType->ScalarKind)
 			{
-				case MIR::EScalarKind::Bool:  Printer << TEXTVIEW("bool"); break;
-				case MIR::EScalarKind::Int:   Printer << TEXTVIEW("int"); break;
+				case MIR::EScalarKind::Boolean:  Printer << TEXTVIEW("bool"); break;
+				case MIR::EScalarKind::Integer:   Printer << TEXTVIEW("int"); break;
 				case MIR::EScalarKind::Float: Printer << TEXTVIEW("MaterialFloat"); break;
 			}
 
@@ -402,9 +402,9 @@ struct FEntryPointCodeGen
 			case MIR::VK_Constant: LowerConstant(static_cast<const MIR::FConstant*>(Value)); break;
 			case MIR::VK_Builtin: LowerBuiltin(static_cast<const MIR::FBuiltin*>(Value)); break;
 			case MIR::VK_MaterialParameterCollection: LowerMaterialParameterCollection(static_cast<const MIR::FMaterialParameterCollection*>(Value)); break;
-			case MIR::VK_TextureObject: LowerTextureObject(static_cast<const MIR::FTextureObject*>(Value)); break;
-			case MIR::VK_RuntimeVirtualTextureObject: LowerRuntimeVirtualTextureObject(static_cast<const MIR::FRuntimeVirtualTextureObject*>(Value)); break;
-			case MIR::VK_UniformParameter: LowerUniformParameter(static_cast<const MIR::FUniformParameter*>(Value)); break;
+			case MIR::VK_PrimitiveUniform: LowerPrimitiveUniform(static_cast<const MIR::FPrimitiveUniform*>(Value)); break;
+			case MIR::VK_TextureUniform: LowerTextureUniform(static_cast<const MIR::FTextureUniform*>(Value)); break;
+			case MIR::VK_VirtualTextureUniform: LowerVirtualTextureUniform(static_cast<const MIR::FVirtualTextureUniform*>(Value)); break;
 			case MIR::VK_GetVertexInterpolator: LowerGetVertexInterpolator(static_cast<const MIR::FGetVertexInterpolator*>(Value)); break;
 			default: UE_MIR_UNREACHABLE();
 		}
@@ -455,11 +455,11 @@ struct FEntryPointCodeGen
 
 		switch (Primitive->ScalarKind)
 		{
-			case MIR::EScalarKind::Bool:
+			case MIR::EScalarKind::Boolean:
 				Printer << (Constant->Boolean ? TEXTVIEW("true") : TEXTVIEW("false"));
 				break;
 				
-			case MIR::EScalarKind::Int:
+			case MIR::EScalarKind::Integer:
 				Printer.Buffer.Appendf(TEXT("%lld"), Constant->Integer);
 				break;
 			
@@ -556,28 +556,21 @@ struct FEntryPointCodeGen
 		ExternArg->ToHLSL(ExternPrinter);
 	}
 
-	void LowerTextureObject(const MIR::FTextureObject* TextureObject)
+	void LowerPrimitiveUniform(const MIR::FPrimitiveUniform* PrimitiveUniform)
 	{
-		LowerTextureReference(TextureObject->Texture->GetMaterialType(), TextureObject->Analysis_UniformParameterIndex);
+		// Get the global float4 component index (e.g. if this is 13, it refer to PreshaderBuffer[3].y)
+		const uint32 GlobalComponentOffset = Module->GetCompilationOutput().UniformExpressionSet.GetNumericParameterEvaluationOffset(PrimitiveUniform->Analysis_UniformIndex);
+		EmitPreshaderBufferRead(PrimitiveUniform->Type.GetPrimitive(), GlobalComponentOffset);
 	}
-
-	void LowerRuntimeVirtualTextureObject(const MIR::FRuntimeVirtualTextureObject* RVTextureObject)
+	
+	void LowerTextureUniform(const MIR::FTextureUniform* TextureUniform)
 	{
-		LowerTextureReference(MCT_TextureVirtual, RVTextureObject->Analysis_UniformParameterIndex);
+		Printer << TEXTVIEW("Material.") << LowerTextureMaterialType(TextureUniform->Texture->GetMaterialType(), true) << TEXTVIEW("_") << TextureUniform->Analysis_UniformIndex;
 	}
 
-	void LowerUniformParameter(const MIR::FUniformParameter* UniformParameter)
+	void LowerVirtualTextureUniform(const MIR::FVirtualTextureUniform* VirtualTextureUniform)
 	{
-		if (UniformParameter->Type.IsTexture() || UniformParameter->Type.IsRuntimeVirtualTexture())
-		{
-			UObject* TextureObject = GetTextureFromUniformParameter(UniformParameter);
-			check(TextureObject);
-			LowerTextureReference(MIR::Internal::GetTextureMaterialValueType(TextureObject), UniformParameter->Analysis_UniformParameterIndex);
-		}
-		else
-		{
-			LowerPrimitiveUniformParameter(UniformParameter);
-		}
+		Printer << TEXTVIEW("Material.") << LowerTextureMaterialType(MCT_TextureVirtual, true) << TEXTVIEW("_") << VirtualTextureUniform->Analysis_UniformIndex;
 	}
 
 	void EmitPreshaderBufferReadDoubleVector(const MIR::FPrimitive& PrimitiveType, uint32 GlobalComponentOffset)
@@ -642,17 +635,6 @@ struct FEntryPointCodeGen
 		}
 	}
 
-	ENoOp LowerPrimitiveUniformParameter(const MIR::FUniformParameter* UniformParameter)
-	{
-		// Get the global float4 component index (e.g. if this is 13, it refer to PreshaderBuffer[3].y)
-		const FUniformExpressionSet& UniformExpressionSet = Module->GetCompilationOutput().UniformExpressionSet;
-		const uint32 GlobalComponentOffset = UniformExpressionSet.GetNumericParameterEvaluationOffset(UniformParameter->Analysis_UniformParameterIndex);
-
-		EmitPreshaderBufferRead(UniformParameter->Type.GetPrimitive(), GlobalComponentOffset);
-
-		return NoOp;
-	}
-
 	ENoOp LowerPreshaderParameter(const MIR::FPreshaderParameter* PreshaderParameter)
 	{
 		// Get the global float4 component index (e.g. if this is 13, it refer to PreshaderBuffer[3].y)
@@ -663,22 +645,6 @@ struct FEntryPointCodeGen
 		return NoOp;
 	}
 
-	// Variation of LowerPrimitiveUniformParameter specific to LWC, to support fetching a subset of an LWC vector.
-	ENoOp LowerPrimitiveUniformParameterLWC(const MIR::FUniformParameter* UniformParameter, int32 NumComponents)
-	{
-		check(NumComponents >= 2 && NumComponents <= 4);
-		check(UniformParameter->Type.IsDouble());
-
-		const FUniformExpressionSet& UniformExpressionSet = Module->GetCompilationOutput().UniformExpressionSet;
-
-		// Get the global float4 component index (e.g. if this is 13, it refer to PreshaderBuffer[3].y)
-		const uint32 GlobalComponentOffset = UniformExpressionSet.GetNumericParameterEvaluationOffset(UniformParameter->Analysis_UniformParameterIndex);
-
-		EmitPreshaderBufferReadDoubleVector(UniformParameter->Type.GetPrimitive(), GlobalComponentOffset);
-	
-		return NoOp;
-	}
-
 	bool HasMatchingScalarComponentCastChain(MIR::FValue* FirstComponent, MIR::FValue* CurrentComponent, int32 Index, int32 NumComponents)
 	{
 		const MIR::FSubscript* FirstComponentAsSubscript = FirstComponent->As<MIR::FSubscript>();
@@ -701,7 +667,7 @@ struct FEntryPointCodeGen
 				if (ArgPrimitiveType.NumColumns >= NumComponents)
 				{
 					// Allow swizzle for non-LWC values or uniform parameters (the latter having a special case for LWC)
-					return !ArgPrimitiveType.IsDouble() || FirstComponentAsSubscript->Arg->As<MIR::FUniformParameter>();
+					return !ArgPrimitiveType.IsDouble() || FirstComponentAsSubscript->Arg->As<MIR::FPrimitiveUniform>();
 				}
 			}
 		}
@@ -776,7 +742,7 @@ struct FEntryPointCodeGen
 				{
 					// Cast from bool requires a select between 1 and 0
 					MIR::FType VectorType = MIR::FType::MakeVector(PrimitiveType.ScalarKind, NumComponents);
-					Printer << TEXTVIEW("select()") << LowerVectorCastChain(Arg, NumComponents) << TEXTVIEW(", (") << VectorType << TEXTVIEW(")1, (") << VectorType << TEXTVIEW(")0)");
+					Printer << TEXTVIEW("select(") << LowerVectorCastChain(Arg, NumComponents) << TEXTVIEW(", (") << VectorType << TEXTVIEW(")1, (") << VectorType << TEXTVIEW(")0)");
 				}
 				else
 				{
@@ -790,10 +756,19 @@ struct FEntryPointCodeGen
 		{
 			// Finally reached the inner subscript, print its vector argument
 			const MIR::FValue* Arg = FirstComponentAsSubscript->Arg;
-			if (Arg->Type.IsVector() && Arg->Type.IsDouble() && Arg->As<MIR::FUniformParameter>())
+			const MIR::FPrimitiveUniform* ArgAsPrimitiveUniform = Arg->As<MIR::FPrimitiveUniform>();
+			if (Arg->Type.IsVector() && Arg->Type.IsDouble() && ArgAsPrimitiveUniform)
 			{
 				// LWC types support special case swizzling logic for uniform parameters on initial fetch. See HasMatchingScalarComponentCastChain above.
-				LowerPrimitiveUniformParameterLWC(static_cast<const MIR::FUniformParameter*>(Arg), NumComponents);
+				check(NumComponents >= 2 && NumComponents <= 4);
+				check(ArgAsPrimitiveUniform->Type.IsDouble());
+
+				const FUniformExpressionSet& UniformExpressionSet = Module->GetCompilationOutput().UniformExpressionSet;
+
+				// Get the global float4 component index (e.g. if this is 13, it refers to PreshaderBuffer[3].y)
+				const uint32 GlobalComponentOffset = UniformExpressionSet.GetNumericParameterEvaluationOffset(ArgAsPrimitiveUniform->Analysis_UniformIndex);
+
+				EmitPreshaderBufferReadDoubleVector(ArgAsPrimitiveUniform->Type.GetPrimitive(), GlobalComponentOffset);
 			}
 			else
 			{
@@ -1259,7 +1234,7 @@ struct FEntryPointCodeGen
 			if (ArgVectorType->IsDouble())
 			{
 				FStringView LwcComponentsStr[] = { TEXTVIEW("WSGetX("), TEXTVIEW("WSGetY("), TEXTVIEW("WSGetZ("), TEXTVIEW("WSGetW(") };
-				check(Subscript->Index <= ArgVectorType->NumComponents());
+				check(Subscript->Index < ArgVectorType->NumComponents() && Subscript->Index < 4);
 
 				Printer << LwcComponentsStr[Subscript->Index] << LowerValue(Subscript->Arg) << TEXTVIEW(")");
 			}
@@ -1268,7 +1243,7 @@ struct FEntryPointCodeGen
 				LowerValue(Subscript->Arg);
 				
 				FStringView ComponentsStr[] = { TEXTVIEW(".x"), TEXTVIEW(".y"), TEXTVIEW(".z"), TEXTVIEW(".w") };
-				check(Subscript->Index <= ArgVectorType->NumComponents());
+				check(Subscript->Index < ArgVectorType->NumComponents() && Subscript->Index < 4);
 
 				Printer << ComponentsStr[Subscript->Index];
 			}
@@ -1360,7 +1335,7 @@ struct FEntryPointCodeGen
 		}
 	}
 
-	ENoOp LowerTextureMaterialType(EMaterialValueType TextureType, bool bForResourceDeclarations = false)
+	ENoOp LowerTextureMaterialType(EMaterialValueType TextureType, bool bForResourceDeclarations)
 	{
 		switch (TextureType)
 		{
@@ -1376,32 +1351,12 @@ struct FEntryPointCodeGen
 		return NoOp;
 	}
 
-	void LowerTextureType(const MIR::FTextureObject* TextureObject)
-	{
-		check(TextureObject->Texture);
-		LowerTextureMaterialType(TextureObject->Texture->GetMaterialType());
-	}
-
-	UObject* GetTextureFromUniformParameter(const MIR::FUniformParameter* UniformParameter) const
-	{
-		const FMaterialParameterValue& TextureParameter = Module->GetParameterMetadata(UniformParameter->ParameterIdInModule).Value;
-		return TextureParameter.AsTextureObject();
-	}
-
-	void LowerTextureParameterType(const MIR::FUniformParameter* UniformParameter)
-	{
-		UObject* Texture = GetTextureFromUniformParameter(UniformParameter);
-		check(Texture);
-		LowerTextureMaterialType(MIR::Internal::GetTextureMaterialValueType(Texture));
-	}
-
 	void LowerStandardTextureRead(const MIR::FTextureRead* TextureRead)
 	{
-		switch (TextureRead->TextureObject->Kind)
+		switch (TextureRead->TextureUniform->Kind)
 		{
-			case MIR::EValueKind::VK_TextureObject: LowerTextureType(TextureRead->TextureObject->As<MIR::FTextureObject>()); break;
-			case MIR::EValueKind::VK_RuntimeVirtualTextureObject: LowerTextureMaterialType(MCT_TextureVirtual); break;
-			case MIR::EValueKind::VK_UniformParameter: LowerTextureParameterType(TextureRead->TextureObject->As<MIR::FUniformParameter>()); break;
+			case MIR::EValueKind::VK_TextureUniform: 		LowerTextureMaterialType(TextureRead->TextureUniform->As<MIR::FTextureUniform>()->Texture->GetMaterialType(), false); break;
+			case MIR::EValueKind::VK_VirtualTextureUniform: LowerTextureMaterialType(MCT_TextureVirtual, false); break;
 			default: UE_MIR_UNREACHABLE();
 		}
 
@@ -1419,8 +1374,8 @@ struct FEntryPointCodeGen
 		}
 
 		Printer << BeginArgs
-				<< ListSeparator << LowerValue(TextureRead->TextureObject)
-				<< ListSeparator << LowerTextureSamplerReference(TextureRead->TextureObject, TextureRead->SamplerSourceMode)
+				<< ListSeparator << LowerValue(TextureRead->TextureUniform)
+				<< ListSeparator << LowerTextureSamplerReference(TextureRead->TextureUniform, TextureRead->SamplerSourceMode)
 				<< ListSeparator << LowerValue(TextureRead->TexCoord);
 
 		switch (TextureRead->Mode)
@@ -1439,7 +1394,7 @@ struct FEntryPointCodeGen
 		checkf(TextureRead->VTPageTable, TEXT("Missing page table for virtual texture read instruction"));
 		const MIR::FVTPageTableRead* VTPageTableResult = TextureRead->VTPageTable->As<MIR::FVTPageTableRead>();
 
-		const int32 VirtualTextureIndex = TextureRead->TextureObject->GetUniformParameterIndex();
+		const int32 VirtualTextureIndex = TextureRead->TextureUniform->GetUniformIndex();
 		check(VirtualTextureIndex >= 0);
 
 		// Sampling function
@@ -1448,7 +1403,7 @@ struct FEntryPointCodeGen
 		// 'Texture name/sampler', 'PageTableResult', 'LayerIndex', 'PackedUniform'
 		Printer
 			<< BeginArgs
-			<< ListSeparator << LowerValue(TextureRead->TextureObject)
+			<< ListSeparator << LowerValue(TextureRead->TextureUniform)
 			<< ListSeparator;
 		
 		// Sampling function argument list
@@ -1481,8 +1436,8 @@ struct FEntryPointCodeGen
 		}
 
 		checkf(
-			TextureRead->TextureObject->Type.IsTexture() ||
-			TextureRead->TextureObject->Type.IsRuntimeVirtualTexture(),
+			TextureRead->TextureUniform->Type.IsTexture() ||
+			TextureRead->TextureUniform->Type.IsRuntimeVirtualTexture(),
 			TEXT("Invalid texture object type")
 		);
 
@@ -1672,13 +1627,6 @@ struct FEntryPointCodeGen
 		return NoOp;
 	}
 
-	ENoOp LowerTextureReference(EMaterialValueType TextureType, int32 TextureParameterIndex)
-	{
-		checkf(TextureParameterIndex != INDEX_NONE, TEXT("Texture uniform parameter not assigned! A texture is used for sampling but hasn't been properly registered during IR analysis."));
-		Printer << TEXTVIEW("Material.") << LowerTextureMaterialType(TextureType, true) << TEXTVIEW("_") << TextureParameterIndex;
-		return NoOp;
-	}
-	
 	void LowerStageSwitch(const MIR::FStageSwitch* StageSwitch)
 	{
 		LowerValue(StageSwitch->Args[CurrentStage]);
diff --git a/Engine/Source/Runtime/Engine/Private/Materials/MaterialIRTypes.cpp b/Engine/Source/Runtime/Engine/Private/Materials/MaterialIRTypes.cpp
--- a/Engine/Source/Runtime/Engine/Private/Materials/MaterialIRTypes.cpp
+++ b/Engine/Source/Runtime/Engine/Private/Materials/MaterialIRTypes.cpp
@@ -27,7 +27,7 @@ const TCHAR* TypeKindToString(ETypeKind Kind)
 
 bool ScalarKindIsArithmetic(EScalarKind Kind)
 {
-	return Kind != EScalarKind::Bool;
+	return Kind != EScalarKind::Boolean;
 }
 
 bool ScalarKindIsAnyFloat(EScalarKind Kind)
@@ -39,8 +39,8 @@ const TCHAR* ScalarKindToString(EScalarKind Kind)
 {
 	switch (Kind)
 	{
-		case EScalarKind::Bool: return TEXT("bool");
-		case EScalarKind::Int: return TEXT("int");
+		case EScalarKind::Boolean: return TEXT("bool");
+		case EScalarKind::Integer: return TEXT("int");
 		case EScalarKind::Float: return TEXT("MaterialFloat");
 		case EScalarKind::Double: return TEXT("FLWCScalar");
 		default: UE_MIR_UNREACHABLE();
@@ -217,8 +217,8 @@ FString FType::GetSpelling() const
 		FString Spelling;
 		switch (Primitive.ScalarKind)
 		{
-			case EScalarKind::Bool:Spelling.Append(TEXT("bool")); break;
-			case EScalarKind::Int: Spelling.Append(TEXT("int")); break;
+			case EScalarKind::Boolean:Spelling.Append(TEXT("bool")); break;
+			case EScalarKind::Integer: Spelling.Append(TEXT("int")); break;
 			case EScalarKind::Float: Spelling.Append(TEXT("float")); break;
 			case EScalarKind::Double:
 				if (Primitive.IsScalar())
@@ -309,8 +309,8 @@ UE::Shader::EValueType FType::ToValueType() const
 
 			switch (Primitive.ScalarKind)
 			{
-				case EScalarKind::Bool: return (EValueType)((int)EValueType::Bool1 + Primitive.NumColumns - 1);
-				case EScalarKind::Int: 	return (EValueType)((int)EValueType::Int1 + Primitive.NumColumns - 1);
+				case EScalarKind::Boolean: return (EValueType)((int)EValueType::Bool1 + Primitive.NumColumns - 1);
+				case EScalarKind::Integer: 	return (EValueType)((int)EValueType::Int1 + Primitive.NumColumns - 1);
 				case EScalarKind::Float: return (EValueType)((int)EValueType::Float1 + Primitive.NumColumns - 1);
 				case EScalarKind::Double:	return (EValueType)((int)EValueType::Double1 + Primitive.NumColumns - 1);
 				default: UE_MIR_UNREACHABLE();
