New material translator work:

diff --git a/Engine/Source/Runtime/Engine/Private/Materials/MaterialIRValueAnalyzer.cpp b/Engine/Source/Runtime/Engine/Private/Materials/MaterialIRValueAnalyzer.cpp
--- a/Engine/Source/Runtime/Engine/Private/Materials/MaterialIRValueAnalyzer.cpp
+++ b/Engine/Source/Runtime/Engine/Private/Materials/MaterialIRValueAnalyzer.cpp
@@ -56,6 +56,21 @@ static EMaterialShaderFrequency MapToMaterialShaderFrequencyOrAny(MIR::EStage St
 	return EMaterialShaderFrequency::Any;
 }
 
+// Maps texture EMaterialValueType values to EMaterialTextureParameterType.
+static EMaterialTextureParameterType TextureMaterialValueTypeToParameterType(EMaterialValueType Type)
+{
+	switch (Type)
+	{
+		case MCT_Texture2D: return EMaterialTextureParameterType::Standard2D;
+		case MCT_Texture2DArray: return EMaterialTextureParameterType::Array2D;
+		case MCT_TextureCube: return EMaterialTextureParameterType::Cube;
+		case MCT_TextureCubeArray: return EMaterialTextureParameterType::ArrayCube;
+		case MCT_VolumeTexture: return EMaterialTextureParameterType::Volume;
+		case MCT_TextureVirtual: return EMaterialTextureParameterType::Virtual;
+		default: UE_MIR_UNREACHABLE();
+	}
+}
+
 static void Extern_Analyze(FMaterialIRValueAnalyzer& Analyzer, MIR::FExtern* Extern)
 {
 	MIR::FExternAnalysisContext Context = {
@@ -105,105 +120,12 @@ static int32 AcquireVTStackIndex(FMaterialIRValueAnalyzer& Analyzer, MIR::FVTPag
 	return Analyzer.CompilationOutput->UniformExpressionSet.AddVTStack(INDEX_NONE);
 }
 
-static int32 RegisterTextureParameter(FMaterialIRValueAnalyzer& Analyzer, UObject* Texture, const FHashedMaterialParameterInfo& ParameterInfo, int32 VTLayerIndex = INDEX_NONE)
-{
-	check(Texture);
-	check(VTLayerIndex == INDEX_NONE || (VTLayerIndex >= 0 && VTLayerIndex < UINT8_MAX));
-
-	const FMaterialTextureParameterInfo TextureParameterInfo
-	{
-		.ParameterInfo = ParameterInfo,
-		.TextureIndex = Analyzer.MaterialInterface->GetDefaultTextureIdx(Texture),
-		.SamplerSource = SSM_FromTextureAsset,
-		.VirtualTextureLayerIndex = static_cast<uint8>(VTLayerIndex == INDEX_NONE ? UINT8_MAX : VTLayerIndex),
-	};
-	check(TextureParameterInfo.TextureIndex != INDEX_NONE);
-
-	const EMaterialTextureParameterType ParamType = MIR::Internal::TextureMaterialValueTypeToParameterType(MIR::Internal::GetTextureMaterialValueType(Texture));
-
-	return Analyzer.CompilationOutput->UniformExpressionSet.FindOrAddTextureParameter(ParamType, TextureParameterInfo);
-}
-
-static void AccessTextureObject(FMaterialIRValueAnalyzer& Analyzer, MIR::FTextureObject* TextureObject)
-{
-	if (TextureObject->Analysis_UniformParameterIndex == INDEX_NONE)
-	{
-		const FHashedMaterialParameterInfo ParameterInfo{ TextureObject->Texture->GetFName(), EMaterialParameterAssociation::GlobalParameter, INDEX_NONE };
-		TextureObject->Analysis_UniformParameterIndex = RegisterTextureParameter(Analyzer, TextureObject->Texture, ParameterInfo, 0);
-	}
-}
-
-static void AccessRuntimeVirtualTextureObject(FMaterialIRValueAnalyzer& Analyzer, MIR::FRuntimeVirtualTextureObject* RVTextureObject)
-{
-	if (RVTextureObject->Analysis_UniformParameterIndex == INDEX_NONE)
-	{
-		const FHashedMaterialParameterInfo ParameterInfo{ RVTextureObject->RVTexture->GetFName(), EMaterialParameterAssociation::GlobalParameter, INDEX_NONE };
-		RVTextureObject->Analysis_UniformParameterIndex = RegisterTextureParameter(Analyzer, RVTextureObject->RVTexture, ParameterInfo, RVTextureObject->VTLayerIndex);
-	}
-}
-
-static void AccessTextureUniformParameter(FMaterialIRValueAnalyzer& Analyzer, MIR::FUniformParameter* UniformParameter)
-{
-	if (UniformParameter->Analysis_UniformParameterIndex == INDEX_NONE)
-	{
-		const FMaterialParameterValue& ParameterValue = Analyzer.Module->GetParameterMetadata(UniformParameter->ParameterIdInModule).Value;
-		const FHashedMaterialParameterInfo& ParameterInfo = Analyzer.Module->GetParameterInfo(UniformParameter->ParameterIdInModule);
-		UniformParameter->Analysis_UniformParameterIndex = RegisterTextureParameter(Analyzer, ParameterValue.AsTextureObject(), ParameterInfo, UniformParameter->VTLayerIndex);
-	}
-}
-
-static void AccessTexture(FMaterialIRValueAnalyzer& Analyzer, MIR::FValue* TextureValue)
-{
-	if (MIR::FTextureObject* TextureObject = TextureValue->As<MIR::FTextureObject>())
-	{
-		AccessTextureObject(Analyzer, TextureObject);
-	}
-	else if (MIR::FRuntimeVirtualTextureObject* RVTextureObject = TextureValue->As<MIR::FRuntimeVirtualTextureObject>())
-	{
-		AccessRuntimeVirtualTextureObject(Analyzer, RVTextureObject);
-	}
-	else if (MIR::FUniformParameter* UniformParameter = TextureValue->As<MIR::FUniformParameter>())
-	{
-		AccessTextureUniformParameter(Analyzer, UniformParameter);
-	}
-	else
-	{
-		UE_MIR_UNREACHABLE();
-	}
-}
-
-static int32 GetVTLayerIndex(MIR::FValue* TextureObject)
-{
-	if (MIR::FRuntimeVirtualTextureObject* RVTextureObject = TextureObject->As<MIR::FRuntimeVirtualTextureObject>())
-	{
-		return RVTextureObject->VTLayerIndex;
-	}
-	if (MIR::FUniformParameter* UniformParameter = TextureObject->As<MIR::FUniformParameter>())
-	{
-		return UniformParameter->VTLayerIndex;
-	}
-	return INDEX_NONE;
-}
-
-static int32 GetPageTableLayerIndex(MIR::FValue* TextureObject)
-{
-	if (MIR::FRuntimeVirtualTextureObject* RVTextureObject = TextureObject->As<MIR::FRuntimeVirtualTextureObject>())
-	{
-		return RVTextureObject->VTPageTableIndex;
-	}
-	if (MIR::FUniformParameter* UniformParameter = TextureObject->As<MIR::FUniformParameter>())
-	{
-		return UniformParameter->VTPageTableIndex;
-	}
-	return INDEX_NONE;
-}
-
 static void VTPageTableRead_AnalyzeInStage(FMaterialIRValueAnalyzer& Analyzer, MIR::FVTPageTableRead* VTPageTableRead, MIR::EStage Stage)
 {
-	check(VTPageTableRead->TextureObject);
+	const MIR::FVirtualTextureUniform* VirtualTextureUniform = MIR::As<MIR::FVirtualTextureUniform>(VTPageTableRead->VirtualTextureUniform);
+	check(VirtualTextureUniform);
 
-	AccessTexture(Analyzer, VTPageTableRead->TextureObject);
-	const int32 TextureUniformIndex = VTPageTableRead->TextureObject->GetUniformParameterIndex();
+	const int32 TextureUniformIndex = VirtualTextureUniform->Analysis_UniformIndex;
 	check(TextureUniformIndex >= 0);
 
 	// Only support GPU feedback from pixel shader
@@ -214,11 +136,11 @@ static void VTPageTableRead_AnalyzeInStage(FMaterialIRValueAnalyzer& Analyzer, M
 	VTPageTableRead->VTStackIndex = VTStackIndex;
 
 	// Check if VT layer is already known. Otherwise, acquire VT layer
-	int32 VTLayerIndex = GetVTLayerIndex(VTPageTableRead->TextureObject);
+	int32 VTLayerIndex = VirtualTextureUniform->VTLayerIndex;
 	if (VTLayerIndex != INDEX_NONE)
 	{
 		// The layer index in the VT stack is already known, so fetch the page table from the texture object and assign it to the VT stack
-		VTPageTableRead->VTPageTableIndex = GetPageTableLayerIndex(VTPageTableRead->TextureObject);
+		VTPageTableRead->VTPageTableIndex = VirtualTextureUniform->VTPageTableIndex;
 		Analyzer.CompilationOutput->UniformExpressionSet.SetVTLayer(VTStackIndex, VTLayerIndex, TextureUniformIndex);
 	}
 	else
@@ -241,13 +163,6 @@ static void Branch_Analyze(FMaterialIRValueAnalyzer& Analyzer, MIR::FBranch* Bra
 	MIR::ZeroArray(Branch->FalseBlock);
 }
 
-static void TextureRead_Analyze(FMaterialIRValueAnalyzer& Analyzer, MIR::FTextureRead* TextureRead)
-{
-	// Ensure a uniform parameter is allocated when a texture read instruction accesses the texture.
-	// Otherwise, no uniform parameter must be allocated for the texture object in case it's used for other nodes only such as texture properties.
-	AccessTexture(Analyzer, TextureRead->TextureObject);
-}
-
 // Returns the next available offset into the preshader buffer for a float vector with the specified number of components (1-4).
 static uint32 NextGlobalComponentOffset(FMaterialIRValueAnalyzer& Analyzer, const MIR::FPrimitive& PrimitiveType)
 {
@@ -306,51 +221,94 @@ static uint32 NextGlobalComponentOffset(FMaterialIRValueAnalyzer& Analyzer, cons
 	return GlobalComponentOffset;
 }
 
-static uint32 FindOrAddDefaultValueOffset(FMaterialIRValueAnalyzer& Analyzer, const FMaterialParameterValue& ParameterValue)
+static const MIR::FConstant* GetConstantComponent(const MIR::FValue* Constant, int32 Index)
 {
-	UE::Shader::FValue DefaultValue;
-	switch (ParameterValue.Type)
+	if (const MIR::FConstant* Scalar = MIR::As<MIR::FConstant>(Constant))
+	{
+		check(Index == 0);
+		return Scalar;
+	}
+	else
 	{
-		case EMaterialParameterType::Scalar: DefaultValue = ParameterValue.AsScalar(); break;
-		case EMaterialParameterType::Vector: DefaultValue = ParameterValue.AsLinearColor(); break;
-		case EMaterialParameterType::DoubleVector: DefaultValue = ParameterValue.AsVector4d(); break;
-		default: UE_MIR_TODO();
+		return MIR::As<MIR::FConstant>(MIR::As<MIR::FComposite>(Constant)->GetComponents()[Index]);
 	}
+}
 
-	uint32 DefaultValueOffset = 0;
-	if (!MIR::Find(Analyzer.UniformDefaultValueOffsets, DefaultValue, DefaultValueOffset))
+static UE::Shader::FValue ConstantToShaderValue(const MIR::FValue* Constant)
+{
+	UE::Shader::FValue DefaultShaderValue;
+	DefaultShaderValue.Type = Constant->Type.ToValueType();
+
+	MIR::FPrimitive PrimitiveType = Constant->Type.GetPrimitive();
+	for (int i = 0, NumComponents = PrimitiveType.NumComponents(); i < NumComponents; ++i)
 	{
-		FUniformExpressionSet& UniformExpressionSet = Analyzer.CompilationOutput->UniformExpressionSet;
-		DefaultValueOffset = UniformExpressionSet.AddDefaultParameterValue(DefaultValue);
-		Analyzer.UniformDefaultValueOffsets.Add(DefaultValue, DefaultValueOffset);
+		switch (PrimitiveType.ScalarKind)
+		{
+			case MIR::EScalarKind::Boolean: DefaultShaderValue.Component.Add(GetConstantComponent(Constant, i)->Boolean); break;
+			case MIR::EScalarKind::Integer: DefaultShaderValue.Component.Add((int)GetConstantComponent(Constant, i)->Integer); break;
+			case MIR::EScalarKind::Float:   DefaultShaderValue.Component.Add(GetConstantComponent(Constant, i)->Float); break;
+			case MIR::EScalarKind::Double:  DefaultShaderValue.Component.Add(GetConstantComponent(Constant, i)->Double); break;
+		}
 	}
 
-	return DefaultValueOffset;
+	return DefaultShaderValue;
 }
 
-static void AnalyzePrimitiveUniformParameter(FMaterialIRValueAnalyzer& Analyzer, MIR::FUniformParameter* Parameter)
+static void PrimitiveUniform_Analyze(FMaterialIRValueAnalyzer& Analyzer, MIR::FPrimitiveUniform* PrimitiveUniform)
 {
+	MIR::FPrimitive PrimitiveType = PrimitiveUniform->Type.GetPrimitive();
 	FUniformExpressionSet& UniformExpressionSet = Analyzer.CompilationOutput->UniformExpressionSet;
-	FMaterialParameterInfo ParameterInfo = Analyzer.Module->GetParameterInfo(Parameter->ParameterIdInModule);
-	FMaterialParameterMetadata ParameterMetadata = Analyzer.Module->GetParameterMetadata(Parameter->ParameterIdInModule);
-
-	// TODO: GetParameterOverrideValueForCurrentFunction?
 
-	const uint32 DefaultValueOffset = FindOrAddDefaultValueOffset(Analyzer, ParameterMetadata.Value);
+	// Declare the default value on the Uniform Expression Set and retrieve its uniform buffer offset.
+	UE::Shader::FValue DefaultShaderValue = ConstantToShaderValue(PrimitiveUniform->DefaultConstant);
 
-	Parameter->Analysis_UniformParameterIndex = UniformExpressionSet.FindOrAddNumericParameter(ParameterMetadata.Value.Type, ParameterInfo, DefaultValueOffset);
+	uint32 DefaultValueOffset = 0;
+	if (!MIR::Find(Analyzer.UniformPrimitiveDefaultValueOffsets, DefaultShaderValue, DefaultValueOffset))
+	{
+		DefaultValueOffset = UniformExpressionSet.AddDefaultParameterValue(DefaultShaderValue);
+		Analyzer.UniformPrimitiveDefaultValueOffsets.Add(DefaultShaderValue, DefaultValueOffset);
+	}
 
-	// Make sure the parameter type is primitive
-	MIR::FPrimitive PrimitiveType = Parameter->Type.GetPrimitive();
+	// Determine the Material Parameter Type from the primitive uniform's type.
+	// TODO: By the way material parameters work we are confined to only a very limited set of types that we could potentially support.
+	//       We should rework the parameter system to accept any primitive type: for example, we should split parameters into "primitive" parameters
+	//       and all other "object" (e.g. texture) parameters. Then make primitive parameter be of any primitive type and allocate them onto the
+	//       uniform buffer efficiently.
+	EMaterialParameterType MaterialParameterType;
+	if (PrimitiveUniform->Type == MIR::FType::MakeFloatScalar())
+	{
+		MaterialParameterType = EMaterialParameterType::Scalar;
+	}
+	else if (PrimitiveUniform->Type == MIR::FType::MakeFloatVector(4))
+	{
+		MaterialParameterType = EMaterialParameterType::Vector;
+	}
+	else if (PrimitiveUniform->Type == MIR::FType::MakeDoubleVector(4))
+	{
+		MaterialParameterType = EMaterialParameterType::DoubleVector;
+	}
+	else
+	{
+		Analyzer.Module->AddError(nullptr, FString::Printf(TEXT("Primitive uniforms can currently only be of these types: float, float4 or double4. Specified type is '%s'."), *PrimitiveUniform->Type.GetSpelling()));
+		return;
+	}
 
-	// Only int, float and LWC parameters supported for now
-	check(PrimitiveType.IsInteger() || PrimitiveType.IsAnyFloat());
+	// Register the numeric parameter on the Uniform Expression Set.
+	// A few important remarks and findings: FindOrAddNumericParameter` API should be improved upon. It has several "leaks":
+	//    - It uses MaterialParameterXXX which are front-end concepts. It should instead have its own, decoupled way to specify a uniform, and the material should keep a "binding table" between material parameters (frontend) and uniform values (backend).
+	//      This binding table would be generated by the translator and stored in MaterialCompilationOutput for instance.
+	//    - With the change above, the Material runtime would "push" values to the uniform buffer, rather than the latter "pulling" from the Material, in a sort of backwards direction. (Things should always go one direction, from higher level to lower level).
+	//    - UE::Shader::FValue isn't great either. It encompasses ALL possible values, coupling together primitive and "objects" for no reason. FindOrAddNumericParameter should take a "PrimitiveValue" as the default as it makes no sense to specify
+	//      a non primitive Value here, or, even better, take no default value at all, and delegate it to the Material runtime to provide some value, default or overriden by the user. Another example of "backwards" dependency.
+	PrimitiveUniform->Analysis_UniformIndex = UniformExpressionSet.FindOrAddNumericParameter(MaterialParameterType, { FName(PrimitiveUniform->Name) }, DefaultValueOffset);
 
+	// Allocate the space for the parameter on the uniform buffer.
 	const uint32 GlobalComponentOffset = NextGlobalComponentOffset(Analyzer, PrimitiveType);
-
+	 
 	// Add the parameter evaluation to the uniform data
-	UniformExpressionSet.AddNumericParameterEvaluation(Parameter->Analysis_UniformParameterIndex, GlobalComponentOffset);
-
+	UniformExpressionSet.AddNumericParameterEvaluation(PrimitiveUniform->Analysis_UniformIndex, GlobalComponentOffset);
+	 
+	// If an insights instance is provided, add information about this parameter to it.
 	if (Analyzer.Insights)
 	{
 		// Push information about this parameter allocation to the insights.
@@ -358,27 +316,48 @@ static void AnalyzePrimitiveUniformParameter(FMaterialIRValueAnalyzer& Analyzer,
 		ParamInsight.BufferSlotIndex = GlobalComponentOffset / 4;
 		ParamInsight.BufferSlotOffset = GlobalComponentOffset % 4;
 		ParamInsight.ComponentsCount = PrimitiveType.NumComponents();
-		ParamInsight.ParameterName = ParameterInfo.Name;
-
+		ParamInsight.ParameterName = FName(PrimitiveUniform->Name);
+	 
 		switch (PrimitiveType.ScalarKind)
 		{
-			case MIR::EScalarKind::Int:   ParamInsight.ComponentType = FMaterialInsights::FUniformBufferSlotComponentType::CT_Int; break;
-			case MIR::EScalarKind::Float: ParamInsight.ComponentType = FMaterialInsights::FUniformBufferSlotComponentType::CT_Float; break;
-			case MIR::EScalarKind::Double:   ParamInsight.ComponentType = FMaterialInsights::FUniformBufferSlotComponentType::CT_LWC; break;
+			case MIR::EScalarKind::Integer:    ParamInsight.ComponentType = FMaterialInsights::FUniformBufferSlotComponentType::CT_Int; break;
+			case MIR::EScalarKind::Float:  ParamInsight.ComponentType = FMaterialInsights::FUniformBufferSlotComponentType::CT_Float; break;
+			case MIR::EScalarKind::Double: ParamInsight.ComponentType = FMaterialInsights::FUniformBufferSlotComponentType::CT_LWC; break;
 			default: UE_MIR_UNREACHABLE();
 		}
-
+	 
 		Analyzer.Insights->UniformParameterAllocationInsights.Push(ParamInsight);
 	}
 }
 
-static void UniformParameter_Analyze(FMaterialIRValueAnalyzer& Analyzer, MIR::FUniformParameter* Parameter)
+static int32 RegisterTextureUniform(FMaterialIRValueAnalyzer& Analyzer, FName Name, UObject* TextureObject, EMaterialValueType TextureValueType, int32 VTLayerIndex)
 {
-	if (!Parameter->Type.IsTexture())
+	check(VTLayerIndex <= UINT8_MAX);
+
+	// TODO: TextureIndex should *not* rely on MaterialInterface->GetDefaultTextureIdx, as this relies on CachedExpressionData which
+	//       should be removed entirely. Instead, we should investigate what this TextureIndex is, and instead determine it here directly.
+	const FMaterialTextureParameterInfo TextureParameterInfo
 	{
-		check(Parameter->Type.AsPrimitive());
-		AnalyzePrimitiveUniformParameter(Analyzer, Parameter);
-	}
+		.ParameterInfo            = { Name },
+		.TextureIndex             = Analyzer.MaterialInterface->GetDefaultTextureIdx(TextureObject),
+		.SamplerSource            = SSM_FromTextureAsset,
+		.VirtualTextureLayerIndex = (uint8)VTLayerIndex,
+	};
+
+	check(TextureParameterInfo.TextureIndex != INDEX_NONE);
+
+	const EMaterialTextureParameterType ParamType = TextureMaterialValueTypeToParameterType(TextureValueType);
+	return Analyzer.CompilationOutput->UniformExpressionSet.FindOrAddTextureParameter(ParamType, TextureParameterInfo);
+}
+
+static void TextureUniform_Analyze(FMaterialIRValueAnalyzer& Analyzer, MIR::FTextureUniform* TextureUniform)
+{
+	TextureUniform->Analysis_UniformIndex = RegisterTextureUniform(Analyzer, FName(TextureUniform->Name), TextureUniform->Texture, TextureUniform->Texture->GetMaterialType(), UINT8_MAX);
+}
+
+static void VirtualTextureUniform_Analyze(FMaterialIRValueAnalyzer& Analyzer, MIR::FVirtualTextureUniform* VirtualTextureUniform)
+{
+	VirtualTextureUniform->Analysis_UniformIndex = RegisterTextureUniform(Analyzer, FName(VirtualTextureUniform->Name), VirtualTextureUniform->VirtualTexture, MCT_TextureVirtual, VirtualTextureUniform->VTLayerIndex);
 }
 
 static void PreshaderParameter_Analyze(FMaterialIRValueAnalyzer& Analyzer, MIR::FPreshaderParameter* Parameter)
@@ -387,13 +366,13 @@ static void PreshaderParameter_Analyze(FMaterialIRValueAnalyzer& Analyzer, MIR::
 
 	// Get parameter name from source parameter
 	FName SourceParameterName;
-	if (MIR::FUniformParameter* UniformSourceParameter = Parameter->SourceParameter->As<MIR::FUniformParameter>())
+	if (MIR::FTextureUniform* TextureUniform = Parameter->SourceParameter->As<MIR::FTextureUniform>())
 	{
-		SourceParameterName = Analyzer.Module->GetParameterInfo(UniformSourceParameter->ParameterIdInModule).Name;
+		SourceParameterName = FName{ TextureUniform->Name };
 	}
-	else if (MIR::FTextureObject* TextureObject = Parameter->SourceParameter->As<MIR::FTextureObject>())
+	else if (MIR::FVirtualTextureUniform* VirtualTextureUniform = Parameter->SourceParameter->As<MIR::FVirtualTextureUniform>())
 	{
-		SourceParameterName = TextureObject->Texture->GetFName();
+		SourceParameterName = FName{ VirtualTextureUniform->Name };
 	}
 	else
 	{
@@ -491,7 +470,7 @@ void FMaterialIRValueAnalyzer::Setup(UMaterialInterface* InMaterialInterface, FM
 	Module = InModule;
 	Insights = InInsights;
 	CompilationOutput = InCompilationOutput;
-	UniformDefaultValueOffsets = {};
+	UniformPrimitiveDefaultValueOffsets = {};
 	FreeOffsetsPerNumComponents[0] = {};
 	FreeOffsetsPerNumComponents[1] = {};
 	FreeOffsetsPerNumComponents[2] = {};
@@ -513,9 +492,10 @@ void FMaterialIRValueAnalyzer::Analyze(MIR::FValue* Value)
 	{
 		EXPAND_CASE(AnalysisError);
 		EXPAND_CASE(MaterialParameterCollection);
-		EXPAND_CASE(UniformParameter);
+		EXPAND_CASE(PrimitiveUniform);
+		EXPAND_CASE(TextureUniform);
+		EXPAND_CASE(VirtualTextureUniform);
 		EXPAND_CASE(Branch);
-		EXPAND_CASE(TextureRead);
 		EXPAND_CASE(PreshaderParameter);
 		EXPAND_CASE(SetMaterialOutput);
 		EXPAND_CASE(Extern);
diff --git a/Engine/Source/Runtime/Engine/Private/Materials/MaterialIRValueAnalyzer.h b/Engine/Source/Runtime/Engine/Private/Materials/MaterialIRValueAnalyzer.h
--- a/Engine/Source/Runtime/Engine/Private/Materials/MaterialIRValueAnalyzer.h
+++ b/Engine/Source/Runtime/Engine/Private/Materials/MaterialIRValueAnalyzer.h
@@ -67,7 +67,7 @@ struct FMaterialIRValueAnalyzer
 	FMaterialCompilationOutput* CompilationOutput{};
 		
 	// Maps default values to their default-value offset, as used in UniformExpressionSet (e.g. UniformExpressionSet.AddDefaultParameterValue() and UniformExpressionSet.FindOrAddNumericParameter()).
-	TMap<UE::Shader::FValue, uint32> UniformDefaultValueOffsets{};
+	TMap<UE::Shader::FValue, uint32> UniformPrimitiveDefaultValueOffsets{};
 		
 	// Stores free uniform buffer offsets for 1, 2, and 3 leftover components to optimize float4 packing
 	TArray<uint32, TInlineAllocator<8>> FreeOffsetsPerNumComponents[3];
diff --git a/Engine/Source/Runtime/Engine/Public/Materials/MaterialExpressionCurveAtlasRowParameter.h b/Engine/Source/Runtime/Engine/Public/Materials/MaterialExpressionCurveAtlasRowParameter.h
--- a/Engine/Source/Runtime/Engine/Public/Materials/MaterialExpressionCurveAtlasRowParameter.h
+++ b/Engine/Source/Runtime/Engine/Public/Materials/MaterialExpressionCurveAtlasRowParameter.h
@@ -36,6 +36,7 @@ class UMaterialExpressionCurveAtlasRowParameter : public UMaterialExpressionScal
 	}
 
 	virtual bool IsInputConnectionRequired(int32 InputIndex) const override { return true; }
+	virtual void Build(MIR::FEmitter& Em) override;
 	virtual int32 Compile(class FMaterialCompiler* Compiler, int32 OutputIndex) override;
 	virtual void GetTexturesForceMaterialRecompile(TArray<UTexture *> &Textures) const override;
 	virtual void PostEditChangeProperty(FPropertyChangedEvent& PropertyChangedEvent) override;
diff --git a/Engine/Source/Runtime/Engine/Public/Materials/MaterialExpressionParameter.h b/Engine/Source/Runtime/Engine/Public/Materials/MaterialExpressionParameter.h
--- a/Engine/Source/Runtime/Engine/Public/Materials/MaterialExpressionParameter.h
+++ b/Engine/Source/Runtime/Engine/Public/Materials/MaterialExpressionParameter.h
@@ -58,7 +58,6 @@ class UMaterialExpressionParameter : public UMaterialExpression
 		return true;
 	}
 
-	virtual void Build(MIR::FEmitter& Emitter) override;
 #endif
 	virtual bool IsAllowedIn(const UObject* MaterialOrFunction) const
 	{
diff --git a/Engine/Source/Runtime/Engine/Public/Materials/MaterialExpressionScalarParameter.h b/Engine/Source/Runtime/Engine/Public/Materials/MaterialExpressionScalarParameter.h
--- a/Engine/Source/Runtime/Engine/Public/Materials/MaterialExpressionScalarParameter.h
+++ b/Engine/Source/Runtime/Engine/Public/Materials/MaterialExpressionScalarParameter.h
@@ -65,6 +65,7 @@ class UMaterialExpressionScalarParameter : public UMaterialExpressionParameter
 	//~ End UObject Interface
 
 	//~ Begin UMaterialExpression Interface
+	virtual void Build(MIR::FEmitter& Emitter) override;
 	virtual int32 Compile(class FMaterialCompiler* Compiler, int32 OutputIndex) override;
 	virtual void GetCaption(TArray<FString>& OutCaptions) const override;
 	virtual bool GetParameterValue(FMaterialParameterMetadata& OutMeta) const override;
diff --git a/Engine/Source/Runtime/Engine/Public/Materials/MaterialExpressionTextureBase.h b/Engine/Source/Runtime/Engine/Public/Materials/MaterialExpressionTextureBase.h
--- a/Engine/Source/Runtime/Engine/Public/Materials/MaterialExpressionTextureBase.h
+++ b/Engine/Source/Runtime/Engine/Public/Materials/MaterialExpressionTextureBase.h
@@ -62,32 +62,13 @@ public:
 	 */
 	ENGINE_API void AutoSetSampleType();
 
-	/**
-	 * Returns the default sampler type for the specified texture.
-	 * @param Texture - The texture for which the default sampler type will be returned.
-	 * @returns the default sampler type for the specified texture.
-	 */
-	static ENGINE_API EMaterialSamplerType GetSamplerTypeForTexture( const UTexture* Texture, bool ForceNoVT = false );
+	UE_DEPRECATED(5.8, "Use MaterialExpressionUtils::GetSamplerTypeForTexture instead.")
+	static ENGINE_API EMaterialSamplerType GetSamplerTypeForTexture(const UTexture* Texture, bool ForceNoVT = false);
 
-	static ENGINE_API bool VerifySamplerType(
-		const FString& TexturePathName,
-		EMaterialSamplerType CorrectSamplerType,
-		bool bSRGB,
-		EMaterialSamplerType SamplerType,
-		FString& OutErrorMessage);
+	UE_DEPRECATED(5.8, "Use MaterialExpressionUtils::VerifySamplerType instead.")
+	static ENGINE_API bool VerifySamplerType(const FString& TexturePathName, EMaterialSamplerType CorrectSamplerType, bool bSRGB, EMaterialSamplerType SamplerType, FString& OutErrorMessage);
 
-	/**
-	 * Verify that the texture and sampler type. Generates a compiler waring if
-	 * they do not.
-	 * @param Texture - The texture to verify. A nullptr texture is considered valid!
-	 * @param SamplerType - The sampler type to verify.
-	 * @param OutErrorMessage - If 'false' is returned, will contain a message describing the error
-	 */
-	static ENGINE_API bool VerifySamplerType(
-		EShaderPlatform ShaderPlatform,
-		const ITargetPlatform* TargetPlatform,
-		const UTexture* Texture,
-		EMaterialSamplerType SamplerType,
-		FString& OutErrorMessage);
+	UE_DEPRECATED(5.8, "Use MaterialExpressionUtils::VerifySamplerType instead.")
+	static ENGINE_API bool VerifySamplerType(EShaderPlatform ShaderPlatform,const ITargetPlatform* TargetPlatform,const UTexture* Texture,EMaterialSamplerType SamplerType,FString& OutErrorMessage);
 #endif // WITH_EDITOR
 };
diff --git a/Engine/Source/Runtime/Engine/Public/Materials/MaterialIR.h b/Engine/Source/Runtime/Engine/Public/Materials/MaterialIR.h
--- a/Engine/Source/Runtime/Engine/Public/Materials/MaterialIR.h
+++ b/Engine/Source/Runtime/Engine/Public/Materials/MaterialIR.h
@@ -50,9 +50,9 @@ enum EValueKind
 	VK_Builtin,
 	VK_MaterialParameterCollection,
 	VK_ShadingModel,
-	VK_TextureObject,
-	VK_RuntimeVirtualTextureObject,
-	VK_UniformParameter,
+	VK_PrimitiveUniform,
+	VK_TextureUniform,
+	VK_VirtualTextureUniform,
 	VK_GetVertexInterpolator,
 
 	// Instructions
@@ -191,6 +191,9 @@ struct FValue
 
 	// Returns whether this value is arithmetic and all components are approximately one.
 	bool AreAllNearlyOne() const;
+
+	//
+	bool IsConstant() const;
 	
 	// Returns whether this value exactly equals Other.
 	bool Equals(const FValue* Other) const;
@@ -201,11 +204,11 @@ struct FValue
 	// Returns whether this value equals the given constant vector.
 	bool EqualsConstant(FVector4f Value) const;
 
-	// Returns the UTexture or URuntimeVirtualTexture object if this value is of type FTextureObject or FRuntimeVirtualTextureObject. Returns null otherwise.
+	// Returns the UTexture or URuntimeVirtualTexture object if this value references such texture. Returns null otherwise.
 	UObject* GetTextureObject() const;
 
-	// Returns the uniform parameter index if this value is of type FTextureObject, FRuntimeVirtualTextureObject, or FUniformParameter. Returns INDEX_NONE otherwise.
-	int32 GetUniformParameterIndex() const;
+	// Returns the uniform index if this value associated to this value if its a uniform (primitive, texture, etc). Returns INDEX_NONE otherwise.
+	int32 GetUniformIndex() const;
 
 	// Tries to cast this value to specified type T and returns the casted pointer, if possible (nullptr otherwise).
 	template <typename T>
@@ -351,63 +354,59 @@ struct FMaterialParameterCollection : TValue<VK_MaterialParameterCollection>
 
 	// Index of this collection in the list of collections referenced by the material.
 	// Note: This field is automatically set by the builder.
-	uint32 Analysis_CollectionIndex;
+	int32 Analysis_CollectionIndex;
 };
 
-// It represents a texture object value.
-struct FTextureObject : TValue<VK_TextureObject>
+// An externally provided primitive uniform value, referenced by a given name.
+struct FPrimitiveUniform : TValue<VK_PrimitiveUniform>
 {
-	// The texture object.
-	UTexture* Texture;
+	// A uniquely identifying name.
+	FScriptName Name; // Is FMinimalName ok too?
 
-	// The sampler type associated to this texture object.
-	EMaterialSamplerType SamplerType;
+	// A *constant* value this uniform is initialized to if left unassigned at runtime.
+	FValue* DefaultConstant;
 
-	// Index of this parameter in the uniform expression set.
-	// Note: This field is automatically set by the builder.
-	int32 Analysis_UniformParameterIndex;
+	// Index of this uniform in the uniform expression set generated during analysis.
+	int16 Analysis_UniformIndex;
 };
 
-// It represents a runtime virtual texture (RVT) object value. Those *do not* extend UTexture so they have their own MIR value.
-struct FRuntimeVirtualTextureObject : TValue<VK_RuntimeVirtualTextureObject>
+// It represents a texture uniform value.
+struct FTextureUniform : TValue<VK_TextureUniform>
 {
-	// The runtime virtual texture object.
-	URuntimeVirtualTexture* RVTexture;
+	// A uniquely identifying name.
+	FScriptName Name; // Is FMinimalName ok too?
 
-	// The sampler type associated to this RVT object.
-	EMaterialSamplerType SamplerType;
-
-	// Field to explicitly assign a VT stack layer index (INDEX_NONE if unset). Regular VTs are assigned these indices automatically.
-	int32 VTLayerIndex : 16;
+	// The texture object.
+	UTexture* Texture;
 
-	// Field to explicitly assign a VT page table index. Regular VTs are assigned these indices automatically.
-	int32 VTPageTableIndex : 16;
+	// The sampler type associated to this texture object.
+	EMaterialSamplerType SamplerType : 16;
 
-	// Index of this parameter in the uniform expression set.
-	// Note: This field is automatically set by the builder.
-	int32 Analysis_UniformParameterIndex;
+	// Index of this uniform in the uniform expression set generated during analysis.
+	int16 Analysis_UniformIndex;
 };
 
-// It represents a material uniform parameter.
-struct FUniformParameter: TValue<VK_UniformParameter>
+// It represents a runtime virtual texture (RVT) uniform value.
+struct FVirtualTextureUniform : TValue<VK_VirtualTextureUniform>
 {
-	// Index of the parameter as registered in the module.
-	uint32 ParameterIdInModule;
+	// A uniquely identifying name.
+	FScriptName Name; // Is FMinimalName ok too?
 
-	// Eventual sampler type to use to sample this parameter's texture (if it is one).
-	EMaterialSamplerType SamplerType;
+	// The runtime virtual texture object.
+	// Note: URuntimeVirtualTexture does not extend UTexture.
+	URuntimeVirtualTexture* VirtualTexture;
 
 	// Field to explicitly assign a VT stack layer index (INDEX_NONE if unset). Regular VTs are assigned these indices automatically.
-	int32 VTLayerIndex : 16;
+	int16 VTLayerIndex;
 
 	// Field to explicitly assign a VT page table index. Regular VTs are assigned these indices automatically.
-	int32 VTPageTableIndex : 16;
+	uint16 VTPageTableIndex;
 
-	// Index of this parameter in the uniform expression set
-	// Note: This field is automatically set by the builder.
-	int32 Analysis_UniformParameterIndex;
+	// Index of this parameter in the uniform expression set generated during analysis.
+	int16 Analysis_UniformIndex;
 };
 
+// A value that identifies a specific shading model
 struct FShadingModel : TValue<VK_ShadingModel>
 {
 	EMaterialShadingModel Id;
@@ -846,8 +845,8 @@ const TCHAR* LexToString(ETextureReadMode Mode);
 // This instruction performs texture read operaation (sample or gather).
 struct FTextureRead : TInstruction<VK_TextureRead, 6>
 {
-	// The texture object to sample. Can be FTextureObject or FUniformParameter.
-	FValue* TextureObject;
+	// The texture to sample.
+	FValue* TextureUniform;
 
 	// The texture coordinate at which to sample.
 	FValue* TexCoord;
@@ -881,7 +880,7 @@ struct FTextureRead : TInstruction<VK_TextureRead, 6>
 struct FVTPageTableRead : TInstruction<VK_VTPageTableRead, 5>
 {
 	// The texture object to load the page table for.
-	FValue* TextureObject;
+	FValue* VirtualTextureUniform;
 
 	// The texture coordinate at which to sample.
 	FValue* TexCoord;
@@ -910,7 +909,7 @@ struct FVTPageTableRead : TInstruction<VK_VTPageTableRead, 5>
 	// Specifies the sampler function to read the page table, i.e. TextureLoadVirtualPageTable(), TextureLoadVirtualPageTableGrad(), or TextureLoadVirtualPageTableLevel().
 	ETextureMipValueMode MipValueMode : 4;
 
-	// Enables VT texture sampling feedback
+	// Specifies whether to enable virtual texture sampling feedback. This is only used for pixel shaders. Enabled by default.
 	bool bEnableFeedback : 1;
 
 	// Determines whether adaptive or regular VT page table loads will be emitted, i.e. TextureLoadVirtualPageTable*() or TextureLoadVirtualPageTableAdaptive*().
diff --git a/Engine/Source/Runtime/Engine/Public/Materials/MaterialIREmitter.h b/Engine/Source/Runtime/Engine/Public/Materials/MaterialIREmitter.h
--- a/Engine/Source/Runtime/Engine/Public/Materials/MaterialIREmitter.h
+++ b/Engine/Source/Runtime/Engine/Public/Materials/MaterialIREmitter.h
@@ -175,7 +175,7 @@ struct FFunctionDesc
 	uint32 GetNumParameters() const { return NumInputOnlyParams + NumInputOutputParams + NumOutputOnlyParams; }
 };
 
-//
+// Describes a user-specified HLSL function.
 struct FFunctionHLSLDesc : FFunctionDesc
 {
 	// The user-specified HLSL code inside the function body.
@@ -188,8 +188,8 @@ struct FTextureSampleAttributes
 	// Specifies which sampler source mode to use. By default SSM_FromTextureAsset.
 	ESamplerSourceMode SamplerSourceMode { SSM_FromTextureAsset };
 	
-	// Specifies Which sampler type to use. By default SAMPLERTYPE_MAX which implies the sampler type is fetched from the texture object.
-	// See SamplerType field in FTextureObject, FRuntimeVirtualTextureObject, and FUniformParameter.
+	// Specifies which sampler type to use. By default SAMPLERTYPE_MAX which implies the sampler type is fetched from the texture object.
+	// See SamplerType field in FTextureUniform or FVirtualTextureUniform.
 	EMaterialSamplerType SamplerType { SAMPLERTYPE_MAX };
 
 	// Specifies whether to enable virtual texture sampling feedback. This is only used for pixel shaders. Enabled by default.
@@ -333,9 +333,6 @@ public:
 	// Returns the type of the i-th attribute in the given material aggregate.
 	FType GetMaterialAggregateAttributeType(const UMaterialAggregate* Aggregate, int32 AttributeIndex);
 
-	// Returns the UTexture or URuntimeVirtualTexture of the specified texture object or uniform parameter.
-	UObject* GetTextureFromValue(MIR::FValueRef Texture);
-
 	/*-------------------------------- Input management --------------------------------*/
 
 	// Retrieves the value from the given input, if present. A null reference otherwise.
@@ -496,6 +493,15 @@ public:
 	// Returns the given constant LWC scalar.
 	FValueRef ConstantDouble(TDouble InX);
 
+	// Emits a constant double 2D column vector and returns it.
+	FValueRef ConstantDouble2(UE::Math::TVector2<TDouble> InValue);
+
+	// Emits a constant double 3D column vector and returns it.
+	FValueRef ConstantDouble3(UE::Math::TVector<TDouble> InValue);
+
+	// Emits a constant double 4D column vector and returns it.
+	FValueRef ConstantDouble4(UE::Math::TVector4<TDouble> InValue);
+
 	/*--------------------- Other non-instruction values emission ---------------------*/
 
 	// Emits the poison value.
@@ -514,14 +520,18 @@ public:
 	// Emits a material parameter collection.
 	FValueRef MaterialParameterCollection(UMaterialParameterCollection* Collection);
 
-	// Emits a texture object value of fiven sampler type and returns it.
-	FValueRef TextureObject(UTexture* Texture, EMaterialSamplerType SamplerType);
+	// Emits a named uniform value of primitive type.
+	// The uniform will be declared onto the MaterialCompilationOutput UniformExpressionSet as a numeric parameter.
+	FValueRef NamedPrimitiveUniform(FName Name, FValueRef DefaultConstant);
 
-	// Emits a runtime virtual texture object value.
-	FValueRef RuntimeVirtualTextureObject(URuntimeVirtualTexture* RVTexture, EMaterialSamplerType SamplerType, int32 VTLayerIndex, int32 VTPageTableIndex);
+	// Emits a texture uniform value of any type except for runtime virtual textures.
+	// - SamplerType specifies the default sampler associated to this texture.
+	FValueRef TextureUniform(FName Name, UTexture* Texture, EMaterialSamplerType SamplerType);
 
-	// Emits a material parameter value with given name, metadata, sampler type (if it's a texture) and optional parameters for RVT layer index and page table index.
-	FValueRef Parameter(FName Name, FMaterialParameterMetadata& Metadata, EMaterialSamplerType SamplerType = SAMPLERTYPE_Color, int32 VTLayerIndex = INDEX_NONE, int32 VTPageTableIndex = INDEX_NONE);
+	// Emits a runtime virtual texture uniform value.
+	// - VTLayerIndex specifies VT stack layer index. If set to INDEX_NONE, it is assigned automatically from the uniform expression set.
+	// - VTPageTableIndex specifies the VT page table index.
+	FValueRef VirtualTextureUniform(FName Name, URuntimeVirtualTexture* VirtualTexture, int32 VTLayerIndex, int32 VTPageTableIndex);
 
 	// Emits a preshader parameter value with given opcode.
 	FValueRef PreshaderParameter(FType Type, EPreshader2Opcode Opcode, FValueRef SourceParameter, FPreshaderParameterPayload Payload = {});
@@ -638,7 +648,7 @@ public:
 
 	// Emits a texture gather instruction.
 	// 
-	// @param Texture The texture to sample. Must be of type MIR::FTextureObject, MIR::FRuntimeVirtualTextureObject, or MIR::FUniformParameter.
+	// @param Texture The texture to sample. Must be a texture or virtual texture uniform value.
 	// @param TexCoord The UV coordinates at which to sample the texture.
 	// @param GatherMode What channel to gather (any ::GatherXXX value).
 	// @param BaseAttributes Base sampling attributes for sampler source and type as well as VT attributes.
@@ -646,7 +656,7 @@ public:
 
 	// Emits a texture sample instruction.
 	// 
-	// @param Texture The texture to sample. Must be of type MIR::FTextureObject, MIR::FRuntimeVirtualTextureObject, or MIR::FUniformParameter.
+	// @param Texture The texture to sample. Must be a texture or virtual texture uniform value.
 	// @param TexCoord The UV coordinates at which to sample the texture.
 	// @param bAddViewMaterialMipBias Whether to add the view material texture mip bias to the mip level.
 	// @param BaseAttributes Base sampling attributes for sampler source and type as well as VT attributes.
@@ -654,7 +664,7 @@ public:
 
 	// Emits a texture sample instruction with a manually given mip level.
 	// 
-	// @param Texture The texture to sample. Must be of type MIR::FTextureObject, MIR::FRuntimeVirtualTextureObject, or MIR::FUniformParameter.
+	// @param Texture The texture to sample. Must be a texture or virtual texture uniform value.
 	// @param TexCoord The UV coordinates at which to sample the texture.
 	// @param MipLevel Which mip level to use for this sample.
 	// @param bAddViewMaterialMipBias Whether to add the view material texture mip bias to the mip level.
@@ -663,7 +673,7 @@ public:
 
 	// Emits a texture sample instruction with bias added to the automatically selected mip level.
 	// 
-	// @param Texture The texture to sample. Must be of type MIR::FTextureObject, MIR::FRuntimeVirtualTextureObject, or MIR::FUniformParameter.
+	// @param Texture The texture to sample. Must be a texture or virtual texture uniform value.
 	// @param TexCoord The UV coordinates at which to sample the texture.
 	// @param MipBias The bias to add to the automatically selected mip level.
 	// @param bAddViewMaterialMipBias Whether to add the view material texture mip bias to the mip level.
@@ -672,7 +682,7 @@ public:
 	
 	// Emits a texture sample instruction using user provided partial derivatives.
 	// 
-	// @param Texture The texture to sample. Must be of type MIR::FTextureObject, MIR::FRuntimeVirtualTextureObject, or MIR::FUniformParameter.
+	// @param Texture The texture to sample. Must be a texture or virtual texture uniform value.
 	// @param TexCoord The UV coordinates at which to sample the texture.
 	// @param TexCoordDdx The partial derivative of the texture coordinates along the screen space x axis.
 	// @param TexCoordDdy The partial derivative of the texture coordinates along the screen space y axis.
@@ -680,18 +690,7 @@ public:
 	// @param BaseAttributes Base sampling attributes for sampler source and type as well as VT attributes.
 	FValueRef TextureSampleGrad(FValueRef Texture, FValueRef TexCoord, FValueRef TexCoordDdx, FValueRef TexCoordDdy, bool bAddViewMaterialMipBias, const FTextureSampleAttributes& BaseAttributes);
 
-	// Emits the virtual texture page load instruction.
-	//
-	// @param Texture The texture to sample. Must be of type MIR::FTextureObject, MIR::FRuntimeVirtualTextureObject, or MIR::FUniformParameter.
-	// @param AddressU Texture address mode for the U-coordinate.
-	// @param AddressV Texture address mode for the V-coordinate.
-	// @param TexCoord The UV coordinates at which to sample the virtual texture.
-	// @param TexCoordDdx The partial derivative of the texture coordinates along the screen space x axis.
-	// @param TexCoordDdy The partial derivative of the texture coordinates along the screen space y axis.
-	// @param bEnableFeedback Specifies whether to enable virtual texture sampling feedback. This is only used for pixel shaders. Enabled by default.
-	// @param bIsAdaptive Specifies whether the page table instruction corresponds to an adaptive RVT.
-	// @param MipValueMode Specifies the MIP value mode. This determines what page table load function in the shader will be emitted, i.e. TextureLoadVirtualPageTable(), TextureLoadVirtualPageTableLevel(), etc..
-	// @param MipValue Optional value for an explicit MIP level or bias.
+	// TODO
 	FValueRef VTPageTableLoad(
 		FValueRef Texture, TextureAddress AddressU, TextureAddress AddressV, FValueRef TexCoord,
 		FValueRef TexCoordDdx = {}, FValueRef TexCoordDdy = {}, bool bEnableFeedback = true, bool bIsAdaptive = false, ETextureMipValueMode MipValueMode = TMVM_None, FValueRef MipValue = {});
@@ -874,6 +873,9 @@ public:
 	// Returns the source material interface. 
 	const UMaterialInterface* GetMaterialInterface() const { return MaterialInterface; }
 
+	//
+	const FStaticParameterSet* GetStaticParameterSet() const { return StaticParameterSet; }
+
 	// Returns the Substrate translator data in use.
 	Substrate::FSubstrateTranslatorData* GetSubstrateTranslatorData() const;
 
diff --git a/Engine/Source/Runtime/Engine/Public/Materials/MaterialIRExtern.cpp b/Engine/Source/Runtime/Engine/Public/Materials/MaterialIRExtern.cpp
--- a/Engine/Source/Runtime/Engine/Public/Materials/MaterialIRExtern.cpp
+++ b/Engine/Source/Runtime/Engine/Public/Materials/MaterialIRExtern.cpp
@@ -77,8 +77,17 @@ void FExternFromMaterialDecl::Analyze(FExternAnalysisContext& Context)
 	// Validate this external code can be used for the current material domain. Empty list implies no  restriction on material domains.
 	if (!Declaration->Domains.IsEmpty() && Declaration->Domains.Find(Context.MaterialInterface->GetMaterial()->MaterialDomain) == INDEX_NONE)
 	{
-		FName AssetPathName = Context.MaterialInterface->GetOutermost()->GetFName();
-		Context.Module->AddError(nullptr, MaterialExpressionUtils::FormatUnsupportedMaterialDomainError(*Declaration, AssetPathName));
+		FString ErrorMessage;
+		ErrorMessage.Appendf(TEXT("External code '%s' (asset: %s) is only available in the following material domains: "), *Declaration->Name.ToString(), *Context.MaterialInterface->GetOutermost()->GetFName().ToString());
+		for (int32 i = 0; i < Declaration->Domains.Num(); ++i)
+		{
+			if (i > 0)
+			{
+				ErrorMessage.Append(TEXT(", "));
+			}
+			ErrorMessage.Append(MaterialDomainString(Declaration->Domains[i]));
+		}
+		Context.Module->AddError(nullptr, MoveTemp(ErrorMessage));
 	}
 
 	// Cast from material feature level enum to RHI feature level enum
@@ -207,16 +216,16 @@ static MIR::FType ToType(EExternSimpleType Id)
 {
 	switch (Id)
 	{
-		case EExternSimpleType::Bool1:			return MIR::FType::MakePrimitive(EScalarKind::Bool, 1, 1);
-		case EExternSimpleType::Bool2:			return MIR::FType::MakePrimitive(EScalarKind::Bool, 1, 2);
-		case EExternSimpleType::Bool3:			return MIR::FType::MakePrimitive(EScalarKind::Bool, 1, 3);
-		case EExternSimpleType::Bool4:			return MIR::FType::MakePrimitive(EScalarKind::Bool, 1, 4);
-		case EExternSimpleType::Int1:			return MIR::FType::MakePrimitive(EScalarKind::Int, 1, 1);
-		case EExternSimpleType::Int2:			return MIR::FType::MakePrimitive(EScalarKind::Int, 1, 2);
-		case EExternSimpleType::Int3:			return MIR::FType::MakePrimitive(EScalarKind::Int, 1, 3);
-		case EExternSimpleType::Int4:			return MIR::FType::MakePrimitive(EScalarKind::Int, 1, 4);
-		case EExternSimpleType::Int3x3:			return MIR::FType::MakePrimitive(EScalarKind::Int, 3, 3);
-		case EExternSimpleType::Int4x4:			return MIR::FType::MakePrimitive(EScalarKind::Int, 4, 4);
+		case EExternSimpleType::Bool1:			return MIR::FType::MakePrimitive(EScalarKind::Boolean, 1, 1);
+		case EExternSimpleType::Bool2:			return MIR::FType::MakePrimitive(EScalarKind::Boolean, 1, 2);
+		case EExternSimpleType::Bool3:			return MIR::FType::MakePrimitive(EScalarKind::Boolean, 1, 3);
+		case EExternSimpleType::Bool4:			return MIR::FType::MakePrimitive(EScalarKind::Boolean, 1, 4);
+		case EExternSimpleType::Int1:			return MIR::FType::MakePrimitive(EScalarKind::Integer, 1, 1);
+		case EExternSimpleType::Int2:			return MIR::FType::MakePrimitive(EScalarKind::Integer, 1, 2);
+		case EExternSimpleType::Int3:			return MIR::FType::MakePrimitive(EScalarKind::Integer, 1, 3);
+		case EExternSimpleType::Int4:			return MIR::FType::MakePrimitive(EScalarKind::Integer, 1, 4);
+		case EExternSimpleType::Int3x3:			return MIR::FType::MakePrimitive(EScalarKind::Integer, 3, 3);
+		case EExternSimpleType::Int4x4:			return MIR::FType::MakePrimitive(EScalarKind::Integer, 4, 4);
 		case EExternSimpleType::Float1:			return MIR::FType::MakePrimitive(EScalarKind::Float, 1, 1);
 		case EExternSimpleType::Float2:			return MIR::FType::MakePrimitive(EScalarKind::Float, 1, 2);
 		case EExternSimpleType::Float3:			return MIR::FType::MakePrimitive(EScalarKind::Float, 1, 3);
diff --git a/Engine/Source/Runtime/Engine/Public/Materials/MaterialIRExtern.h b/Engine/Source/Runtime/Engine/Public/Materials/MaterialIRExtern.h
--- a/Engine/Source/Runtime/Engine/Public/Materials/MaterialIRExtern.h
+++ b/Engine/Source/Runtime/Engine/Public/Materials/MaterialIRExtern.h
@@ -517,6 +517,13 @@ struct FExternPrinterHLSL
 	}
 
 	FExternPrinterHLSL& operator<<(const TCHAR* String);
+
+	template <typename... Types>
+	FExternPrinterHLSL& Appendf(UE::Core::TCheckedFormatString<TCHAR, Types...> Fmt, Types... Args)
+	{
+		Buffer.Appendf(Fmt, Args...);
+		return *this;
+	}
 };
 
 } // namespace MIR
diff --git a/Engine/Source/Runtime/Engine/Public/Materials/MaterialIRTypes.h b/Engine/Source/Runtime/Engine/Public/Materials/MaterialIRTypes.h
--- a/Engine/Source/Runtime/Engine/Public/Materials/MaterialIRTypes.h
+++ b/Engine/Source/Runtime/Engine/Public/Materials/MaterialIRTypes.h
@@ -35,10 +35,10 @@ const TCHAR* TypeKindToString(ETypeKind Kind);
 // Note: These are listed in precision order. Converting one to the other is then simply performed taking the max EScalarKind.
 enum class EScalarKind : uint8
 {
-	Bool,	// Boolean type (true/false).
-	Int,	// 32-bit signed integer type.
-	Float,	// 32-bit floating-point type.
-	Double,	// Large World Coordinates type.
+	Boolean,	// Boolean type (true/false).
+	Integer,	// 32-bit signed integer type.
+	Float,		// 32-bit floating-point type.
+	Double,		// Large World Coordinates type.
 };
 
 // Returns the larger scalar kind between A and B.
@@ -87,13 +87,13 @@ struct FPrimitive
 	bool IsMatrix() const { return NumRows > 1 && NumColumns > 1; }
 	
 	// Returns whether this is a boolean primitive type (of any dimension).
-	bool IsBoolean() const { return ScalarKind == EScalarKind::Bool; }
+	bool IsBoolean() const { return ScalarKind == EScalarKind::Boolean; }
 	
 	// Returns whether this is a single boolean scalar.
 	bool IsBoolScalar() const { return IsBoolean() && IsScalar(); }
 
 	// Returns whether this is an integer primitive type (of any dimension).
-	bool IsInteger() const { return ScalarKind == EScalarKind::Int; }
+	bool IsInteger() const { return ScalarKind == EScalarKind::Integer; }
 	
 	// Returns whether this is a float primitive type (of any dimension).
 	bool IsFloat() const { return ScalarKind == EScalarKind::Float; }
@@ -144,10 +144,10 @@ struct FType
 	static FType MakeScalar(EScalarKind InScalarKind) { return MakePrimitive(InScalarKind, 1, 1); }
 
 	// Returns a scalar type representing a boolean value.
-	static FType MakeBoolScalar() { return MakeScalar(EScalarKind::Bool); }
+	static FType MakeBoolScalar() { return MakeScalar(EScalarKind::Boolean); }
 
 	// Returns a scalar type representing a integer value.
-	static FType MakeIntScalar() { return MakeScalar(EScalarKind::Int); }
+	static FType MakeIntScalar() { return MakeScalar(EScalarKind::Integer); }
 
 	// Returns a scalar type representing a single precision floating point value.
 	static FType MakeFloatScalar() { return MakeScalar(EScalarKind::Float); }
@@ -159,10 +159,10 @@ struct FType
 	static FType MakeVector(EScalarKind InScalarKind, int NumColumns) { return MakePrimitive(InScalarKind, 1, NumColumns); }
 
 	// Returns a row-major vector type of boolean values with the specified number of columns.
-	static FType MakeBoolVector(int NumColumns) { return MakeVector(EScalarKind::Bool, NumColumns); }
+	static FType MakeBoolVector(int NumColumns) { return MakeVector(EScalarKind::Boolean, NumColumns); }
 
 	// Returns a row-major vector type of integer values with the specified number of columns.
-	static FType MakeIntVector(int NumColumns) { return MakeVector(EScalarKind::Int, NumColumns); }
+	static FType MakeIntVector(int NumColumns) { return MakeVector(EScalarKind::Integer, NumColumns); }
 
 	// Returns a row-major vector type of single precision floating point values with the specified number of columns.
 	static FType MakeFloatVector(int NumColumns) { return MakeVector(EScalarKind::Float, NumColumns); }
@@ -171,7 +171,7 @@ struct FType
 	static FType MakeDoubleVector(int NumColumns) { return MakeVector(EScalarKind::Double, NumColumns); }
 	
 	// Creates an integer primitive type with given dimensions.
-	static FType MakeInt(int NumRows, int NumColumns) { return MakePrimitive(EScalarKind::Int, NumRows, NumColumns); }
+	static FType MakeInt(int NumRows, int NumColumns) { return MakePrimitive(EScalarKind::Integer, NumRows, NumColumns); }
 
 	// Creates a single precision float primitive type with given dimensions.
 	static FType MakeFloat(int NumRows, int NumColumns) { return MakePrimitive(EScalarKind::Float, NumRows, NumColumns); }
