Fix bug where changes to r.SetRes during a loading screen do not take effect.
- Some loading screen implementations change the Slate widget hierarchy, which can prevent r.SetRes from finding the window to resize in FSceneViewport::ResizeFrame.
- If SceneViewport fails to retrieve the SWindow from the ViewportWidget, also check the UGameEngine instance.

#Synced-CL 51451766

diff --git a/Engine/Source/Runtime/Engine/Private/Slate/SceneViewport.cpp b/Engine/Source/Runtime/Engine/Private/Slate/SceneViewport.cpp
--- a/Engine/Source/Runtime/Engine/Private/Slate/SceneViewport.cpp
+++ b/Engine/Source/Runtime/Engine/Private/Slate/SceneViewport.cpp
@@ -4,6 +4,7 @@
 #include "Slate/SceneViewport.h"
 #include "Engine/GameInstance.h"
 #include "Engine/GameViewportClient.h"
+#include "Engine/GameEngine.h"
 #include "Widgets/SViewport.h"
 #include "EngineLogs.h"
 #include "Misc/App.h"
@@ -465,8 +466,7 @@ void FSceneViewport::OnDrawViewport( const FGeometry& AllottedGeometry, const FS
 
 	    if (GetSizeXY() != DrawSize || bHDRStale)
 		{
-			TSharedPtr<SWindow> Window = FSlateApplication::Get().FindWidgetWindow(ViewportWidget.Pin().ToSharedRef());
-			if (Window.IsValid())
+			if (TSharedPtr<SWindow> Window = FindWindow())
 			{
 				//@HACK VREDITOR
 				//check(Window.IsValid());
@@ -494,8 +494,7 @@ void FSceneViewport::OnDrawViewport( const FGeometry& AllottedGeometry, const FS
 		// This ensures the spectator view will be stretched to fit to the newly resized window without changing the size of the HMD render target.
 		else if (GIsEditor && IsStereoRenderingAllowed())
 		{
-			TSharedPtr<SWindow> Window = FSlateApplication::Get().FindWidgetWindow(ViewportWidget.Pin().ToSharedRef());
-			if (Window.IsValid())
+			if (TSharedPtr<SWindow> Window = FindWindow())
 			{
 				ensure(PaintWindow == Window.Get());
 
@@ -511,8 +510,7 @@ bool FSceneViewport::IsForegroundWindow() const
 	bool bIsForeground = false;
 	if( ViewportWidget.IsValid() )
 	{
-		TSharedPtr<SWindow> Window = FSlateApplication::Get().FindWidgetWindow( ViewportWidget.Pin().ToSharedRef() );
-		if( Window.IsValid() )
+		if (TSharedPtr<SWindow> Window = FindWindow())
 		{
 			bIsForeground = Window->GetNativeWindow()->IsForegroundWindow();
 		}
@@ -669,7 +667,7 @@ FReply FSceneViewport::AcquireFocusAndCapture(FIntPoint MousePosition, EFocusCau
 			
 			// The slate app will correct mouse positions for non-standard screen / viewport resolution combos
 			// We want to save the mouse position pre-correction so it isn't applied twice when restoring mouse position
-			TSharedPtr<SWindow> Window = FSlateApplication::Get().FindWidgetWindow(ViewportWidgetRef);
+			TSharedPtr<SWindow> Window = FindWindow();
 			if (FSlateApplication::Get().GetTransformFullscreenMouseInput() && !GIsEditor && Window.IsValid() && Window->GetWindowMode() == EWindowMode::Fullscreen)
 			{
 				const FDisplayMetrics& CachedDisplayMetrics = FSlateApplication::Get().GetCachedDisplayMetricsByRef();
@@ -1494,10 +1492,8 @@ void FSceneViewport::ResizeFrame(uint32 NewWindowSizeX, uint32 NewWindowSizeY, E
 {
 	// Resizing the window directly is only supported in the game
 	if( FApp::IsGame() && FApp::CanEverRender() && NewWindowSizeX > 0 && NewWindowSizeY > 0 )
-	{		
-		TSharedPtr<SWindow> WindowToResize = FSlateApplication::Get().FindWidgetWindow( ViewportWidget.Pin().ToSharedRef());
-
-		if( WindowToResize.IsValid() )
+	{
+		if (TSharedPtr<SWindow> WindowToResize = FindWindow())
 		{
 			NewWindowMode = GetWindowModeType(NewWindowMode);
 
@@ -1714,8 +1710,7 @@ void FSceneViewport::SetFixedViewportSize(uint32 NewViewportSizeX, uint32 NewVie
 		if (NewViewportSizeX > 0 && NewViewportSizeY > 0)
 		{
 			bForceViewportSize = true;
-			TSharedPtr<SWindow> Window = FSlateApplication::Get().FindWidgetWindow(ViewportWidget.Pin().ToSharedRef());
-			if (Window.IsValid())
+			if (TSharedPtr<SWindow> Window = FindWindow())
 			{
 				ResizeViewport(NewViewportSizeX, NewViewportSizeY, Window->GetWindowMode());
 			}
@@ -1723,8 +1718,7 @@ void FSceneViewport::SetFixedViewportSize(uint32 NewViewportSizeX, uint32 NewVie
 		else
 		{
 			bForceViewportSize = false;
-			TSharedPtr<SWindow> Window = FSlateApplication::Get().FindWidgetWindow(ViewportWidget.Pin().ToSharedRef());
-			if (Window.IsValid())
+			if (TSharedPtr<SWindow> Window = FindWindow())
 			{
 				Window->Invalidate(EInvalidateWidget::PaintAndVolatility);
 			}
@@ -1734,8 +1728,7 @@ void FSceneViewport::SetFixedViewportSize(uint32 NewViewportSizeX, uint32 NewVie
 
 void FSceneViewport::SetViewportSize(uint32 NewViewportSizeX, uint32 NewViewportSizeY)
 {
-	TSharedPtr<SWindow> Window = FSlateApplication::Get().FindWidgetWindow(ViewportWidget.Pin().ToSharedRef());
-	if (Window.IsValid())
+	if (TSharedPtr<SWindow> Window = FindWindow())
 	{
 		Window->SetIndependentViewportSize(FVector2D(NewViewportSizeX, NewViewportSizeY));
 		const FVector2D vp = Window->IsMirrorWindow() ? Window->GetSizeInScreen() : Window->GetViewportSize();
@@ -1744,15 +1737,28 @@ void FSceneViewport::SetViewportSize(uint32 NewViewportSizeX, uint32 NewViewport
 	}
 }
 
-TSharedPtr<SWindow> FSceneViewport::FindWindow()
+TSharedPtr<SWindow> FSceneViewport::FindWindow() const
 {
-	if ( ViewportWidget.IsValid() )
+	TSharedPtr<SWindow> Window;
+	if (ViewportWidget.IsValid())
 	{
 		TSharedPtr<SViewport> PinnedViewportWidget = ViewportWidget.Pin();
-		return FSlateApplication::Get().FindWidgetWindow(PinnedViewportWidget.ToSharedRef());
+		Window = FSlateApplication::Get().FindWidgetWindow(PinnedViewportWidget.ToSharedRef());
+	}
+
+	if (!Window.IsValid())
+	{
+		// If we can't find the window via the ViewportWidget, check the game engine instance instead.
+		if (UGameEngine* GameEngine = Cast<UGameEngine>(GEngine))
+		{
+			if (this == GameEngine->SceneViewport.Get())
+			{
+				Window = GameEngine->GameViewportWindow.Pin();
+			}
+		}
 	}
 
-	return TSharedPtr<SWindow>();
+	return Window;
 }
 
 bool FSceneViewport::IsStereoRenderingAllowed() const
@@ -1864,8 +1870,7 @@ int32 FSceneViewport::GetCurrentMonitorIndex() const
 	TSharedPtr<SWidget> PinnedViewport = ViewportWidget.Pin();
 	if (PinnedViewport.IsValid())
 	{
-		TSharedPtr<SWindow> Window = FSlateApplication::Get().FindWidgetWindow(PinnedViewport.ToSharedRef());
-		if (Window.IsValid())
+		if (TSharedPtr<SWindow> Window = FindWindow())
 		{
 			const FSlateRect FullScreenInfo = Window->GetFullScreenInfo();
 			for (int32 Index = 0; Index < Num; ++Index)
@@ -1887,11 +1892,9 @@ bool FSceneViewport::IsExclusiveFullscreenSupported() const
 {
 	bool bSupported = true;
 
-	TSharedPtr<SWidget> PinnedViewport = ViewportWidget.Pin();
-	if (PinnedViewport.IsValid())
+	if (TSharedPtr<SWindow> Window = FindWindow())
 	{
-		TSharedPtr<SWindow> Window = FSlateApplication::Get().FindWidgetWindow(PinnedViewport.ToSharedRef());
-		bSupported = !Window.IsValid() || Window->GetNativeWindow()->IsFullscreenSupported();
+		bSupported = Window->GetNativeWindow()->IsFullscreenSupported();
 	}
 
 	return bSupported;
@@ -1956,8 +1959,7 @@ void FSceneViewport::UpdateViewportRHI(bool bDestroyed, uint32 NewSizeX, uint32
 
 		if (PinnedViewport.IsValid())
 		{
-			TSharedPtr<SWindow> Window = FSlateApplication::Get().FindWidgetWindow(PinnedViewport.ToSharedRef());
-			if (Window.IsValid())
+			if (TSharedPtr<SWindow> Window = FindWindow())
 			{
 				OSWindow = Window->GetNativeWindow()->GetOSWindowHandle();
 				WindowTopLeft = Window->GetPositionInScreen();
@@ -1986,8 +1988,6 @@ void FSceneViewport::UpdateViewportRHI(bool bDestroyed, uint32 NewSizeX, uint32
 
 		if (PinnedViewport.IsValid())
 		{
-			TSharedPtr<SWindow> Window = FSlateApplication::Get().FindWidgetWindow(PinnedViewport.ToSharedRef());
-
 			if (UseSeparateRenderTarget())
 			{
 				uint32 TexSizeX = SizeX, TexSizeY = SizeY;
@@ -2001,7 +2001,7 @@ void FSceneViewport::UpdateViewportRHI(bool bDestroyed, uint32 NewSizeX, uint32
 
 		if (!UseSeparateRenderTarget())
 		{
-			TSharedPtr<SWindow> Window = FSlateApplication::Get().FindWidgetWindow(ViewportWidget.Pin().ToSharedRef());
+			TSharedPtr<SWindow> Window = FindWindow();
 			Renderer->UpdateFullscreenState(Window.ToSharedRef(), NewSizeX, NewSizeY);
 		}
 
@@ -2075,8 +2075,8 @@ void FSceneViewport::EnqueueBeginRenderFrame(const bool bShouldPresent)
 	}
 	else
 	{
-		TSharedPtr<SWindow> WidgetWindow = FSlateApplication::Get().FindWidgetWindow(ViewportWidget.Pin().ToSharedRef());
-		ViewportProvider = FSlateApplication::Get().GetRenderer()->GetViewportProvider(*WidgetWindow);
+		TSharedPtr<SWindow> Window = FindWindow();
+		ViewportProvider = FSlateApplication::Get().GetRenderer()->GetViewportProvider(*Window);
 	}
 
 	DebugCanvasDrawer->InitDebugCanvas(GetClient(), GetClient()->GetWorld());
diff --git a/Engine/Source/Runtime/Engine/Public/Slate/SceneViewport.h b/Engine/Source/Runtime/Engine/Public/Slate/SceneViewport.h
--- a/Engine/Source/Runtime/Engine/Public/Slate/SceneViewport.h
+++ b/Engine/Source/Runtime/Engine/Public/Slate/SceneViewport.h
@@ -270,7 +270,7 @@ PRAGMA_ENABLE_DEPRECATION_WARNINGS
 	/** Does the viewport has a fixed size */
 	ENGINE_API bool HasFixedSize() const;
 
-	ENGINE_API TSharedPtr<SWindow> FindWindow();
+	ENGINE_API TSharedPtr<SWindow> FindWindow() const;
 
 	/** Should return true, if stereo rendering is allowed in this viewport */
 	ENGINE_API virtual bool IsStereoRenderingAllowed() const override;
