Fix VT InitializeEditorResources non-VT proxy texture wasn't created with correct SRGB flag
This caused thumbnail creation to look brighter than they should, after thumbnails were changed to use the non-VT texture proxy

#Synced-CL 51463632

diff --git a/Engine/Source/Editor/UnrealEd/Private/ThumbnailRendering/TextureThumbnailRenderer.cpp b/Engine/Source/Editor/UnrealEd/Private/ThumbnailRendering/TextureThumbnailRenderer.cpp
--- a/Engine/Source/Editor/UnrealEd/Private/ThumbnailRendering/TextureThumbnailRenderer.cpp
+++ b/Engine/Source/Editor/UnrealEd/Private/ThumbnailRendering/TextureThumbnailRenderer.cpp
@@ -135,6 +135,7 @@ void UTextureThumbnailRenderer::Draw(UObject* Object, int32 X, int32 Y, uint32 W
 				//bool bSampleAsVirtualTexture = bIsVirtualTexture;
 
 				// if you just always lie and say you are not a VT, you get a good thumbnail :
+				//	this makes use of the non-VT Editor-only texture that VT makes in FVirtualTexture2DResource::InitializeEditorResources
 				bool bSampleAsVirtualTexture = false;
 
 				BatchedElementParameters = new FBatchedElementTexture2DPreviewParameters(MipLevel, LayerIndex, SliceIndex, bIsNormalMap, bIsSingleChannel, bSingleVTPhysicalSpace, bSampleAsVirtualTexture, bIsTextureArray, bUsePointSampling);
diff --git a/Engine/Source/Runtime/Engine/Private/Texture2D.cpp b/Engine/Source/Runtime/Engine/Private/Texture2D.cpp
--- a/Engine/Source/Runtime/Engine/Private/Texture2D.cpp
+++ b/Engine/Source/Runtime/Engine/Private/Texture2D.cpp
@@ -1733,7 +1733,8 @@ void FVirtualTexture2DResource::InitializeEditorResources(FRHICommandListBase& R
 			const uint32 MipWidthInTiles = FMath::DivideAndRoundUp(GetNumTilesX(), MipScaleFactor);
 			const uint32 MipHeightInTiles = FMath::DivideAndRoundUp(GetNumTilesY(), MipScaleFactor);
 			const uint32 TileSizeInPixels = bCopyUnwantedBordersForAlignment ? GetTileSize() + 2 * GetBorderSize() : GetTileSize();
-			const uint32 LayerMask = 1u;
+			const uint32 LayerMask = 1u;			
+			const ETextureCreateFlags RHITexCreateFlags = ETextureCreateFlags::ShaderResource | (this->TexCreateFlags & TexCreate_SRGB);
 
 			// Get or create the target texture to composite tiles to.
 			// First try reusing the current referenced texture.
@@ -1741,7 +1742,8 @@ void FVirtualTexture2DResource::InitializeEditorResources(FRHICommandListBase& R
 			if (FRHITexture* ReferencedTexture = TextureReferenceRHI->GetReferencedTexture())
 			{
 				FRHITextureDesc const& ExistingDesc = ReferencedTexture->GetDesc();
-				if (ExistingDesc.Format == PixelFormat && ExistingDesc.GetSize() == FIntVector(MipWidthInTiles * TileSizeInPixels, MipHeightInTiles * TileSizeInPixels, 1))
+				if (ExistingDesc.Format == PixelFormat && ExistingDesc.GetSize() == FIntVector(MipWidthInTiles * TileSizeInPixels, MipHeightInTiles * TileSizeInPixels, 1) &&
+					ExistingDesc.Flags == RHITexCreateFlags )
 				{
 					Texture2DRHI = ReferencedTexture;
 				}
@@ -1752,7 +1754,7 @@ void FVirtualTexture2DResource::InitializeEditorResources(FRHICommandListBase& R
 			{
 				const FRHITextureCreateDesc Desc =
 					FRHITextureCreateDesc::Create2D(TEXT("VTEditorResource"), MipWidthInTiles * TileSizeInPixels, MipHeightInTiles * TileSizeInPixels, PixelFormat)
-					.SetFlags(ETextureCreateFlags::ShaderResource);
+					.SetFlags(RHITexCreateFlags);
 
 				Texture2DRHI = RHICmdList.CreateTexture(Desc);
 			}
