New Material Translator:  Preshader2 expression support.  Supports expressions on floats only (covers 99.9% of preshaders).  Significant optimizations to the generated bytecode already implemented -- could possibly be a shippable MVP.  On a test expression, it takes less than half the space of the original Preshaders (23 versus 48 bytes), and 3 times fewer opcodes (9 -> 3), although not all expressions will show that level of savings, and some may be slightly larger (although overall memory should be way down).  More detailed comments in MaterialIRPreshader.h.

Optimizations:
* Trivial Tail Swizzle optimization (avoid generating CPU side opcodes for swizzles that can be trivially handled in HLSL).  Eliminates 25% of all Preshaders observed in a replay on a large real world project.
* Aliasing to avoid generating move ops when accessing subsets or individual components of other vectors.
* Allocator class to recycle and reuse temporary buffer space.  Temporary space is distinguished from output data that goes into the uniform buffer, with the extra space only allocated for the temporary buffer on the CPU side.
* Output location propagation to improve frequency of more efficient in-place opcodes.
* Optimize out unused parameter components.

Preshader2 significant changes:
* Support for fixup, allowing initial bytecode to be generated with virtual "registers", fixed up to offsets in a final pass.
* Support for pretty-printing.
* Added support for scalar inputs to Dot and Cross (custom binary ops), to avoid a messy special case where a scalar alias would need to be expanded somewhere just for those ops.
* Added type and dimension support to texture ops (texture size, texel size, runtime virtual textures).  Texture size and texel size component count varies depending on whether the texture is 2D or 3D, and runtime virtual textures have a mix of types for their outputs.

Other fixes:
* Fixed issue where Fractional was called instead of Frac.  The latter matches the behavior of HLSL frac (returns unsigned rather than signed fraction).
* Fixed broken Runtime Virtual Texture support:  assert in value analyze, not uploading the correct value type, and buffer overwrites.  Besides the assert, the other bugs were hidden by chance (the front end of DF format basically being equivalent to float, and the arbitrary order of allocations avoiding buffer overwrites), and exposed in the refactor.

#rnx

#Synced-CL 51474609

diff --git a/Engine/Source/Runtime/Engine/Private/Materials/MaterialIR.cpp b/Engine/Source/Runtime/Engine/Private/Materials/MaterialIR.cpp
--- a/Engine/Source/Runtime/Engine/Private/Materials/MaterialIR.cpp
+++ b/Engine/Source/Runtime/Engine/Private/Materials/MaterialIR.cpp
@@ -831,6 +831,56 @@ const TCHAR* LexToString(EOperator Op)
 	return TEXT("???");
 }
 
+EPreshader2Opcode ToPreshader2Opcode(EOperator Op, bool& bOutIsCommutative)
+{
+	bOutIsCommutative = false;
+
+	switch (Op)
+	{
+	case UO_Abs: return EPreshader2Opcode::Abs;
+	case UO_ACos: return EPreshader2Opcode::ACos;
+	case UO_ACosFast: return EPreshader2Opcode::ACos;
+	case UO_ASin: return EPreshader2Opcode::ASin;
+	case UO_ASinFast: return EPreshader2Opcode::ASin;
+	case UO_ATan: return EPreshader2Opcode::ATan;
+	case UO_ATanFast: return EPreshader2Opcode::ATan;
+	case UO_Ceil: return EPreshader2Opcode::Ceil;
+	case UO_Cos: return EPreshader2Opcode::Cos;
+	case UO_Exponential: return EPreshader2Opcode::Exponential;
+	case UO_Exponential2: return EPreshader2Opcode::Exponential2;
+	case UO_Floor: return EPreshader2Opcode::Floor;
+	case UO_Frac: return EPreshader2Opcode::Frac;
+	case UO_Length: return EPreshader2Opcode::Length;
+	case UO_Logarithm: return EPreshader2Opcode::Logarithm;
+	case UO_Logarithm10: return EPreshader2Opcode::Logarithm10;
+	case UO_Logarithm2: return EPreshader2Opcode::Logarithm2;
+	case UO_Negate: return EPreshader2Opcode::Negate;
+	case UO_Reciprocal: return EPreshader2Opcode::Reciprocal;
+	case UO_Round: return EPreshader2Opcode::Round;
+	case UO_Saturate: return EPreshader2Opcode::Saturate;
+	case UO_Sign: return EPreshader2Opcode::Sign;
+	case UO_Sin: return EPreshader2Opcode::Sin;
+	case UO_Sqrt: return EPreshader2Opcode::Sqrt;
+	case UO_Tan: return EPreshader2Opcode::Tan;
+	case UO_Truncate: return EPreshader2Opcode::Truncate;
+
+	case BO_Add: bOutIsCommutative = true; return EPreshader2Opcode::Add;
+	case BO_ATan2: return EPreshader2Opcode::ATan2;
+	case BO_ATan2Fast: return EPreshader2Opcode::ATan2;
+	case BO_Cross: return EPreshader2Opcode::Cross;
+	case BO_Divide: return EPreshader2Opcode::Divide;
+	case BO_Dot: bOutIsCommutative = true; return EPreshader2Opcode::Dot;
+	case BO_Fmod: return EPreshader2Opcode::Fmod;
+	case BO_Max: bOutIsCommutative = true; return EPreshader2Opcode::Max;
+	case BO_Min: bOutIsCommutative = true; return EPreshader2Opcode::Min;
+	case BO_Modulo: return EPreshader2Opcode::Fmod;
+	case BO_Multiply: bOutIsCommutative = true; return EPreshader2Opcode::Multiply;
+	case BO_Subtract: return EPreshader2Opcode::Subtract;
+	}
+
+	return EPreshader2Opcode::Invalid;
+}
+
 const TCHAR* LexToString(ETextureReadMode Mode)
 {
 	switch (Mode)
@@ -866,6 +916,13 @@ bool FFunction::Equals(const FFunction* Other) const
 		&& FMemory::Memcmp(Parameters, Other->Parameters, NumParameters * sizeof(FFunctionParameter)) == 0;
 }
 
+int32 FPrimitiveUniform::Analysis_NumComponents() const
+{
+	check(Analysis_ComponentMask);
+
+	return 32 - FMath::CountLeadingZeros(Analysis_ComponentMask);
+}
+
 } // namespace MIR
 
 #endif // #if WITH_EDITOR
diff --git a/Engine/Source/Runtime/Engine/Private/Materials/MaterialIRDebug.cpp b/Engine/Source/Runtime/Engine/Private/Materials/MaterialIRDebug.cpp
--- a/Engine/Source/Runtime/Engine/Private/Materials/MaterialIRDebug.cpp
+++ b/Engine/Source/Runtime/Engine/Private/Materials/MaterialIRDebug.cpp
@@ -340,7 +340,7 @@ struct FDebugDumpIRState
 		else if (const FPreshaderParameter* PreshaderParameter = Instr->As<FPreshaderParameter>())
 		{
 			Out.Appendf(TEXT(" TextureIndex=%d"), PreshaderParameter->TextureIndex);
-			Out.Appendf(TEXT(" PreshaderOffset=%u"), PreshaderParameter->Analysis_PreshaderOffset);
+			Out.Appendf(TEXT(" PreshaderOffset=%d"), PreshaderParameter->Analysis_PreshaderOffset - 1);			// One is added to preshader offsets, so zero can represent none, subtract that here.
 		}
 		else if (const FHardwarePartialDerivative* Derivative = Instr->As<FHardwarePartialDerivative>())
 		{
@@ -367,6 +367,15 @@ struct FDebugDumpIRState
 		{
 			Out.Append(ExternDifferential->Axis == EDerivativeAxis::X ? TEXT("ddx") : TEXT("ddy"));
 		}
+
+		if (Instr->HasSubgraphProperties(MIR::EGraphProperties::PreshaderOut))
+		{
+			Out.Appendf(TEXT(" PreshaderOut"));
+		}
+		else if (Instr->HasSubgraphProperties(MIR::EGraphProperties::PreshaderTemp))
+		{
+			Out.Appendf(TEXT(" PreshaderTemp"));
+		}
 	}
 
 	// Appends a reference to the specified value. This will look like "%x" if Value is an
@@ -412,7 +421,17 @@ struct FDebugDumpIRState
 		}
 		else if (const FPrimitiveUniform* PrimitiveUniform = Value->As<FPrimitiveUniform>())
 		{
-			Out.Appendf(TEXT("[UniformIndex #%d]"), PrimitiveUniform->Analysis_UniformIndex);
+			const TCHAR* PreshaderText = TEXT("");
+			if (Value->HasSubgraphProperties(MIR::EGraphProperties::PreshaderOut))
+			{
+				PreshaderText = TEXT(" PreshaderOut");
+			}
+			else if (Value->HasSubgraphProperties(MIR::EGraphProperties::PreshaderTemp))
+			{
+				PreshaderText = TEXT(" PreshaderTemp");
+			}
+
+			Out.Appendf(TEXT("[Parameter #%d \"%s\"%s]"), PrimitiveUniform->Analysis_UniformIndex, *PrimitiveUniform->Name.ToString(), PreshaderText);
 			ReferencedUniforms.AddUnique({ PrimitiveUniform->Analysis_UniformIndex, PrimitiveUniform });
 		}
 		else if (const FTextureUniform* TextureUniform = Value->As<FTextureUniform>())
@@ -459,6 +478,60 @@ struct FDebugDumpIRState
 		Out += LeftColumn;
 	}
 
+	// Adapted from the AppendUniformLocation utility function in Preshader2.cpp, but applied to MIR type.
+	void AppendUniformLocation(const FValue* Value, FString& InoutResult)
+	{
+		check(Value->Analysis_PreshaderOffset);
+
+		static const TCHAR* LinearSwizzles[4][5] =
+		{
+			{ TEXT(".?"), TEXT(".x"), TEXT(".xy"), TEXT(".xyz"), TEXT(".xyzw") },
+			{ TEXT(".?"), TEXT(".y"), TEXT(".yz"), TEXT(".yzw"), TEXT(".????") },
+			{ TEXT(".?"), TEXT(".z"), TEXT(".zw"), TEXT(".???"), TEXT(".????") },
+			{ TEXT(".?"), TEXT(".w"), TEXT(".??"), TEXT(".???"), TEXT(".????") },
+		};
+
+		FPrimitive Primitive = Value->Type.GetPrimitive();
+		EScalarKind ScalarKind = Primitive.ScalarKind;
+		int32 Dimension = Primitive.NumComponents();
+		int32 Offset = Value->Analysis_PreshaderOffset - 1;			// One is added to preshader offsets, so zero can represent none, subtract that here.
+
+		int32 WordCount;
+		if (const FPrimitiveUniform* Parameter = Value->As<FPrimitiveUniform>())
+		{
+			WordCount = Parameter->Analysis_NumComponents();
+		}
+		else
+		{
+			WordCount = Primitive.NumComponents();
+		}
+
+		if (ScalarKind == EScalarKind::Double)
+		{
+			WordCount *= 2;
+		}
+
+		// Note that the Preshader system supports locations that span float4 buffer elements, such as LWC/Double.  So we need to wrap across buffer
+		// elements in order to communicate that.  For example, a 3-element double would display as [1].xyzw, [2].xy, indicating that
+		// it uses 6 words across two buffer elements.
+		bool bNeedComma = false;
+		while (WordCount > 0)
+		{
+			if (bNeedComma)
+			{
+				InoutResult.Append(TEXT(", "));
+			}
+			bNeedComma = true;
+
+			// Figure out how many words we need to print at the current offset
+			int32 WordsToPrint = FMath::Min(4 - Offset % 4, WordCount);
+			InoutResult.Appendf(TEXT("[%d]%s"), Offset / 4, LinearSwizzles[Offset % 4][WordsToPrint]);
+
+			WordCount -= WordsToPrint;
+			Offset += WordsToPrint;
+		}
+	}
+
 	// Prints the parameter recap section.
 	void DumpReferencedParameters()
 	{
@@ -470,7 +543,9 @@ struct FDebugDumpIRState
 		// Dump the list of referenced parameters.
 		ReferencedUniforms.Sort([](const auto& Param1, const auto& Param2)
 		{
-			return Param1.Key < Param2.Key;
+			return (Param1.Value->Analysis_PreshaderOffset != Param2.Value->Analysis_PreshaderOffset)
+				? Param1.Value->Analysis_PreshaderOffset < Param2.Value->Analysis_PreshaderOffset
+				: Param1.Key < Param2.Key;
 		});
 
 		Out.Append(TEXT("\n; Referenced material parameters\n"));
@@ -479,7 +554,16 @@ struct FDebugDumpIRState
 			const FValue* Value = ParamPair.Value;
 			if (const FPrimitiveUniform* PrimitiveUniform = Value->As<FPrimitiveUniform>())
 			{
-				Out.Appendf(TEXT("#%d = Name=\"%s\" Type=\"%s\"\n"), PrimitiveUniform->Analysis_UniformIndex, *PrimitiveUniform->Name.ToString(), *PrimitiveUniform->Type.GetSpelling());
+				AppendUniformLocation(PrimitiveUniform, Out);
+
+				Out.Appendf(TEXT(" = #%d Name=\"%s\" Type=\"%s\" Mask=.%s%s%s%s\n"),
+					PrimitiveUniform->Analysis_UniformIndex,
+					*PrimitiveUniform->Name.ToString(),
+					*PrimitiveUniform->Type.GetSpelling(),
+					(PrimitiveUniform->Analysis_ComponentMask & (1 << 0)) ? TEXT("x") : TEXT(""),
+					(PrimitiveUniform->Analysis_ComponentMask & (1 << 1)) ? TEXT("y") : TEXT(""),
+					(PrimitiveUniform->Analysis_ComponentMask & (1 << 2)) ? TEXT("z") : TEXT(""),
+					(PrimitiveUniform->Analysis_ComponentMask & (1 << 3)) ? TEXT("w") : TEXT(""));
 			}
 			else if (const FTextureUniform* TextureUniform = Value->As<FTextureUniform>())
 			{
@@ -492,6 +576,43 @@ struct FDebugDumpIRState
 		}
 	}
 
+	// Prints Preshader data
+	void DumpPreshader()
+	{
+		const UE::Shader::FPreshaderData& PreshaderData = Module->GetCompilationOutput().UniformExpressionSet.UniformPreshaderData;
+		if (PreshaderData.Data.IsEmpty())
+		{
+			return;
+		}
+
+		TConstArrayView<FMaterialUniformParameterEvaluation> ParameterEvaluations = Module->GetCompilationOutput().UniformExpressionSet.UniformParameterEvaluations;
+		TArray<UE::Preshader2::FPrintParameterInfo> ParameterInfos;
+
+		ParameterInfos.SetNum(ParameterEvaluations.Num());
+		for (int32 ParameterIndex = 0; ParameterIndex < ParameterEvaluations.Num(); ParameterIndex++)
+		{
+			const FMaterialNumericParameterInfo& NumericParameterInfo = Module->GetCompilationOutput().UniformExpressionSet.GetNumericParameter(ParameterEvaluations[ParameterIndex].ParameterIndex);
+
+			ParameterInfos[ParameterIndex].Name = NumericParameterInfo.ParameterInfo.Name;
+			ParameterInfos[ParameterIndex].Offset = ParameterEvaluations[ParameterIndex].BufferOffset;
+			ParameterInfos[ParameterIndex].ComponentCount = NumericParameterInfo.ParameterType == EMaterialParameterType::Scalar ? 1 : 4;
+			ParameterInfos[ParameterIndex].bDouble = NumericParameterInfo.ParameterType == EMaterialParameterType::DoubleVector;
+		}
+
+		Out.Append(TEXT("\n; Preshader Instructions\n"));
+
+		TArray<FString> Instructions;
+		UE::Preshader2::Print(PreshaderData, ParameterInfos, Instructions);
+
+		for (const FString& Instruction : Instructions)
+		{
+			Out.Append(Instruction);
+			Out.Append(TEXT("\n"));
+		}
+
+		Out.Append(TEXT("\n"));
+	}
+
 	void DumpFunctionHLSLs()
 	{
 		if (ReferencedFunctionHLSLs.IsEmpty())
@@ -537,6 +658,7 @@ FString DebugDumpIR(FStringView MaterialName, const FMaterialIRModule& Module)
 
 	// Print referenced material parameters recap if any
 	State.DumpReferencedParameters();
+	State.DumpPreshader();
 	State.DumpFunctionHLSLs();
 	return State.Out;
 }
diff --git a/Engine/Source/Runtime/Engine/Private/Materials/MaterialIREmitter.cpp b/Engine/Source/Runtime/Engine/Private/Materials/MaterialIREmitter.cpp
--- a/Engine/Source/Runtime/Engine/Private/Materials/MaterialIREmitter.cpp
+++ b/Engine/Source/Runtime/Engine/Private/Materials/MaterialIREmitter.cpp
@@ -31,6 +31,13 @@ static TAutoConsoleVariable<bool> CVarMaterialIRDebugBreakOnPoison(
 	TEXT("Whether the material translator break in the debugger when hitting a poison value in the IR.\n"),
 	ECVF_Default);
 
+static bool GAllowPreshaderOperators = true;
+static FAutoConsoleVariableRef CVarAllowPreshaderOperators(
+	TEXT("r.Material.Translator.AllowPreshaderOperators"),
+	GAllowPreshaderOperators,
+	TEXT("Chicken switch to disable preshader operator support, in case of stability issues."),
+	ECVF_Default);
+
 namespace MIR
 {
 
@@ -170,6 +177,10 @@ static FComposite* MakeCompositePrototype(FEmitter& Emitter, FType Type, int Num
 	FMemory::Memzero(Value, SizeInBytes);
 	Value->Kind = VK_Composite;
 	Value->Type = Type;
+	if (Type.IsPrimitive() && !Type.IsDouble())
+	{
+		Value->SetSubgraphProperties(EGraphProperties::Uniform);
+	}
 
 	return Value;
 }
@@ -1021,6 +1032,7 @@ FValueRef FEmitter::ConstantFloat(TFloat InX)
 {
 	FConstant Scalar = MakePrototype<FConstant>(FType::MakeScalar(EScalarKind::Float));
 	Scalar.Float = InX;
+	Scalar.SetSubgraphProperties(EGraphProperties::Uniform);
 	return EmitPrototype(*this, Scalar);
 }
 
@@ -1138,6 +1150,7 @@ FValueRef FEmitter::NamedPrimitiveUniform(FName Name, FValueRef DefaultConstant)
 	Proto.Name = FScriptName{ Name };
 	Proto.DefaultConstant = DefaultConstant; 
 	Proto.Analysis_UniformIndex = -1;
+	Proto.SetSubgraphProperties(EGraphProperties::HasParameter | EGraphProperties::Uniform);
 
 	return EmitPrototype(*this, Proto);
 }
@@ -1160,6 +1173,10 @@ FValueRef FEmitter::TextureUniform(FName Name, UTexture* Texture, EMaterialSampl
 	Proto.SamplerType = SamplerType;
 	Proto.Analysis_UniformIndex = INDEX_NONE;
 	
+	// This needs to flagged as Uniform, because it may be referenced as a child of FPreshaderParameter, which we want to be Uniform.
+	// Uniform status on a graph node requires all children of the node to also be Uniformn.
+	Proto.SetSubgraphProperties(EGraphProperties::Uniform);
+	
 	return EmitPrototype(*this, Proto);
 }
 
@@ -1208,6 +1225,10 @@ FValueRef FEmitter::VirtualTextureUniform(FName Name, URuntimeVirtualTexture* Vi
 	Proto.VTPageTableIndex = VTPageTableIndex;
 	Proto.Analysis_UniformIndex = INDEX_NONE;
 
+	// This needs to flagged as Uniform, because it may be referenced as a child of FPreshaderParameter, which we want to be Uniform.
+	// Uniform status on a graph node requires all children of the node to also be Uniformn.
+	Proto.SetSubgraphProperties(EGraphProperties::Uniform);
+
 	return EmitPrototype(*this, Proto);
 }
 
@@ -1233,6 +1254,7 @@ FValueRef FEmitter::PreshaderParameter(FType Type, EPreshader2Opcode Opcode, FVa
 	Prototype.Opcode = Opcode;
 	Prototype.TextureIndex = MaterialInterface->GetDefaultTextureIdx(SourceParameterTexture);
 	Prototype.Payload = MoveTemp(Payload);
+	Prototype.SetSubgraphProperties(EGraphProperties::HasParameter | EGraphProperties::Uniform);
 
 	return EmitPrototype(*this, Prototype);
 }
@@ -1693,7 +1715,7 @@ static T FoldScalarOperator(FEmitter& Emitter, EOperator Operator, T A, T B, T C
 				}
 			case UO_Frac:
 			{
-				return FGenericPlatformMath::Fractional(A);
+				return FGenericPlatformMath::Frac(A);
 			}
 			case UO_Logarithm:
 			case UO_Logarithm2:
@@ -2140,27 +2162,48 @@ enum EOperatorResult
 	OR_MatrixMultiplyResult,				// The result type of the matrix multiplication of first two arguments
 };
 
+enum EOperatorUniform
+{
+	OU_False = 0,							// Operator does not support uniform evaluation
+	OU_True = 1,							// Operator does support uniform evaluation
+};
+
 // The signature of an operator consisting of its parameter and return type information.
 struct FOperatorSignature
 {
 	EOperatorParameterFlags ParameterFlags[3] = { OPF_Unknown, OPF_Unknown, OPF_Unknown };
 	EOperatorResult 		Result = OR_Unknown;
+	EOperatorUniform		Uniform = OU_False;
 };
 
 // Returns the signature of an operator.
 static const FOperatorSignature* GetOperatorSignature(EOperator Op)
 {
+	// For reference, these Preshader2 ops are not supported in MIR:
+	//		Fractional, Normalize
+	//		NOTE:  Fractional is slightly different than Frac, in that it returns signed fractions for negative numbers.
+	// 
+	// And these MIR ops are not supported in Preshader2:
+	//		UO_Rsqrt, TO_Clamp == good candidates to add support for, used in existing preshaders.  Clamp would be min plus max ops.
+	//		BO_Distance, BO_Pow, BO_Step
+	//		Bitwise ops, hyperbolics, logical ops, number categorizers, comparisons, ternary ops
+	//
+	// Cast is supported by Preshader2, but not implemented yet.
+	//
+
 	static const FOperatorSignature* Signatures = [] ()
 	{
 		const FOperatorSignature UnaryFloat = { { OPF_CheckIsArithmetic | OPF_CastToAnyFloat }, OR_FirstArgumentType };
-		const FOperatorSignature UnaryFloatOrDouble = { { OPF_CheckIsArithmetic | OPF_CastToAnyFloat | OPF_AllowDouble }, OR_FirstArgumentTypeAllowDouble };
-		const FOperatorSignature UnaryFloatLWCDemote = { { OPF_CheckIsArithmetic | OPF_CastToAnyFloat | OPF_AllowDouble }, OR_FirstArgumentType };
+		const FOperatorSignature UnaryFloatUniform = { { OPF_CheckIsArithmetic | OPF_CastToAnyFloat }, OR_FirstArgumentType, OU_True };
+		const FOperatorSignature UnaryFloatOrDouble = { { OPF_CheckIsArithmetic | OPF_CastToAnyFloat | OPF_AllowDouble }, OR_FirstArgumentTypeAllowDouble, OU_True };
+		const FOperatorSignature UnaryFloatLWCDemote = { { OPF_CheckIsArithmetic | OPF_CastToAnyFloat | OPF_AllowDouble }, OR_FirstArgumentType, OU_True };
 		const FOperatorSignature UnaryFloatToBoolean = { { OPF_CheckIsArithmetic | OPF_CastToAnyFloat }, OR_BooleanWithFirstArgumentDimensions };
 		const FOperatorSignature BinaryArithmetic = { { OPF_CastToCommonArithmeticType, OPF_CastToCommonArithmeticType }, OR_FirstArgumentType };
-		const FOperatorSignature BinaryArithmeticAllowDouble = { { OPF_CastToCommonArithmeticTypeAllowDouble, OPF_CastToCommonArithmeticTypeAllowDouble }, OR_FirstArgumentTypeAllowDouble };
+		const FOperatorSignature BinaryArithmeticAllowDouble = { { OPF_CastToCommonArithmeticTypeAllowDouble, OPF_CastToCommonArithmeticTypeAllowDouble }, OR_FirstArgumentTypeAllowDouble, OU_True };
 		const FOperatorSignature BinaryArithmeticLWCDemote = { { OPF_CastToCommonArithmeticTypeAllowDouble, OPF_CastToCommonArithmeticTypeAllowDouble }, OR_FirstArgumentType };
 		const FOperatorSignature BinaryInteger = { { OPF_CheckIsInteger | OPF_CastToCommonArithmeticType, OPF_CheckIsInteger | OPF_CastToCommonArithmeticType }, OR_FirstArgumentType };
 		const FOperatorSignature BinaryFloat = { { OPF_CastToCommonFloatType, OPF_CastToCommonFloatType }, OR_FirstArgumentType };
+		const FOperatorSignature BinaryFloatUniform = { { OPF_CastToCommonFloatType, OPF_CastToCommonFloatType }, OR_FirstArgumentType, OU_True };
 		const FOperatorSignature BinaryArithmeticComparisonAllowDouble = { { OPF_CastToCommonArithmeticType | OPF_AllowDouble, OPF_CastToCommonArithmeticType | OPF_AllowDouble }, OR_BooleanWithFirstArgumentDimensions };
 		const FOperatorSignature BinaryLogical = { { OPF_CheckIsBoolean | OPF_CastToCommonType, OPF_CheckIsBoolean | OPF_CastToCommonType }, OR_FirstArgumentType };
 		const FOperatorSignature TernaryArithmeticDouble = { { OPF_CastToCommonArithmeticTypeAllowDouble, OPF_CastToCommonArithmeticTypeAllowDouble, OPF_CastToCommonArithmeticTypeAllowDouble }, OR_FirstArgumentTypeAllowDouble };
@@ -2170,40 +2213,40 @@ static const FOperatorSignature* GetOperatorSignature(EOperator Op)
 
 		/* unary operators */
 		S[UO_BitwiseNot] 				= { { OPF_CheckIsInteger }, OR_FirstArgumentType };
-		S[UO_Negate] 					= { { OPF_CheckIsArithmetic | OPF_AllowDouble }, OR_FirstArgumentTypeAllowDouble };
+		S[UO_Negate] 					= { { OPF_CheckIsArithmetic | OPF_AllowDouble }, OR_FirstArgumentTypeAllowDouble, OU_True };
 		S[UO_Not] 						= { { OPF_CheckIsBoolean }, OR_FirstArgumentType };
 
 		S[UO_Abs] 						= UnaryFloatOrDouble;
 		S[UO_ACos] 						= UnaryFloatLWCDemote;
-		S[UO_ACosFast]					= UnaryFloat;
+		S[UO_ACosFast]					= UnaryFloatUniform;
 		S[UO_ACosh] 					= { { OPF_CheckIsArithmetic | OPF_CastToAnyFloat | OPF_CheckIsOneOrGreaterFloatConst }, OR_FirstArgumentType };
 		S[UO_ASin] 						= UnaryFloatLWCDemote;
-		S[UO_ASinFast]					= UnaryFloat;
+		S[UO_ASinFast]					= UnaryFloatUniform;
 		S[UO_ASinh] 					= UnaryFloat;
 		S[UO_ATan] 						= UnaryFloatLWCDemote;
-		S[UO_ATanFast]					= UnaryFloat;
+		S[UO_ATanFast]					= UnaryFloatUniform;
 		S[UO_ATanh] 					= { { OPF_CheckIsArithmetic | OPF_CastToAnyFloat | OPF_CheckIsBetweenMinusOneAndPlusOneFloatConst }, OR_FirstArgumentType };
 		S[UO_Ceil] 						= UnaryFloatOrDouble;
 		S[UO_Cos] 						= UnaryFloatLWCDemote;
-		S[UO_Exponential]				= UnaryFloat;
-		S[UO_Exponential2]				= UnaryFloat;
+		S[UO_Exponential]				= UnaryFloatUniform;
+		S[UO_Exponential2]				= UnaryFloatUniform;
 		S[UO_Floor] 					= UnaryFloatOrDouble;
 		S[UO_Frac]						= UnaryFloatLWCDemote;
 		S[UO_IsFinite]					= UnaryFloatToBoolean;
 		S[UO_IsInf]						= UnaryFloatToBoolean;
 		S[UO_IsNan]						= UnaryFloatToBoolean;
-		S[UO_Length]					= { { OPF_CheckIsArithmetic | OPF_CheckIsNotMatrix | OPF_CastToAnyFloat | OPF_AllowDouble }, OR_FirstArgumentTypeToScalarLWC };
-		S[UO_Logarithm] 				= { { OPF_CheckIsArithmetic | OPF_CheckIsNonZeroFloatConst | OPF_CheckIsNonNegativeFloatConst | OPF_CastToAnyFloat }, OR_FirstArgumentType };
-		S[UO_Logarithm10]				= { { OPF_CheckIsArithmetic | OPF_CheckIsNonZeroFloatConst | OPF_CheckIsNonNegativeFloatConst | OPF_CastToAnyFloat }, OR_FirstArgumentType };
-		S[UO_Logarithm2] 				= { { OPF_CheckIsArithmetic | OPF_CheckIsNonZeroFloatConst | OPF_CheckIsNonNegativeFloatConst | OPF_CastToAnyFloat }, OR_FirstArgumentType };
+		S[UO_Length]					= { { OPF_CheckIsArithmetic | OPF_CheckIsNotMatrix | OPF_CastToAnyFloat | OPF_AllowDouble }, OR_FirstArgumentTypeToScalarLWC, OU_True };
+		S[UO_Logarithm] 				= { { OPF_CheckIsArithmetic | OPF_CheckIsNonZeroFloatConst | OPF_CheckIsNonNegativeFloatConst | OPF_CastToAnyFloat }, OR_FirstArgumentType, OU_True };
+		S[UO_Logarithm10]				= { { OPF_CheckIsArithmetic | OPF_CheckIsNonZeroFloatConst | OPF_CheckIsNonNegativeFloatConst | OPF_CastToAnyFloat }, OR_FirstArgumentType, OU_True };
+		S[UO_Logarithm2] 				= { { OPF_CheckIsArithmetic | OPF_CheckIsNonZeroFloatConst | OPF_CheckIsNonNegativeFloatConst | OPF_CastToAnyFloat }, OR_FirstArgumentType, OU_True };
 		S[UO_LWCTile]					= { };		// UNUSED
-		S[UO_Reciprocal]				= { { OPF_CheckIsArithmetic | OPF_CheckIsNonZeroFloatConst | OPF_CastToAnyFloat }, OR_FirstArgumentType };
+		S[UO_Reciprocal]				= { { OPF_CheckIsArithmetic | OPF_CheckIsNonZeroFloatConst | OPF_CastToAnyFloat }, OR_FirstArgumentType, OU_True };
 		S[UO_Round] 					= UnaryFloatOrDouble;
 		S[UO_Rsqrt]						= { { OPF_CheckIsArithmetic | OPF_CheckIsNonZeroFloatConst | OPF_CastToAnyFloat }, OR_FirstArgumentType };
 		S[UO_Saturate]					= UnaryFloatLWCDemote;
 		S[UO_Sign]						= UnaryFloatLWCDemote;
 		S[UO_Sin] 						= UnaryFloatLWCDemote;
-		S[UO_Sqrt]						= { { OPF_CheckIsArithmetic | OPF_CheckIsNonNegativeFloatConst | OPF_CastToAnyFloat | OPF_AllowDouble }, OR_FirstArgumentType };
+		S[UO_Sqrt]						= { { OPF_CheckIsArithmetic | OPF_CheckIsNonNegativeFloatConst | OPF_CastToAnyFloat | OPF_AllowDouble }, OR_FirstArgumentType, OU_True };
 		S[UO_Tan] 						= UnaryFloatLWCDemote;
 		S[UO_Tanh] 						= UnaryFloat;
 		S[UO_Truncate]					= UnaryFloatOrDouble;
@@ -2231,12 +2274,12 @@ static const FOperatorSignature* GetOperatorSignature(EOperator Op)
 		S[BO_BitShiftRight]				= BinaryInteger;
 
 
-		S[BO_ATan2]						= BinaryFloat;
-		S[BO_ATan2Fast]					= BinaryFloat;
-		S[BO_Cross] 					= { { OPF_CheckIsArithmetic | OPF_CheckIsVector3, OPF_CastToFirstArgumentType }, OR_FirstArgumentType };
+		S[BO_ATan2]						= BinaryFloatUniform;
+		S[BO_ATan2Fast]					= BinaryFloatUniform;
+		S[BO_Cross] 					= { { OPF_CheckIsArithmetic | OPF_CheckIsVector3, OPF_CastToFirstArgumentType }, OR_FirstArgumentType, OU_True };
 		S[BO_Distance]					= { { OPF_CastToCommonFloatType | OPF_AllowDouble, OPF_CastToCommonFloatType | OPF_AllowDouble }, OR_FirstArgumentTypeToScalarLWC };
-		S[BO_Dot] 						= { { OPF_CheckIsArithmetic | OPF_CheckIsNotMatrix | OPF_AllowDouble, OPF_CastToFirstArgumentType | OPF_AllowDouble }, OR_FirstArgumentTypeToScalarLWC };
-		S[BO_Fmod] 						= { { OPF_CastToCommonFloatType | OPF_AllowDouble, OPF_CastToCommonFloatType }, OR_FirstArgumentType };		// First input can be LWC, but second and output are always demoted
+		S[BO_Dot] 						= { { OPF_CheckIsArithmetic | OPF_CheckIsNotMatrix | OPF_AllowDouble, OPF_CastToFirstArgumentType | OPF_AllowDouble }, OR_FirstArgumentTypeToScalarLWC, OU_True };
+		S[BO_Fmod] 						= { { OPF_CastToCommonFloatType | OPF_AllowDouble, OPF_CastToCommonFloatType }, OR_FirstArgumentType, OU_True };		// First input can be LWC, but second and output are always demoted
 
 		S[BO_Max] 						= BinaryArithmeticAllowDouble;
 		S[BO_Min] 						= BinaryArithmeticAllowDouble;
@@ -2562,6 +2605,19 @@ static FType ValidateOperatorAndGetResultType(FEmitter& Emitter, EOperator Op, F
 	UE_MIR_UNREACHABLE();
 }
 
+// Conditionally mark the given operator as uniform if it's supported in preshaders.
+static void MarkUniformIfSupportedPreshaderOperator(FOperator& Proto)
+{
+	// Only allow operators on float types now -- need to implement casting logic for other types to work, but it's low priority,
+	// given that 99.9% of preshaders use purely float math.  There is currently no user facing integer parameter type, and math
+	// on doubles is rarely uniform, because interesting double typed values (such as camera, world, or object position) are all
+	// non-uniform.
+	if (GAllowPreshaderOperators && GetOperatorSignature(Proto.Op)->Uniform && Proto.AArg->Type.IsFloat() && (!Proto.BArg || Proto.BArg->Type.IsFloat()))
+	{
+		Proto.SetSubgraphProperties(EGraphProperties::Uniform);
+	}
+}
+
 // Returns whether the operator supports componentwise application. In other words, if the following is true:
 // 	op(v, w) == [op(v_0, w_0), ..., op(v_n, w_n)]
 static bool IsComponentwiseOperator(EOperator Op)
@@ -2659,6 +2715,9 @@ static FValue* TryFoldComponentwiseOperator(FEmitter& Emitter, EOperator Op, FVa
 				Proto.AArg = Emitter.Subscript(A, i);
 				Proto.BArg = B ? Emitter.Subscript(B, i).Value : nullptr;
 				Proto.CArg = C ? Emitter.Subscript(C, i).Value : nullptr;
+
+				MarkUniformIfSupportedPreshaderOperator(Proto);
+
 				ResultComponents[i] = EmitPrototype(Emitter, Proto);
 			}
 		}
@@ -2844,6 +2903,8 @@ FValueRef FEmitter::Operator(EOperator Op, FValueRef A, FValueRef B, FValueRef C
 		Proto.BArg = B;
 		Proto.CArg = C;
 
+		MarkUniformIfSupportedPreshaderOperator(Proto);
+
 		Result = EmitPrototype(*this, Proto);
 	}
 
@@ -2979,6 +3040,17 @@ FValueRef FEmitter::Subscript(FValueRef Value, int32 Index)
 	Prototype.Arg = Value;
 	Prototype.Index = Index;
 
+	// Don't allow preshader code to be generated for subscripts of double types.  These must be accessed from the original
+	// full source vector, and subscripted on the GPU, since the data format is Structure of Arrays.  High and low components
+	// of values are not interleaved, and thus there's no way to encode an offset to a specific component without the full
+	// context of the size of vector it came from, which adds major complexity.  We currently don't support Preshader math
+	// on doubles anyway, so this works fine, but if we do want to support it, we will need to change the new translator to
+	// use AoS, so subsets or components of a vector can be trivially addressed.
+	if (!Value->Type.IsDouble())
+	{
+		Prototype.SetSubgraphProperties(EGraphProperties::Uniform);
+	}
+
 	return Value.To(EmitPrototype(*this, Prototype));
 }
 
@@ -3186,6 +3258,7 @@ static FValueRef CastToPrimitive(FEmitter& Emitter, FValueRef Value, FType Targe
 			// Otherwise emit a cast instruction of the subscript value to the target type.
 			FScalar Prototype = MakePrototype<FScalar>(TargetType);
 			Prototype.Arg = Value;
+			// Prototype.SetSubgraphProperties(EGraphProperties::Uniform);		// TODO:  Add preshader cast support
 			return EmitPrototype(Emitter, Prototype);
 		}
 	}
diff --git a/Engine/Source/Runtime/Engine/Private/Materials/MaterialIRModuleBuilder.cpp b/Engine/Source/Runtime/Engine/Private/Materials/MaterialIRModuleBuilder.cpp
--- a/Engine/Source/Runtime/Engine/Private/Materials/MaterialIRModuleBuilder.cpp
+++ b/Engine/Source/Runtime/Engine/Private/Materials/MaterialIRModuleBuilder.cpp
@@ -792,7 +792,7 @@ struct FMaterialIRModuleBuilderImpl
 
 			// The value being set to this material output is now valid.
 			check(PropertyValue);
-			 
+
 			// Add support for lerp to selection color for PC development builds.
 			if (Property == MP_EmissiveColor &&
 				Builder->MaterialInterface->GetMaterial()->MaterialDomain != MD_Volume &&
@@ -1081,7 +1081,14 @@ struct FMaterialIRModuleBuilderImpl
 					{
 						if (Use)
 						{
-							Value->GraphProperties |= Use->GraphProperties;
+							// Note that we don't put this UnionMask of flags in the enum class, because it breaks Natvis flag display.
+							const MIR::EGraphProperties UnionMask = MIR::EGraphProperties::PixelStageOnly | MIR::EGraphProperties::ReadsPixelNormal | MIR::EGraphProperties::HasParameter;
+
+							Value->GraphProperties |= Use->GraphProperties & UnionMask;
+
+							// Uniform property is a logical intersection (AND) rather than union (OR), but we also need to make sure we don't
+							// clear any of the unioned properties, so always include those in the AND mask.
+							Value->GraphProperties &= (Use->GraphProperties & MIR::EGraphProperties::Uniform) | UnionMask;
 						}
 					}
 
@@ -1109,6 +1116,13 @@ struct FMaterialIRModuleBuilderImpl
 			}
 		}
 
+		// After the entire tree has been analyzed, we can run fixup, which allocates output buffer offsets and fixes up
+		// parameter evaluations and preshader bytecode with those.
+		if (Module->IsValid())
+		{
+			ValueAnalyzer.Preshader.PreshaderFixup(ValueAnalyzer.GetPreshaderData(), Module, ValueAnalyzer.Insights);
+		}
+
 		return Module->IsValid();
 	}
 
diff --git a/Engine/Source/Runtime/Engine/Private/Materials/MaterialIRToHLSLTranslator.cpp b/Engine/Source/Runtime/Engine/Private/Materials/MaterialIRToHLSLTranslator.cpp
--- a/Engine/Source/Runtime/Engine/Private/Materials/MaterialIRToHLSLTranslator.cpp
+++ b/Engine/Source/Runtime/Engine/Private/Materials/MaterialIRToHLSLTranslator.cpp
@@ -328,6 +328,22 @@ struct FEntryPointCodeGen
 				continue;
 			}
 
+			// @jason.hoerner This works but it isn't ideal. We should instead omit linking instruction
+			// in the Step_LinkInstructions (skip them) that do not satisfy this condition. Only instructions that truly need
+			// evaluation in the various stages should be linked so that here we never skip any.
+
+			// PreshaderTemp values don't need to be lowered to HLSL, since they are only used by other preshaders, and can
+			// just be skipped.  TrivialTailSwizzle values are an exception, and should be lowered, even if also marked as
+			// PreshaderTemp.  TrivialTailSwizzle indicates this value is used by HLSL, but isn't a full fledged PreshaderOut
+			// written to an explicit buffer location, and is instead cobbled together with HLSL swizzles from an underlying
+			// value that will have been marked as PreshaderOut.  PreshaderOut values never need to be emitted, since those
+			// just become uniform buffer fetches.
+			if ((Instr->HasSubgraphProperties(MIR::EGraphProperties::PreshaderTemp) && !Instr->HasSubgraphProperties(MIR::EGraphProperties::TrivialTailSwizzle))
+				|| Instr->HasSubgraphProperties(MIR::EGraphProperties::PreshaderOut))
+			{
+				continue;
+			}
+
 			// If this isn't a set material output instruction but last one was, add a newline for better readability.
 			const MIR::FSetMaterialOutput* SetMaterialOutput = MIR::As<MIR::FSetMaterialOutput>(Instr);
 			if (!SetMaterialOutput && bLastInstructionWasSetMaterialOutput)
@@ -380,6 +396,13 @@ struct FEntryPointCodeGen
 	
 	ENoOp LowerValue(const MIR::FValue* Value)
 	{
+		// @jason.hoerner this is okay as this is simply emitting the instruction that reads the preshader result.  
+		if (Value->Analysis_PreshaderOffset && Value->HasSubgraphProperties(MIR::EGraphProperties::PreshaderOut))
+		{
+			EmitPreshaderBufferRead(Value);
+			return NoOp;
+		}
+
 		// Instruction results may be shared among other dependent instructions.
 		if (const MIR::FInstruction* Instr = MIR::AsInstruction(Value))
 		{
@@ -402,7 +425,6 @@ struct FEntryPointCodeGen
 			case MIR::VK_Constant: LowerConstant(static_cast<const MIR::FConstant*>(Value)); break;
 			case MIR::VK_Builtin: LowerBuiltin(static_cast<const MIR::FBuiltin*>(Value)); break;
 			case MIR::VK_MaterialParameterCollection: LowerMaterialParameterCollection(static_cast<const MIR::FMaterialParameterCollection*>(Value)); break;
-			case MIR::VK_PrimitiveUniform: LowerPrimitiveUniform(static_cast<const MIR::FPrimitiveUniform*>(Value)); break;
 			case MIR::VK_TextureUniform: LowerTextureUniform(static_cast<const MIR::FTextureUniform*>(Value)); break;
 			case MIR::VK_VirtualTextureUniform: LowerVirtualTextureUniform(static_cast<const MIR::FVirtualTextureUniform*>(Value)); break;
 			case MIR::VK_GetVertexInterpolator: LowerGetVertexInterpolator(static_cast<const MIR::FGetVertexInterpolator*>(Value)); break;
@@ -414,6 +436,12 @@ struct FEntryPointCodeGen
 
 	ENoOp LowerInstruction(const MIR::FInstruction* Instr)
 	{
+		if (Instr->Analysis_PreshaderOffset && Instr->HasSubgraphProperties(MIR::EGraphProperties::PreshaderOut))
+		{
+			EmitPreshaderBufferRead(Instr);
+			return NoOp;
+		}
+
 		switch (Instr->Kind)
 		{
 			case MIR::VK_Nop: LowerNop(static_cast<const MIR::FNop*>(Instr)); break;
@@ -432,7 +460,6 @@ struct FEntryPointCodeGen
 			case MIR::VK_SetVertexInterpolator: LowerSetVertexInterpolator(static_cast<const MIR::FSetVertexInterpolator*>(Instr)); break;
 			case MIR::VK_Call: LowerCall(static_cast<const MIR::FCall*>(Instr)); break;
 			case MIR::VK_CallParameterOutput: LowerCallOutput(static_cast<const MIR::FCallParameterOutput*>(Instr)); break;
-			case MIR::VK_PreshaderParameter: LowerPreshaderParameter(static_cast<const MIR::FPreshaderParameter*>(Instr)); break;
 			case MIR::VK_SubstrateDefaultSlab: LowerSubstrateDefaultSlab(static_cast<const MIR::FSubstrateDefaultSlab*>(Instr)); break;
 			case MIR::VK_SubstrateSlab: LowerSubstrateSlab(static_cast<const MIR::FSubstrateSlab*>(Instr)); break;
 			case MIR::VK_SubstrateShadingModels: LowerSubstrateShadingModel(static_cast<const MIR::FSubstrateShadingModels*>(Instr)); break;
@@ -556,13 +583,6 @@ struct FEntryPointCodeGen
 		ExternArg->ToHLSL(ExternPrinter);
 	}
 
-	void LowerPrimitiveUniform(const MIR::FPrimitiveUniform* PrimitiveUniform)
-	{
-		// Get the global float4 component index (e.g. if this is 13, it refer to PreshaderBuffer[3].y)
-		const uint32 GlobalComponentOffset = Module->GetCompilationOutput().UniformExpressionSet.GetNumericParameterEvaluationOffset(PrimitiveUniform->Analysis_UniformIndex);
-		EmitPreshaderBufferRead(PrimitiveUniform->Type.GetPrimitive(), GlobalComponentOffset);
-	}
-	
 	void LowerTextureUniform(const MIR::FTextureUniform* TextureUniform)
 	{
 		Printer << TEXTVIEW("Material.") << LowerTextureMaterialType(TextureUniform->Texture->GetMaterialType(), true) << TEXTVIEW("_") << TextureUniform->Analysis_UniformIndex;
@@ -573,16 +593,67 @@ struct FEntryPointCodeGen
 		Printer << TEXTVIEW("Material.") << LowerTextureMaterialType(MCT_TextureVirtual, true) << TEXTVIEW("_") << VirtualTextureUniform->Analysis_UniformIndex;
 	}
 
-	void EmitPreshaderBufferReadDoubleVector(const MIR::FPrimitive& PrimitiveType, uint32 GlobalComponentOffset)
+	void EmitPreshaderBufferReadDoubleVector(uint32 GlobalComponentOffset, int32 SourceComponents, int32 UsedComponents)
 	{
-		const int32 NumComponents = PrimitiveType.NumComponents();
+		// Must be a multiple of 2 if a single component, and multiple of 4 if more than one component.
+		check(!(GlobalComponentOffset & (SourceComponents == 1 ? 0x1 : 0x3)));
+		check(UsedComponents <= SourceComponents);
 
-		// Index of the float4 slot.  LWC parameters always have four elements, so there won't be an offset.
+		// Index of the float4 slot.
 		const uint32 BufferSlotIndex = GlobalComponentOffset / 4;
 
-		Printer << TEXTVIEW("DFToWS(MakeDFVector") << NumComponents;
-		Printer << TEXTVIEW("(Material.PreshaderBuffer[") << BufferSlotIndex << TEXTVIEW("]") << GVector4SwizzleSubset[NumComponents - 1];
-		Printer << TEXTVIEW(", Material.PreshaderBuffer[") << BufferSlotIndex + 1 << TEXTVIEW("]") << GVector4SwizzleSubset[NumComponents - 1] << TEXTVIEW("))");
+		if (UsedComponents == 1)
+		{
+			Printer << TEXTVIEW("DFToWS(MakeDFScalar");
+		}
+		else
+		{
+			Printer << TEXTVIEW("DFToWS(MakeDFVector") << UsedComponents;
+		}
+
+		static const TCHAR* Components = TEXT("xyzw");
+		static const TCHAR* ComponentSubsets[4] = { TEXT("x"), TEXT("xy"), TEXT("xyz"), TEXT("xyzw") };
+
+		switch (SourceComponents)
+		{
+		case 1:
+			// Emit two individual components
+			Printer << TEXTVIEW("(Material.PreshaderBuffer[") << BufferSlotIndex << TEXTVIEW("].") << Components[GlobalComponentOffset % 4];
+			Printer << TEXTVIEW(", Material.PreshaderBuffer[") << BufferSlotIndex << TEXTVIEW("].") << Components[GlobalComponentOffset % 4 + 1] << TEXTVIEW("))");
+			break;
+		case 2:
+			// Emit XY and ZW, or just X and Z
+			Printer << TEXTVIEW("(Material.PreshaderBuffer[") << BufferSlotIndex << (UsedComponents == 2 ? TEXTVIEW("].xy") : TEXTVIEW("].x"));
+			Printer << TEXTVIEW(", Material.PreshaderBuffer[") << BufferSlotIndex << (UsedComponents == 2 ? TEXTVIEW("].zw))") : TEXTVIEW("].z))"));
+			break;
+		case 3:
+			if (UsedComponents == 3)
+			{
+				// Emit XYZ and float3(W, XY)
+				Printer << TEXTVIEW("(Material.PreshaderBuffer[") << BufferSlotIndex << TEXTVIEW("].xyz");
+				Printer << TEXTVIEW(", float3(Material.PreshaderBuffer[") << BufferSlotIndex << TEXTVIEW("].w");
+				Printer << TEXTVIEW(", Material.PreshaderBuffer[") << BufferSlotIndex + 1 << TEXTVIEW("].xy)))");
+			}
+			else if (UsedComponents == 2)
+			{
+				// Emit XY and float2(W, X)
+				Printer << TEXTVIEW("(Material.PreshaderBuffer[") << BufferSlotIndex << TEXTVIEW("].xy");
+				Printer << TEXTVIEW(", float2(Material.PreshaderBuffer[") << BufferSlotIndex << TEXTVIEW("].w");
+				Printer << TEXTVIEW(", Material.PreshaderBuffer[") << BufferSlotIndex + 1 << TEXTVIEW("].x)))");
+			}
+			else
+			{
+				// Emit X and W
+				Printer << TEXTVIEW("(Material.PreshaderBuffer[") << BufferSlotIndex << TEXTVIEW("].x");
+				Printer << TEXTVIEW(", Material.PreshaderBuffer[") << BufferSlotIndex << TEXTVIEW("].w))");
+			}
+			break;
+		case 4:
+			// Emit used components
+			Printer << TEXTVIEW("(Material.PreshaderBuffer[") << BufferSlotIndex << TEXTVIEW("].") << ComponentSubsets[UsedComponents - 1];
+			Printer << TEXTVIEW(", Material.PreshaderBuffer[") << BufferSlotIndex + 1 << TEXTVIEW("].") << ComponentSubsets[UsedComponents - 1] << TEXTVIEW("))");
+			break;
+		}
 	}
 
 	void EmitPreshaderBufferReadFloatVector(const MIR::FPrimitive& PrimitiveType, uint32 GlobalComponentOffset)
@@ -620,14 +691,28 @@ struct FEntryPointCodeGen
 		}
 	}
 
-	void EmitPreshaderBufferRead(const MIR::FPrimitive& PrimitiveType, uint32 GlobalComponentOffset)
+	void EmitPreshaderBufferRead(const MIR::FValue* Value)
 	{
+		check(Value->Analysis_PreshaderOffset);
+
+		uint32 GlobalComponentOffset = Value->Analysis_PreshaderOffset - 1;
+		MIR::FPrimitive PrimitiveType = Value->Type.GetPrimitive();
+
 		check(PrimitiveType.IsScalar() || PrimitiveType.IsRowVector()); // no matrices yet
 
 		// LWC parameters are handled differently, they have their own dedicated function.
 		if (PrimitiveType.IsDouble())
 		{
-			EmitPreshaderBufferReadDoubleVector(PrimitiveType, GlobalComponentOffset);
+			// For double uniform parameters, we need to get the analysis component count, because it affects the stride
+			// between the high and low portion in the DF format.  We can be lazier for float vectors below, because
+			// there's no stride to worry about, and it doesn't matter how many components were in the original vector.
+			int32 NumComponents = PrimitiveType.NumComponents();
+			if (const MIR::FPrimitiveUniform* PrimitiveUniform = Value->As<MIR::FPrimitiveUniform>())
+			{
+				NumComponents = PrimitiveUniform->Analysis_NumComponents();
+			}
+
+			EmitPreshaderBufferReadDoubleVector(GlobalComponentOffset, NumComponents, NumComponents);
 		}
 		else
 		{
@@ -635,16 +720,6 @@ struct FEntryPointCodeGen
 		}
 	}
 
-	ENoOp LowerPreshaderParameter(const MIR::FPreshaderParameter* PreshaderParameter)
-	{
-		// Get the global float4 component index (e.g. if this is 13, it refer to PreshaderBuffer[3].y)
-		const uint32 GlobalComponentOffset = PreshaderParameter->Analysis_PreshaderOffset;
-
-		EmitPreshaderBufferRead(PreshaderParameter->Type.GetPrimitive(), GlobalComponentOffset);
-
-		return NoOp;
-	}
-
 	bool HasMatchingScalarComponentCastChain(MIR::FValue* FirstComponent, MIR::FValue* CurrentComponent, int32 Index, int32 NumComponents)
 	{
 		const MIR::FSubscript* FirstComponentAsSubscript = FirstComponent->As<MIR::FSubscript>();
@@ -763,12 +838,10 @@ struct FEntryPointCodeGen
 				check(NumComponents >= 2 && NumComponents <= 4);
 				check(ArgAsPrimitiveUniform->Type.IsDouble());
 
-				const FUniformExpressionSet& UniformExpressionSet = Module->GetCompilationOutput().UniformExpressionSet;
-
-				// Get the global float4 component index (e.g. if this is 13, it refers to PreshaderBuffer[3].y)
-				const uint32 GlobalComponentOffset = UniformExpressionSet.GetNumericParameterEvaluationOffset(ArgAsPrimitiveUniform->Analysis_UniformIndex);
+				// Get the global float4 component index (e.g. if this is 13, it refer to PreshaderBuffer[3].y)
+				const uint32 GlobalComponentOffset = ArgAsPrimitiveUniform->Analysis_PreshaderOffset - 1;
 
-				EmitPreshaderBufferReadDoubleVector(ArgAsPrimitiveUniform->Type.GetPrimitive(), GlobalComponentOffset);
+				EmitPreshaderBufferReadDoubleVector(GlobalComponentOffset, ArgAsPrimitiveUniform->Analysis_NumComponents(), NumComponents);
 			}
 			else
 			{
diff --git a/Engine/Source/Runtime/Engine/Private/Materials/MaterialIRValueAnalyzer.cpp b/Engine/Source/Runtime/Engine/Private/Materials/MaterialIRValueAnalyzer.cpp
--- a/Engine/Source/Runtime/Engine/Private/Materials/MaterialIRValueAnalyzer.cpp
+++ b/Engine/Source/Runtime/Engine/Private/Materials/MaterialIRValueAnalyzer.cpp
@@ -163,64 +163,6 @@ static void Branch_Analyze(FMaterialIRValueAnalyzer& Analyzer, MIR::FBranch* Bra
 	MIR::ZeroArray(Branch->FalseBlock);
 }
 
-// Returns the next available offset into the preshader buffer for a float vector with the specified number of components (1-4).
-static uint32 NextGlobalComponentOffset(FMaterialIRValueAnalyzer& Analyzer, const MIR::FPrimitive& PrimitiveType)
-{
-	const uint32 NumComponents = PrimitiveType.NumComponents();
-	const bool bAreComponentsDouble = PrimitiveType.IsDouble();
-
-	// Following code calculates GlobalComponentOffset.
-	// GlobalComponentOffset is the i-th component in the array of float4s that make the uniform buffer.
-	// For example a GlobalComponentOffset of 13 references PreshaderBuffer[3].y.
-	// First, try to find an available sequence of free components in any previous allocation, in
-	// order to reduce the number of allocations and thus preshader buffer memory footprint.
-	// If the parameter type is too large and we can't find space for it in previous allocations,
-	// allocate a new uniform buffer slot (a float4, 16 bytes) and put any unused component
-	// in the appropriate freelist.
-
-	uint32 GlobalComponentOffset = UINT32_MAX;
-	check(NumComponents <= 4); // Only vectors supported for now
-
-	uint32 UsedNumComponents;
-
-	for (UsedNumComponents = NumComponents; UsedNumComponents < 4; ++UsedNumComponents)
-	{
-		auto& FreeOffsetsForCurrentNumComponents = Analyzer.FreeOffsetsPerNumComponents[UsedNumComponents - 1];
-
-		// If there are no available slots at this nubmer of components, try with a larger one
-		if (FreeOffsetsForCurrentNumComponents.IsEmpty())
-		{
-			continue;
-		}
-
-		// There is space for this number `NumComponents` in a a chunk with `i` free components
-		GlobalComponentOffset = FreeOffsetsForCurrentNumComponents.Last();
-		FreeOffsetsForCurrentNumComponents.Pop();
-
-		break;
-	}
-
-	// If UsedNumComponents is 4, it means we looked for all possible components sizes (1, 2 and 3) and
-	// could not find a chunk that would fit this parameter. Allocate a new chunk.
-	if (UsedNumComponents >= 4)
-	{
-		FUniformExpressionSet& UniformExpressionSet = Analyzer.CompilationOutput->UniformExpressionSet;
-		check(UsedNumComponents == 4);
-		const uint32 NumOfFloat4s = bAreComponentsDouble ? 2 : 1;
-		GlobalComponentOffset = UniformExpressionSet.AllocateFromUniformBuffer(NumOfFloat4s) * 4;
-	}
-
-	// Finally, check whether in the used slot there's any slack left, and if so, record it in the appropriate free-list.
-	int32 NumComponentsLeft = UsedNumComponents - NumComponents;
-	if (NumComponentsLeft > 0)
-	{
-		int32 LeftOverOffset = GlobalComponentOffset + NumComponents;
-		Analyzer.FreeOffsetsPerNumComponents[NumComponentsLeft - 1].Push(LeftOverOffset);
-	}
-
-	return GlobalComponentOffset;
-}
-
 static const MIR::FConstant* GetConstantComponent(const MIR::FValue* Constant, int32 Index)
 {
 	if (const MIR::FConstant* Scalar = MIR::As<MIR::FConstant>(Constant))
@@ -303,31 +245,11 @@ static void PrimitiveUniform_Analyze(FMaterialIRValueAnalyzer& Analyzer, MIR::FP
 	PrimitiveUniform->Analysis_UniformIndex = UniformExpressionSet.FindOrAddNumericParameter(MaterialParameterType, { FName(PrimitiveUniform->Name) }, DefaultValueOffset);
 
 	// Allocate the space for the parameter on the uniform buffer.
-	const uint32 GlobalComponentOffset = NextGlobalComponentOffset(Analyzer, PrimitiveType);
-	 
-	// Add the parameter evaluation to the uniform data
-	UniformExpressionSet.AddNumericParameterEvaluation(PrimitiveUniform->Analysis_UniformIndex, GlobalComponentOffset);
-	 
-	// If an insights instance is provided, add information about this parameter to it.
-	if (Analyzer.Insights)
-	{
-		// Push information about this parameter allocation to the insights.
-		FMaterialInsights::FUniformParameterAllocationInsight ParamInsight;
-		ParamInsight.BufferSlotIndex = GlobalComponentOffset / 4;
-		ParamInsight.BufferSlotOffset = GlobalComponentOffset % 4;
-		ParamInsight.ComponentsCount = PrimitiveType.NumComponents();
-		ParamInsight.ParameterName = FName(PrimitiveUniform->Name);
-	 
-		switch (PrimitiveType.ScalarKind)
-		{
-			case MIR::EScalarKind::Integer:    ParamInsight.ComponentType = FMaterialInsights::FUniformBufferSlotComponentType::CT_Int; break;
-			case MIR::EScalarKind::Float:  ParamInsight.ComponentType = FMaterialInsights::FUniformBufferSlotComponentType::CT_Float; break;
-			case MIR::EScalarKind::Double: ParamInsight.ComponentType = FMaterialInsights::FUniformBufferSlotComponentType::CT_LWC; break;
-			default: UE_MIR_UNREACHABLE();
-		}
+	int32 PreshaderIndex = Analyzer.Preshader.AddValueInfo(PrimitiveUniform);
 	 
-		Analyzer.Insights->UniformParameterAllocationInsights.Push(ParamInsight);
-	}
+	// Add the parameter evaluation to the uniform data.  Note that we don't know the used component count until the entire
+	// tree has been processed, so the final component count is filled in during fixup.
+	UniformExpressionSet.AddNumericParameterEvaluation(PrimitiveUniform->Analysis_UniformIndex, PreshaderIndex, MaterialParameterType == EMaterialParameterType::Scalar ? 1 : 4);
 }
 
 static int32 RegisterTextureUniform(FMaterialIRValueAnalyzer& Analyzer, FName Name, UObject* TextureObject, EMaterialValueType TextureValueType, int32 VTLayerIndex)
@@ -385,15 +307,23 @@ static void PreshaderParameter_Analyze(FMaterialIRValueAnalyzer& Analyzer, MIR::
 	// Only int, float and LWC parameters supported for now
 	check(PrimitiveType.IsInteger() || PrimitiveType.IsAnyFloat());
 
-	const uint32 GlobalComponentOffset = NextGlobalComponentOffset(Analyzer, PrimitiveType);
+	// Get Preshader2 type, and if it's an RVT double parameter, convert its type as LWC.  These are never read by
+	// other non-uniforms, and so are directly stored as LWC by the op.
+	EPreshader2Type Preshader2Type = FMaterialIRPreshader::ToPreshader2Type(PrimitiveType.ScalarKind);
+
+	if (Parameter->Opcode == EPreshader2Opcode::RuntimeVirtualTextureUniform && Preshader2Type == EPreshader2Type::Double)
+	{
+		Preshader2Type = EPreshader2Type::LWC;
+	}
+
+	int32 PreshaderIndex = Analyzer.Preshader.AddValueInfo(Parameter);
 
-	UE::Shader::FPreshaderData& PreshaderData = Analyzer.GetPreshaderData();
 	const FHashedMaterialParameterInfo HashedSourceParameterInfo{ SourceParameterName, EMaterialParameterAssociation::GlobalParameter, INDEX_NONE };
 
 	// UniformIndex argument only applies for RuntimeVirtualTextureUniform, but it's harmless to pass it in regardless.
-	UE::Preshader2::EmitTextureOp(PreshaderData, Parameter->Opcode, (uint16)GlobalComponentOffset, HashedSourceParameterInfo, Parameter->TextureIndex, Parameter->Payload.UniformIndex);
-
-	Parameter->Analysis_PreshaderOffset = GlobalComponentOffset;
+	UE::Preshader2::EmitTextureOp(Analyzer.Preshader.Data, Parameter->Opcode,
+		Preshader2Type, (uint8)PrimitiveType.NumComponents(), (uint16)PreshaderIndex,
+		HashedSourceParameterInfo, Parameter->TextureIndex, Parameter->Payload.UniformIndex);
 }
 
 static void SetMaterialOutput_Analyze(FMaterialIRValueAnalyzer& Analyzer, MIR::FSetMaterialOutput* SetMaterialOutput)
@@ -471,15 +401,14 @@ void FMaterialIRValueAnalyzer::Setup(UMaterialInterface* InMaterialInterface, FM
 	Insights = InInsights;
 	CompilationOutput = InCompilationOutput;
 	UniformPrimitiveDefaultValueOffsets = {};
-	FreeOffsetsPerNumComponents[0] = {};
-	FreeOffsetsPerNumComponents[1] = {};
-	FreeOffsetsPerNumComponents[2] = {};
 	EnvironmentDefines = {};
 	VTStacks = {};
 }
 
 void FMaterialIRValueAnalyzer::Analyze(MIR::FValue* Value)
 {
+	AnalyzeUses(Value);
+
 	#define EXPAND_CASE(ValueType) \
 		case MIR::VK_##ValueType: \
 		{ \
@@ -505,6 +434,14 @@ void FMaterialIRValueAnalyzer::Analyze(MIR::FValue* Value)
 	}
 
 	#undef EXPAND_CASE
+
+	// After analyzing the value, check if this is a preshader value (Uniform with a Parameter in its subtree),
+	// and if so, emit preshader code for it.  Note that the emitter only sets Uniform for operators that are
+	// supported, so we don't have to worry about hitting this code for unsupported operators.
+	if (Value->HasSubgraphProperties(MIR::EGraphProperties::Uniform | MIR::EGraphProperties::HasParameter))
+	{
+		Preshader.EmitValue(Value);
+	}
 }
 
 void FMaterialIRValueAnalyzer::AnalyzeInStage(MIR::FValue* Value, MIR::EStage Stage)
@@ -527,5 +464,194 @@ void FMaterialIRValueAnalyzer::AnalyzeInStage(MIR::FValue* Value, MIR::EStage St
 
 	#undef EXPAND_CASE
 }
-	
+
+static void RecurseMarkPreshaderTemp(MIR::FValue* Value)
+{
+	for (MIR::FValue* Use : Value->GetUses())
+	{
+		// Check if this value has already been marked as preshader (either Out or Temp).
+		// If so, we can also assume its children have already been marked.
+		if (Use &&
+			!Use->HasSubgraphProperties(MIR::EGraphProperties::PreshaderOut) &&
+			!Use->HasSubgraphProperties(MIR::EGraphProperties::PreshaderTemp) &&
+			Use->HasSubgraphProperties(MIR::EGraphProperties::HasParameter))
+		{
+			Use->SetSubgraphProperties(MIR::EGraphProperties::PreshaderTemp);
+
+			RecurseMarkPreshaderTemp(Use);
+		}
+	}
+}
+
+void FMaterialIRValueAnalyzer::MarkPreshaderOut(MIR::FValue* Value)
+{
+	// Check if this is the first time we're marking the value as a preshader output.
+	if (Value->HasSubgraphProperties(MIR::EGraphProperties::Uniform | MIR::EGraphProperties::HasParameter) && !Value->HasSubgraphProperties(MIR::EGraphProperties::PreshaderOut))
+	{
+		// Both Uniform and HasParameter are set, mark it as PreshaderOut, and its children as PreshaderTemp if not already marked.
+		Value->SetSubgraphProperties(MIR::EGraphProperties::PreshaderOut);
+		Value->GraphProperties &= ~MIR::EGraphProperties::PreshaderTemp;
+
+		// If this is an alias, mark the underlying value as PreshaderOut as well.  This may waste some space in the output buffer
+		// if there is a mix of used and unused components, but this is assumed to be rare, and we want to err on the side of avoiding
+		// the need to generate extra move opcodes, to save CPU perf.
+		if (Value->Analysis_PreshaderOffset)
+		{
+			MIR::FValue* Alias = Preshader.ValueInfos[Value->Analysis_PreshaderOffset].Alias;
+			if (Alias)
+			{
+				Alias->SetSubgraphProperties(MIR::EGraphProperties::PreshaderOut);
+				Alias->GraphProperties &= ~MIR::EGraphProperties::PreshaderTemp;
+			}
+		}
+
+		RecurseMarkPreshaderTemp(Value);
+
+		// Composite is emitted here, as their emission is deferred, assuming composites will often be folded into operators or
+		// dead stripped by the trivial tail swizzle optimization, unless there is a case where the final result actually needs
+		// to be available on the GPU (the case we are encountering here).
+		if (!Value->Analysis_PreshaderOffset)
+		{
+			if (MIR::FComposite* Composite = Value->As<MIR::FComposite>())
+			{
+				Preshader.EmitComposite(Composite);
+			}
+			else
+			{
+				// All other preshader compatible child expressions should already have been processed.
+				UE_MIR_UNREACHABLE();
+			}
+		}
+	}
+}
+
+void FMaterialIRValueAnalyzer::AnalyzeUses(MIR::FValue* Value)
+{
+	for (MIR::FValue* Use : Value->GetUses())
+	{
+		if (!Use)
+		{
+			continue;
+		}
+
+		// Updated used component mask of uniform values.  Can't be done in UniformParameter_Analyze, because it needs
+		// access to the parent Value to fetch its subscript index.
+		if (MIR::FPrimitiveUniform* UniformParameter = Use->As<MIR::FPrimitiveUniform>())
+		{
+			if (UniformParameter->Type.IsScalar())
+			{
+				UniformParameter->Analysis_ComponentMask |= (1 << 0);
+			}
+			else if (UniformParameter->Type.IsVector())
+			{
+				if (Value->IsA(MIR::VK_Subscript))
+				{
+					UniformParameter->Analysis_ComponentMask |= 1 << FMath::Clamp(Value->As<MIR::FSubscript>()->Index, 0, 3);
+				}
+				else
+				{
+					UniformParameter->Analysis_ComponentMask |= 0xf;
+				}
+			}
+		}
+
+		// If this graph node is not Uniform, attempt to mark any of its children as PreshaderOut where necessary.
+		if (!(Value->GraphProperties & MIR::EGraphProperties::Uniform))
+		{
+			// Trivial Tail Swizzle optimization -- if this use is a swizzle (Composite of Subscripts all pointing to the same
+			// vector input value, or all identical values), we want to mark that underlying value as "PreshaderOut" rather than
+			// this one.  This prevents trivial tail swizzles from generating extra CPU preshader opcodes, given that swizzle on
+			// load is free.  For example, this is a common pattern that occurs when 3 components of a vector param are fetched:
+			//
+			//     _0 = Composite( Subscript(Param,0), Subscript(Param,1), Subscript(Param,2) );
+			//
+			// We're trying to avoid this generating an ".xyz" swizzle on the CPU.  For this example, the swizzle has the same
+			// order of components, but this optimization also applies if the components are reordered.  We want to apply this
+			// optimization here, to simplify downstream logic, which may involve multiple passes over the tree.  We can also
+			// accept constants as any of the parameters, as swizzling in a constant is also free.
+			bool bTrivialTailSwizzle = false;
+
+			if (Use->IsA(MIR::VK_Composite) && Use->Type.IsVector() && Use->GetUses().Num() > 0)
+			{
+				TConstArrayView<MIR::FValue*> Components = Use->As<MIR::FComposite>()->GetComponents();
+				MIR::FSubscript* FirstSubscript = nullptr;
+				MIR::FValue* FirstValue = nullptr;
+
+				// Assume it's trivial until we find a case where it's not.
+				bTrivialTailSwizzle = true;
+
+				for (int32 ComponentIndex = 0; ComponentIndex < Components.Num(); ComponentIndex++)
+				{
+					// Constants can just be ignored
+					if (Components[ComponentIndex]->IsA(MIR::VK_Constant))
+					{
+						continue;
+					}
+
+					if (Components[ComponentIndex]->IsA(MIR::VK_Subscript))
+					{
+						// If there's a mix of (non-constant) values and subscripts, it's not compatible with trivial tail swizzle!
+						if (FirstValue)
+						{
+							bTrivialTailSwizzle = false;
+							break;
+						}
+
+						MIR::FSubscript* CurrentSubscript = Components[ComponentIndex]->As<MIR::FSubscript>();
+						if (!FirstSubscript)
+						{
+							FirstSubscript = CurrentSubscript;
+						}
+						else if (FirstSubscript->Arg != CurrentSubscript->Arg)
+						{
+							bTrivialTailSwizzle = false;
+							break;
+						}
+					}
+					else
+					{
+						// If there's a mix of (non-constant) values and subscripts, it's not compatible with trivial tail swizzle!
+						if (FirstSubscript)
+						{
+							bTrivialTailSwizzle = false;
+							break;
+						}
+
+						if (!FirstValue)
+						{
+							FirstValue = Components[ComponentIndex];
+						}
+						else if (FirstValue != Components[ComponentIndex])
+						{
+							bTrivialTailSwizzle = false;
+							break;
+						}
+					}
+				}
+
+				if (bTrivialTailSwizzle)
+				{
+					if (FirstSubscript)
+					{
+						// Subscripts are always aliases of the underlying arg, so we need to mark that as used by non-uniform.
+						MarkPreshaderOut(FirstSubscript->Arg);
+					}
+					else if (FirstValue)
+					{
+						MarkPreshaderOut(FirstValue);
+					}
+
+					Use->SetSubgraphProperties(MIR::EGraphProperties::TrivialTailSwizzle);
+				}
+			}
+
+			// Trivial tail swizzles are handled above.
+			if (!bTrivialTailSwizzle)
+			{
+				MarkPreshaderOut(Use);
+			}
+		}
+	}
+}
+
 #endif // #if WITH_EDITOR
diff --git a/Engine/Source/Runtime/Engine/Private/Materials/MaterialIRValueAnalyzer.h b/Engine/Source/Runtime/Engine/Private/Materials/MaterialIRValueAnalyzer.h
--- a/Engine/Source/Runtime/Engine/Private/Materials/MaterialIRValueAnalyzer.h
+++ b/Engine/Source/Runtime/Engine/Private/Materials/MaterialIRValueAnalyzer.h
@@ -10,6 +10,7 @@
 #include "Materials/MaterialIRTypes.h"
 #include "Materials/MaterialIRInternal.h"
 #include "Materials/MaterialInsights.h"
+#include "Materials/MaterialIRPreshader.h"
 #include "MaterialShared.h"
 #include "Engine/Texture.h"
 
@@ -54,6 +55,14 @@ struct FMaterialIRValueAnalyzer
 	// Return writable preshader data.
 	UE::Shader::FPreshaderData& GetPreshaderData() { return CompilationOutput->UniformExpressionSet.UniformPreshaderData; }
 
+private:
+	// Marks the given value as PreshaderOut, if the value's properties allow it, and it's not already marked.
+	void MarkPreshaderOut(MIR::FValue* Value);
+
+	// Handles logic that involves analyzing Uses of the given value.  Updates used component masks and sets PreshaderOut where needed.
+	void AnalyzeUses(MIR::FValue* Value);
+
+public:
 	// The material being built.
 	UMaterialInterface* MaterialInterface{};
 	
@@ -68,13 +77,13 @@ struct FMaterialIRValueAnalyzer
 		
 	// Maps default values to their default-value offset, as used in UniformExpressionSet (e.g. UniformExpressionSet.AddDefaultParameterValue() and UniformExpressionSet.FindOrAddNumericParameter()).
 	TMap<UE::Shader::FValue, uint32> UniformPrimitiveDefaultValueOffsets{};
-		
-	// Stores free uniform buffer offsets for 1, 2, and 3 leftover components to optimize float4 packing
-	TArray<uint32, TInlineAllocator<8>> FreeOffsetsPerNumComponents[3];
 
 	// List of enabled shader environment defines.
 	TSet<FName> EnvironmentDefines;
 
+	// Preshader intermediate representation
+	FMaterialIRPreshader Preshader;
+
 	// Translator specific structure during generation of VT stacks. This is later converted to FMaterialVirtualTextureStack entries in FUniformExpressionSet.
 	// The equivalent in the old translator is FHLSLMaterialTranslator::FMaterialVTStackEntry.
 	struct FVTStackEntry
diff --git a/Engine/Source/Runtime/Engine/Private/Materials/MaterialRenderProxy.cpp b/Engine/Source/Runtime/Engine/Private/Materials/MaterialRenderProxy.cpp
--- a/Engine/Source/Runtime/Engine/Private/Materials/MaterialRenderProxy.cpp
+++ b/Engine/Source/Runtime/Engine/Private/Materials/MaterialRenderProxy.cpp
@@ -140,11 +140,12 @@ public:
 
 					for (const FItem& Item : Items)
 					{
-						uint8* TempBuffer = FMemStack::Get().PushBytes(Item.UniformBufferLayout->ConstantBufferSize, SHADER_PARAMETER_STRUCT_ALIGNMENT);
+						uint32 TempBufferSize = Item.UniformBufferLayout->ConstantBufferSize + Item.UniformExpressionSet->GetTemporaryBufferSize() * (4 * sizeof(float));
+						uint8* TempBuffer = FMemStack::Get().PushBytes(TempBufferSize, SHADER_PARAMETER_STRUCT_ALIGNMENT);
 
 						FMaterialRenderContext Context(Item.MaterialRenderProxy, *Item.Material, Item.bShowSelection);
 
-						Item.UniformExpressionSet->FillUniformBuffer(Context, Item.AllocatedVTs, Item.UniformBufferLayout, TempBuffer, Item.UniformBufferLayout->ConstantBufferSize);
+						Item.UniformExpressionSet->FillUniformBuffer(Context, Item.AllocatedVTs, Item.UniformBufferLayout, TempBuffer, TempBufferSize);
 				
 						if (Item.UniformBufferVS)
 						{
@@ -598,8 +599,9 @@ void FMaterialRenderProxy::EvaluateUniformExpressions(FRHICommandListBase& RHICm
 	else
 	{
 		FMemMark Mark(FMemStack::Get());
-		uint8* TempBuffer = FMemStack::Get().PushBytes(UniformBufferLayout->ConstantBufferSize, SHADER_PARAMETER_STRUCT_ALIGNMENT);
-		UniformExpressionSet.FillUniformBuffer(Context, OutUniformExpressionCache, UniformBufferLayout, TempBuffer, UniformBufferLayout->ConstantBufferSize);
+		uint32 TempBufferSize = UniformBufferLayout->ConstantBufferSize + UniformExpressionSet.GetTemporaryBufferSize() * (4 * sizeof(float));
+		uint8* TempBuffer = FMemStack::Get().PushBytes(TempBufferSize, SHADER_PARAMETER_STRUCT_ALIGNMENT);
+		UniformExpressionSet.FillUniformBuffer(Context, OutUniformExpressionCache, UniformBufferLayout, TempBuffer, TempBufferSize);
 
 		auto CreateOrUpdateUniformBuffer = [](FRHICommandListBase& RHICmdList, FUniformBufferRHIRef& UniformBuffer, const uint8* TempBuffer, const FRHIUniformBufferLayout* UniformBufferLayout)
 		{
diff --git a/Engine/Source/Runtime/Engine/Private/Materials/MaterialUniformExpressions.cpp b/Engine/Source/Runtime/Engine/Private/Materials/MaterialUniformExpressions.cpp
--- a/Engine/Source/Runtime/Engine/Private/Materials/MaterialUniformExpressions.cpp
+++ b/Engine/Source/Runtime/Engine/Private/Materials/MaterialUniformExpressions.cpp
@@ -1110,19 +1110,32 @@ uint32 FUniformExpressionSet::AddDefaultParameterValue(const UE::Shader::FValue&
 	return Offset;
 }
 
-void FUniformExpressionSet::AddNumericParameterEvaluation(uint32 ParameterIndex, uint32 BufferOffset)
+void FUniformExpressionSet::AddNumericParameterEvaluation(uint32 ParameterIndex, uint32 BufferOffset, uint32 ComponentCount)
 {
 	check(ParameterIndex <= UINT16_MAX);
 	check(BufferOffset <= UINT16_MAX);
+	check(ComponentCount <= 4);
 	UniformParameterEvaluations.Add({
 		.ParameterIndex = (uint16)ParameterIndex,
-		.BufferOffset = (uint16)BufferOffset
+		.BufferOffset = (uint16)BufferOffset,
+		.ComponentCount = (uint16)ComponentCount
 	});
 }
 
-uint32 FUniformExpressionSet::GetNumericParameterEvaluationOffset(uint32 ParameterIndex) const
+void FUniformExpressionSet::FixupNumericParameterEvaluations(TConstArrayView<uint16> Remap, TConstArrayView<uint8> UniformComponentCounts)
 {
-	return UniformParameterEvaluations[ParameterIndex].BufferOffset;
+	for (FMaterialUniformParameterEvaluation& Evaluation : UniformParameterEvaluations)
+	{
+		// Override component count if provided
+		uint16 ComponentCount = UniformComponentCounts.Num() ? UniformComponentCounts[Evaluation.BufferOffset] : 0;
+		if (ComponentCount)
+		{
+			Evaluation.ComponentCount = ComponentCount;
+		}
+
+		// Then remap BufferOffset
+		Evaluation.BufferOffset = Remap[Evaluation.BufferOffset];
+	}
 }
 
 uint32 FUniformExpressionSet::AllocateFromUniformBuffer(uint32 NumOfFloat4s)
@@ -1529,7 +1542,10 @@ void FUniformExpressionSet::FillUniformBuffer(const FMaterialRenderContext& Mate
 				ParameterShaderValue = GetDefaultParameterValue(Parameter.ParameterType, Parameter.DefaultValueOffset);
 			}
 
-			// Copy the value to the target position
+			// Copy the value to the target position, after modifying the type and value to the desired component count.
+			ParameterShaderValue.Type.ValueType = MakeValueType(ParameterShaderValue.Type.ValueType, Evaluation.ComponentCount);
+			ParameterShaderValue.Component.SetNum(Evaluation.ComponentCount);
+
 			CopyValueToUniformBuffer(ParameterShaderValue, (float*)BufferCursor, Evaluation.BufferOffset);
 		}
 
@@ -1538,7 +1554,7 @@ void FUniformExpressionSet::FillUniformBuffer(const FMaterialRenderContext& Mate
 		if (UniformPreshaderData.bPreshader2)
 		{
 			// Preshader2 executes a single bytecode stream to write all outputs to the buffer.  UniformPreshaderBufferSize is a count of float4.
-			FPreshaderDataContext Preshader2Context(UniformPreshaderData, TArrayView<float>(PreshaderBuffer, UniformPreshaderBufferSize*4));
+			FPreshaderDataContext Preshader2Context(UniformPreshaderData, TArrayView<float>(PreshaderBuffer, (UniformPreshaderBufferSize + UniformPreshaderData.Preshader2TemporarySize)*4));
 			UE::Preshader2::Evaluate(Preshader2Context, &MaterialRenderContext);
 		}
 		else
diff --git a/Engine/Source/Runtime/Engine/Private/Shader/Preshader2.cpp b/Engine/Source/Runtime/Engine/Private/Shader/Preshader2.cpp
--- a/Engine/Source/Runtime/Engine/Private/Shader/Preshader2.cpp
+++ b/Engine/Source/Runtime/Engine/Private/Shader/Preshader2.cpp
@@ -46,6 +46,34 @@ static const TCHAR* GPreshader2OpcodeNames[] =
 #undef PRESHADER2_BINARY_OP_CUSTOM
 #undef PRESHADER2_OP
 
+static const TCHAR* GPreshader2TypeNames[] =
+{
+	TEXT("float"),		// EPreshader2Type::Float
+	TEXT("double"),		// EPreshader2Type::Double
+	TEXT("int"),		// EPreshader2Type::Integer,
+	TEXT("LWC"),		// EPreshader2Type::LWC
+};
+
+static const TCHAR* GPreshader2DimensionNames[] =
+{
+	TEXT("?"),			// Invalid zero dimension
+	TEXT(""),			// Use an empty string for one
+	TEXT("2"),
+	TEXT("3"),
+	TEXT("4"),
+	TEXT("5"),			// Fill out the array to avoid static analysis warnings when reading from it,
+	TEXT("6"),			// even though valid opcodes won't have dimensions above 4 set.
+	TEXT("7"),
+};
+
+// Table of swizzles by first component index and component count
+static const TCHAR* GPreshader2LinearSwizzles[4][5] =
+{
+	{ TEXT(".?"), TEXT(".x"), TEXT(".xy"), TEXT(".xyz"), TEXT(".xyzw") },
+	{ TEXT(".?"), TEXT(".y"), TEXT(".yz"), TEXT(".yzw"), TEXT(".????") },
+	{ TEXT(".?"), TEXT(".z"), TEXT(".zw"), TEXT(".???"), TEXT(".????") },
+	{ TEXT(".?"), TEXT(".w"), TEXT(".??"), TEXT(".???"), TEXT(".????") },
+};
 
 struct FOpAbs { template<typename T> static T Call(T Value) { return FMath::Abs(Value); } };
 struct FOpACos { template<typename T> static T Call(T Value) { return FMath::Acos(Value); } };
@@ -112,6 +140,7 @@ struct FOpDot {};
 // Custom.
 struct FOpConstant {};
 struct FOpMove {};
+struct FOpMoveForFixup {};
 
 // External Inputs.
 struct FOpExternalTextureCoordinateOffset {};
@@ -121,6 +150,10 @@ struct FOpSparseVolumeTextureUniform {};
 struct FOpTexelSize {};
 struct FOpTextureSize {};
 
+// Forward declarations.
+static void AppendOffset(TMemoryImageArray<uint8>& Opcodes, uint16 Offset);
+static void EnablePreshader2(UE::Shader::FPreshaderData& PreshaderData);
+
 static float& AsFloat(void* Buffer, int32 Offset, int32 Index)
 {
 	return ((float*)Buffer)[Offset + Index];
@@ -169,7 +202,7 @@ static FORCEINLINE void BoundsCheck(TArrayView<float> Buffer, int32 ResultOffset
 	check(FMath::Max(AOffset, ResultOffset) + Dimension <= Buffer.Num());
 }
 
-static FORCEINLINE void BoundsCheck(TArrayView<float> Buffer, int32 ResultOffset, int32 AOffset, int32 BOffset, int32 Dimension, bool bAScalar = false, bool bBScalar = false, bool bDouble = false)
+static FORCEINLINE void BoundsCheck(TArrayView<float> Buffer, int32 ResultOffset, int32 AOffset, int32 BOffset, int32 Dimension, bool bAScalar, bool bBScalar, bool bDouble = false)
 {
 	if (bDouble)
 	{
@@ -539,66 +572,68 @@ void BinaryEvaluateConst(TArrayView<float> Buffer, EPreshader2Type Type, int32 D
 }
 
 template<typename T>
-void CrossCompute(T* Result, T* A, T* B)
+void CrossCompute(T* Result, T* A, T* B, int32 AInc, int32 BInc)
 {
 	T Temp[3];
-	Temp[0] = A[1] * B[2] - A[2] * B[1];
-	Temp[1] = A[2] * B[0] - A[0] * B[2];
-	Temp[2] = A[0] * B[1] - A[1] * B[0];
+	Temp[0] = A[1*AInc] * B[2*BInc] - A[2*AInc] * B[1*BInc];
+	Temp[1] = A[2*AInc] * B[0*BInc] - A[0*AInc] * B[2*BInc];
+	Temp[2] = A[0*AInc] * B[1*BInc] - A[1*AInc] * B[0*BInc];
 	Result[0] = Temp[0];
 	Result[1] = Temp[1];
 	Result[2] = Temp[2];
 }
 
 // Specialization of BinaryEvaluate for Cross.
-static void BinaryEvaluate(const FOpCross&, TArrayView<float> Buffer, EPreshader2Type Type, int32 Dimension, int32 ResultOffset, int32 AOffset, int32 BOffset)
+static void BinaryEvaluate(const FOpCross&, TArrayView<float> Buffer, EPreshader2Type Type, int32 Dimension, int32 ResultOffset, int32 AOffset, int32 BOffset, int32 AInc, int32 BInc)
 {
 	float* BufferData = Buffer.GetData();
 
 	// Ignores dimension, assuming dimension of 3 (we check this at creation of the op, not here at runtime)
 	if (LIKELY(Type == EPreshader2Type::Float))
 	{
-		BoundsCheck(Buffer, ResultOffset, AOffset, BOffset, 3);
-		CrossCompute(&AsFloat(BufferData, ResultOffset, 0), &AsFloat(BufferData, AOffset, 0), &AsFloat(BufferData, BOffset, 0));
+		BoundsCheck(Buffer, ResultOffset, AOffset, BOffset, 3, AInc == 0, BInc == 0);
+		CrossCompute(&AsFloat(BufferData, ResultOffset, 0), &AsFloat(BufferData, AOffset, 0), &AsFloat(BufferData, BOffset, 0), AInc, BInc);
 	}
 	else if (Type == EPreshader2Type::Double)
 	{
-		BoundsCheck(Buffer, ResultOffset, AOffset, BOffset, 3, false, false, /*bDouble=*/ true);
-		CrossCompute(&AsDouble(BufferData, ResultOffset, 0), &AsDouble(BufferData, AOffset, 0), &AsDouble(BufferData, BOffset, 0));
+		BoundsCheck(Buffer, ResultOffset, AOffset, BOffset, 3, AInc == 0, BInc == 0, /*bDouble=*/ true);
+		CrossCompute(&AsDouble(BufferData, ResultOffset, 0), &AsDouble(BufferData, AOffset, 0), &AsDouble(BufferData, BOffset, 0), AInc, BInc);
 	}
 	else
 	{
-		BoundsCheck(Buffer, ResultOffset, AOffset, BOffset, 3);
-		CrossCompute(&AsInt(BufferData, ResultOffset, 0), &AsInt(BufferData, AOffset, 0), &AsInt(BufferData, BOffset, 0));
+		BoundsCheck(Buffer, ResultOffset, AOffset, BOffset, 3, AInc == 0, BInc == 0);
+		CrossCompute(&AsInt(BufferData, ResultOffset, 0), &AsInt(BufferData, AOffset, 0), &AsInt(BufferData, BOffset, 0), AInc, BInc);
 	}
 }
 
 // Specialization of BinaryEvaluate for Dot.  Accepts Dimension as a reference by design!
-static void BinaryEvaluate(const FOpDot&, TArrayView<float> Buffer, EPreshader2Type Type, int32& Dimension, int32 ResultOffset, int32 AOffset, int32 BOffset)
+static void BinaryEvaluate(const FOpDot&, TArrayView<float> Buffer, EPreshader2Type Type, int32& Dimension, int32 ResultOffset, int32 AOffset, int32 BOffset, int32 AInc, int32 BInc)
 {
 	float* BufferData = Buffer.GetData();
 
 	if (LIKELY(Type == EPreshader2Type::Float))
 	{
-		BoundsCheck(Buffer, AOffset, BOffset, Dimension);
+		BoundsCheck(Buffer, AOffset, AInc ? Dimension : 1);
+		BoundsCheck(Buffer, BOffset, BInc ? Dimension : 1);
 		BoundsCheck(Buffer, ResultOffset, 1);
 
 		float DotSum = 0.0f;
 		for (int32 Index = 0; Index < Dimension; ++Index)
 		{
-			DotSum += AsFloat(BufferData, AOffset, Index) * AsFloat(BufferData, BOffset, Index);
+			DotSum += AsFloat(BufferData, AOffset, AInc*Index) * AsFloat(BufferData, BOffset, BInc*Index);
 		}
 		AsFloat(BufferData, ResultOffset, 0) = DotSum;
 	}
 	else
 	{
-		BoundsCheck(Buffer, AOffset, BOffset, Dimension, /*bDouble=*/ true);
+		BoundsCheck(Buffer, AOffset, AInc ? Dimension : 1, /*bDouble=*/ true);
+		BoundsCheck(Buffer, BOffset, BInc ? Dimension : 1, /*bDouble=*/ true);
 		BoundsCheck(Buffer, ResultOffset, 1, /*bDouble=*/ true);
 
 		double DotSum = 0.0;
 		for (int32 Index = 0; Index < Dimension; ++Index)
 		{
-			DotSum += AsDouble(BufferData, AOffset, Index) * AsDouble(BufferData, BOffset, Index);
+			DotSum += AsDouble(BufferData, AOffset, AInc*Index) * AsDouble(BufferData, BOffset, BInc*Index);
 		}
 		AsDouble(BufferData, ResultOffset, 0) = DotSum;
 	}
@@ -624,6 +659,23 @@ static void OpEvaluate(const FOpConstant&, const FMaterialRenderContext* Context
 	OpcodePtr = DataPtr + CopySize;
 }
 
+static void OpFixup(const FOpConstant&, UE::Shader::FPreshaderData& PreshaderOut, TConstArrayView<uint16> Remap, const uint8*& OpcodePtr, EPreshader2Type& Type, int32& Dimension, int32& ResultOffset)
+{
+	// Opcode(1), ResultOffset(2), Dimension+Type(1), Data(N)
+	ResultOffset = Remap[FetchOffset(OpcodePtr)];
+	Dimension = OpcodePtr[2];
+	Type = (EPreshader2Type)(Dimension >> 6);
+	Dimension = Dimension & 0x7;
+
+	const uint8* DataPtr = OpcodePtr + 3;
+	int32 ElementSize = Type == EPreshader2Type::Double ? 8 : 4;
+	int32 CopySize = Dimension * ElementSize;
+
+	EmitConstantOp(PreshaderOut, Type, Dimension, ResultOffset, DataPtr);
+
+	OpcodePtr = DataPtr + CopySize;
+}
+
 static void OpEvaluate(const FOpMove&, const FMaterialRenderContext* Context, TArrayView<const FScriptName> Names, TArrayView<float> Buffer, const uint8*& OpcodePtr, EPreshader2Type& Type, int32& Dimension, int32& ResultOffset)
 {
 	float* BufferData = Buffer.GetData();
@@ -689,6 +741,44 @@ static void OpEvaluate(const FOpMove&, const FMaterialRenderContext* Context, TA
 	}
 }
 
+static void OpFixup(const FOpMove&, UE::Shader::FPreshaderData& PreshaderOut, TConstArrayView<uint16> Remap, const uint8*& OpcodePtr, EPreshader2Type& Type, int32& Dimension, int32& ResultOffset)
+{
+	check(!"Attempting to fixup wrong version of Move op!  Use EmitMoveOpForFixup when using FixupAndCompress.");
+	UE_ASSUME(false);
+}
+
+static void OpEvaluate(const FOpMoveForFixup&, const FMaterialRenderContext* Context, TArrayView<const FScriptName> Names, TArrayView<float> Buffer, const uint8*& OpcodePtr, EPreshader2Type& Type, int32& Dimension, int32& ResultOffset)
+{
+	check(!"Attempting to evaluate wrong version of Move op!  Must call FixupAndCompress after using EmitMoveOpForFixup.");
+	UE_ASSUME(false);
+}
+
+static void OpFixup(const FOpMoveForFixup&, UE::Shader::FPreshaderData& PreshaderOut, TConstArrayView<uint16> Remap, const uint8*& OpcodePtr, EPreshader2Type& Type, int32& Dimension, int32& ResultOffset)
+{
+	Dimension = OpcodePtr[0];
+	ResultOffset = Remap[FetchOffset(OpcodePtr + 1)];
+	uint8 ComponentStart = OpcodePtr[3];
+	Type = (EPreshader2Type)(Dimension >> 6);
+	Dimension = Dimension & 0x7;
+	OpcodePtr += 4;
+
+	int32 ElementSizeInWords = Type == EPreshader2Type::Double ? 2 : 1;
+	ResultOffset += ElementSizeInWords * ComponentStart;
+
+	TArray<uint16, TFixedAllocator<4>> SourceOffsets;
+	SourceOffsets.SetNumUninitialized(Dimension);
+	for (int32 RegisterIndex = 0; RegisterIndex < Dimension; RegisterIndex++, OpcodePtr += 2)
+	{
+		SourceOffsets[RegisterIndex] = Remap[FetchOffset(OpcodePtr)];
+	}
+	for (int32 RegisterIndex = 0; RegisterIndex < Dimension; RegisterIndex++, OpcodePtr++)
+	{
+		SourceOffsets[RegisterIndex] += ElementSizeInWords * OpcodePtr[0];
+	}
+
+	EmitMoveOp(PreshaderOut, Type, ResultOffset, SourceOffsets);
+}
+
 static const uint8* FetchHashedMaterialParameterInfo(TArrayView<const FScriptName> Names, const uint8* OpcodePtr, FHashedMaterialParameterInfo& OutParameterInfo)
 {
 	// Parameter name
@@ -703,6 +793,13 @@ static const uint8* FetchHashedMaterialParameterInfo(TArrayView<const FScriptNam
 	return OpcodePtr;
 }
 
+// Should match "Fetch" function above, but just advances OpcodePtr without reading anything
+static const uint8* SkipHashedMaterialParameterInfo(const uint8* OpcodePtr)
+{
+	const FHashedMaterialParameterInfo* Dummy = (const FHashedMaterialParameterInfo*)nullptr;
+	return OpcodePtr + sizeof(uint16) + sizeof(Dummy->Index) + sizeof(Dummy->Association);
+}
+
 static const uint8* GetTextureParameter(const FMaterialRenderContext* Context, TArrayView<const FScriptName> Names, const uint8* OpcodePtr, const UTexture*& OutTexture, const USparseVolumeTexture*& OutSparseVolumeTexture)
 {
 	// Parameter info.
@@ -726,15 +823,27 @@ static const uint8* GetTextureParameter(const FMaterialRenderContext* Context, T
 	return OpcodePtr;
 }
 
+static const uint8* SkipTextureParameter(const uint8* OpcodePtr)
+{
+	OpcodePtr = SkipHashedMaterialParameterInfo(OpcodePtr);
+	return OpcodePtr + sizeof(int32);							// TextureIndex
+}
+
 static void OpEvaluate(const FOpRuntimeVirtualTextureUniform&, const FMaterialRenderContext* Context, TArrayView<const FScriptName> Names, TArrayView<float> Buffer, const uint8*& OpcodePtr, EPreshader2Type& Type, int32& Dimension, int32& ResultOffset)
 {
 	check(Context);
 
 	ResultOffset = FetchOffset(OpcodePtr);
+	Dimension = OpcodePtr[2];
+	Type = (EPreshader2Type)(Dimension >> 6);
+	Dimension = Dimension & 0x7;
+
+	int32 WordCount = Dimension * (Type == EPreshader2Type::Float ? 1 : 2);
+	BoundsCheck(Buffer, ResultOffset, WordCount);
 
 	// Parameter info.
 	FHashedMaterialParameterInfo ParameterInfo;
-	OpcodePtr = FetchHashedMaterialParameterInfo(Names, OpcodePtr + 2, ParameterInfo);
+	OpcodePtr = FetchHashedMaterialParameterInfo(Names, OpcodePtr + 3, ParameterInfo);
 
 	// Other data.
 	int32 TextureIndex;
@@ -750,17 +859,50 @@ static void OpEvaluate(const FOpRuntimeVirtualTextureUniform&, const FMaterialRe
 	if (Texture && VectorIndex != INDEX_NONE)
 	{
 		FVector4 Value = Texture->GetUniformParameter(VectorIndex);
-		FDFVector4 DFValue(Value);
-		BoundsCheck(Buffer, ResultOffset, 8);
-		FMemory::Memcpy(Buffer.GetData() + ResultOffset, &DFValue.High[0], 8 * sizeof(float));
+		if (Type == EPreshader2Type::Float)
+		{
+			FVector4f ValueF = FVector4f(Value);
+			FMemory::Memcpy(Buffer.GetData() + ResultOffset, &ValueF, Dimension * sizeof(float));
+		}
+		else
+		{
+			// Assume double types for runtime virtual textures are LWC.
+			check(Type == EPreshader2Type::LWC);
+
+			FDFVector4 ValueDF(Value);
+			FMemory::Memcpy(Buffer.GetData() + ResultOffset,             &ValueDF.High[0], Dimension * sizeof(float));
+			FMemory::Memcpy(Buffer.GetData() + ResultOffset + Dimension, &ValueDF.Low[0],  Dimension * sizeof(float));
+		}
 	}
 	else
 	{
-		BoundsCheck(Buffer, ResultOffset, 8);
-		FMemory::Memset(Buffer.GetData() + ResultOffset, 0, 8 * sizeof(float));
+		FMemory::Memset(Buffer.GetData() + ResultOffset, 0, WordCount * sizeof(float));
 	}
 }
 
+static void OpFixup(const FOpRuntimeVirtualTextureUniform&, UE::Shader::FPreshaderData& PreshaderOut, TConstArrayView<uint16> Remap, const uint8*& OpcodePtr, EPreshader2Type& Type, int32& Dimension, int32& ResultOffset)
+{
+	// We are emitting code here without explicitly calling the Emit function, so mark this as Preshader2 enabled
+	EnablePreshader2(PreshaderOut);
+
+	ResultOffset = Remap[FetchOffset(OpcodePtr)];
+	uint8 DimensionAndType = OpcodePtr[2];
+	Type = (EPreshader2Type)(DimensionAndType >> 6);
+	Dimension = DimensionAndType & 0x7;
+	OpcodePtr += 3;
+
+	// We can just append the data after the offset, which doesn't need fixup -- simpler than expanding the data to pass to the original emit function.
+	const uint8* PayloadStart = OpcodePtr;
+
+	OpcodePtr = SkipHashedMaterialParameterInfo(OpcodePtr);
+	OpcodePtr += sizeof(int32) + sizeof(int32);		// TextureIndex, VectorIndex
+
+	PreshaderOut.Data.Add((uint8)EPreshader2Opcode::RuntimeVirtualTextureUniform);
+	AppendOffset(PreshaderOut.Data, ResultOffset);
+	PreshaderOut.Data.Add(DimensionAndType);
+	PreshaderOut.Data.Append(PayloadStart, OpcodePtr - PayloadStart);
+}
+
 void OpEvaluate(const FOpTexelSize&, const FMaterialRenderContext* Context, TArrayView<const FScriptName> Names, TArrayView<float> Buffer, const uint8*& OpcodePtr, EPreshader2Type& Type, int32& Dimension, int32& ResultOffset)
 {
 	check(Context);
@@ -768,38 +910,42 @@ void OpEvaluate(const FOpTexelSize&, const FMaterialRenderContext* Context, TArr
 	float* BufferData = Buffer.GetData();
 
 	ResultOffset = FetchOffset(OpcodePtr);
-	BoundsCheck(Buffer, ResultOffset, 3);		// Writes 3 floats to output
+	Dimension = OpcodePtr[2];
+	Type = (EPreshader2Type)(Dimension >> 6);
+	Dimension = Dimension & 0x7;
+
+	BoundsCheck(Buffer, ResultOffset, Dimension);
 
 	const UTexture* Texture;
 	const USparseVolumeTexture* SparseVolumeTexture;
-	OpcodePtr = GetTextureParameter(Context, Names, OpcodePtr + 2, Texture, SparseVolumeTexture);
+	OpcodePtr = GetTextureParameter(Context, Names, OpcodePtr + 3, Texture, SparseVolumeTexture);
 
-	uint32 SizeX;
-	uint32 SizeY;
-	uint32 SizeZ;
+	uint32 Sizes[3];
 	if (Texture && Texture->GetResource())
 	{
-		SizeX = Texture->GetResource()->GetSizeX();
-		SizeY = Texture->GetResource()->GetSizeY();
-		SizeZ = Texture->GetResource()->GetSizeZ();
+		Sizes[0] = Texture->GetResource()->GetSizeX();
+		Sizes[1] = Texture->GetResource()->GetSizeY();
+		Sizes[2] = Texture->GetResource()->GetSizeZ();
 	}
 	else if (SparseVolumeTexture)
 	{
-		SizeX = SparseVolumeTexture->GetSizeX();
-		SizeY = SparseVolumeTexture->GetSizeY();
-		SizeZ = SparseVolumeTexture->GetSizeZ();
+		Sizes[0] = SparseVolumeTexture->GetSizeX();
+		Sizes[1] = SparseVolumeTexture->GetSizeY();
+		Sizes[2] = SparseVolumeTexture->GetSizeZ();
 	}
 	else
 	{
-		AsFloat(BufferData, ResultOffset, 0) = 0.0f;
-		AsFloat(BufferData, ResultOffset, 1) = 0.0f;
-		AsFloat(BufferData, ResultOffset, 2) = 0.0f;
+		for (int32 Index = 0; Index < Dimension; Index++)
+		{
+			AsFloat(BufferData, ResultOffset, Index) = 0.0f;
+		}
 		return;
 	}
 
-	AsFloat(BufferData, ResultOffset, 0) = 1.0f / (float)SizeX;
-	AsFloat(BufferData, ResultOffset, 1) = 1.0f / (float)SizeY;
-	AsFloat(BufferData, ResultOffset, 2) = 1.0f / (float)SizeZ;
+	for (int32 Index = 0; Index < Dimension; Index++)
+	{
+		AsFloat(BufferData, ResultOffset, Index) = 1.0f / (float)Sizes[Index];
+	}
 }
 
 void OpEvaluate(const FOpTextureSize&, const FMaterialRenderContext* Context, TArrayView<const FScriptName> Names, TArrayView<float> Buffer, const uint8*& OpcodePtr, EPreshader2Type& Type, int32& Dimension, int32& ResultOffset)
@@ -809,38 +955,74 @@ void OpEvaluate(const FOpTextureSize&, const FMaterialRenderContext* Context, TA
 	float* BufferData = Buffer.GetData();
 
 	ResultOffset = FetchOffset(OpcodePtr);
-	BoundsCheck(Buffer, ResultOffset, 3);		// Writes 3 floats to output
+	Dimension = OpcodePtr[2];
+	Type = (EPreshader2Type)(Dimension >> 6);
+	Dimension = Dimension & 0x7;
+
+	BoundsCheck(Buffer, ResultOffset, Dimension);
 
 	const UTexture* Texture;
 	const USparseVolumeTexture* SparseVolumeTexture;
-	OpcodePtr = GetTextureParameter(Context, Names, OpcodePtr + 2, Texture, SparseVolumeTexture);
+	OpcodePtr = GetTextureParameter(Context, Names, OpcodePtr + 3, Texture, SparseVolumeTexture);
 
-	uint32 SizeX;
-	uint32 SizeY;
-	uint32 SizeZ;
+	uint32 Sizes[3];
 	if (Texture && Texture->GetResource())
 	{
-		SizeX = Texture->GetResource()->GetSizeX();
-		SizeY = Texture->GetResource()->GetSizeY();
-		SizeZ = Texture->GetResource()->GetSizeZ();
+		Sizes[0] = Texture->GetResource()->GetSizeX();
+		Sizes[1] = Texture->GetResource()->GetSizeY();
+		Sizes[2] = Texture->GetResource()->GetSizeZ();
 	}
 	else if (SparseVolumeTexture)
 	{
-		SizeX = SparseVolumeTexture->GetSizeX();
-		SizeY = SparseVolumeTexture->GetSizeY();
-		SizeZ = SparseVolumeTexture->GetSizeZ();
+		Sizes[0] = SparseVolumeTexture->GetSizeX();
+		Sizes[1] = SparseVolumeTexture->GetSizeY();
+		Sizes[2] = SparseVolumeTexture->GetSizeZ();
 	}
 	else
 	{
-		AsFloat(BufferData, ResultOffset, 0) = 0.0f;
-		AsFloat(BufferData, ResultOffset, 1) = 0.0f;
-		AsFloat(BufferData, ResultOffset, 2) = 0.0f;
+		for (int32 Index = 0; Index < Dimension; Index++)
+		{
+			AsFloat(BufferData, ResultOffset, Index) = 0.0f;
+		}
 		return;
 	}
 
-	AsFloat(BufferData, ResultOffset, 0) = (float)SizeX;
-	AsFloat(BufferData, ResultOffset, 1) = (float)SizeY;
-	AsFloat(BufferData, ResultOffset, 2) = (float)SizeZ;
+	for (int32 Index = 0; Index < Dimension; Index++)
+	{
+		AsFloat(BufferData, ResultOffset, Index) = (float)Sizes[Index];
+	}
+}
+
+static void FixupTextureOp(EPreshader2Opcode Opcode, UE::Shader::FPreshaderData& PreshaderOut, TConstArrayView<uint16> Remap, const uint8*& OpcodePtr, EPreshader2Type& Type, int32& Dimension, int32& ResultOffset)
+{
+	// We are emitting code here without explicitly calling the Emit function, so mark this as Preshader2 enabled
+	EnablePreshader2(PreshaderOut);
+
+	ResultOffset = Remap[FetchOffset(OpcodePtr)];
+	uint8 DimensionAndType = OpcodePtr[2];
+	Type = (EPreshader2Type)(DimensionAndType >> 6);
+	Dimension = DimensionAndType & 0x7;
+	OpcodePtr += 3;
+
+	// We can just append the data after the offset, which doesn't need fixup -- simpler than expanding the data to pass to the original emit function.
+	const uint8* PayloadStart = OpcodePtr;
+
+	OpcodePtr = SkipTextureParameter(OpcodePtr);
+
+	PreshaderOut.Data.Add((uint8)Opcode);
+	AppendOffset(PreshaderOut.Data, ResultOffset);
+	PreshaderOut.Data.Add(DimensionAndType);
+	PreshaderOut.Data.Append(PayloadStart, OpcodePtr - PayloadStart);
+}
+
+static void OpFixup(const FOpTexelSize&, UE::Shader::FPreshaderData& PreshaderOut, TConstArrayView<uint16> Remap, const uint8*& OpcodePtr, EPreshader2Type& Type, int32& Dimension, int32& ResultOffset)
+{
+	FixupTextureOp(EPreshader2Opcode::TexelSize, PreshaderOut, Remap, OpcodePtr, Type, Dimension, ResultOffset);
+}
+
+static void OpFixup(const FOpTextureSize&, UE::Shader::FPreshaderData& PreshaderOut, TConstArrayView<uint16> Remap, const uint8*& OpcodePtr, EPreshader2Type& Type, int32& Dimension, int32& ResultOffset)
+{
+	FixupTextureOp(EPreshader2Opcode::TextureSize, PreshaderOut, Remap, OpcodePtr, Type, Dimension, ResultOffset);
 }
 
 // TODO
@@ -848,6 +1030,11 @@ void OpEvaluate(const FOpExternalTextureCoordinateOffset&, const FMaterialRender
 void OpEvaluate(const FOpExternalTextureCoordinateScaleRotation&, const FMaterialRenderContext* Context, TArrayView<const FScriptName> Names, TArrayView<float> Buffer, const uint8*& OpcodePtr, EPreshader2Type& Type, int32& Dimension, int32& ResultOffset) { checkNoEntry(); }
 void OpEvaluate(const FOpSparseVolumeTextureUniform&, const FMaterialRenderContext* Context, TArrayView<const FScriptName> Names, TArrayView<float> Buffer, const uint8*& OpcodePtr, EPreshader2Type& Type, int32& Dimension, int32& ResultOffset) { checkNoEntry(); }
 
+void OpFixup(const FOpExternalTextureCoordinateOffset&, UE::Shader::FPreshaderData& PreshaderOut, TConstArrayView<uint16> Remap, const uint8*& OpcodePtr, EPreshader2Type& Type, int32& Dimension, int32& ResultOffset) { checkNoEntry(); }
+void OpFixup(const FOpExternalTextureCoordinateScaleRotation&, UE::Shader::FPreshaderData& PreshaderOut, TConstArrayView<uint16> Remap, const uint8*& OpcodePtr, EPreshader2Type& Type, int32& Dimension, int32& ResultOffset) { checkNoEntry(); }
+void OpFixup(const FOpSparseVolumeTextureUniform&, UE::Shader::FPreshaderData& PreshaderOut, TConstArrayView<uint16> Remap, const uint8*& OpcodePtr, EPreshader2Type& Type, int32& Dimension, int32& ResultOffset) { checkNoEntry(); }
+
+
 // Unary opcode variants.
 // Reuse:				Opcode(1)  													ResultOffset/Dimension/Type used from previous op, AOffset same as ResultOffset
 // InPlace:				Opcode(1)  ResultOffset(2)  Dimension+Type(1)				AOffset same as ResultOffset
@@ -946,15 +1133,17 @@ void OpEvaluate(const FOpSparseVolumeTextureUniform&, const FMaterialRenderConte
 
 // Opcode(1)  ResultOffset(2)  AOffset(2)  BOffset(2)  Dimension+Type(1)
 #define PRESHADER2_BINARY_OP_CUSTOM(OP_NAME) \
-	case EPreshader2Opcode::OP_NAME:																						\
-		ResultOffset = FetchOffset(OpcodePtr);																				\
-		AOffset = FetchOffset(OpcodePtr + 2);																				\
-		BOffset = FetchOffset(OpcodePtr + 4);																				\
-		Dimension = OpcodePtr[6];																							\
-		Type = (EPreshader2Type)(Dimension >> 6);																			\
-		Dimension = Dimension & 7;																							\
-		OpcodePtr += 7;																										\
-		BinaryEvaluate(FOp ## OP_NAME(), DataContext.ResultsBuffer, Type, Dimension, ResultOffset, AOffset, BOffset);		\
+	case EPreshader2Opcode::OP_NAME:																								\
+		ResultOffset = FetchOffset(OpcodePtr);																						\
+		AOffset = FetchOffset(OpcodePtr + 2);																						\
+		BOffset = FetchOffset(OpcodePtr + 4);																						\
+		Dimension = OpcodePtr[6];																									\
+		Type = (EPreshader2Type)(Dimension >> 6);																					\
+		AInc = (Dimension >> 5) & 1;																								\
+		BInc = (Dimension >> 4) & 1;																								\
+		Dimension = Dimension & 7;																									\
+		OpcodePtr += 7;																												\
+		BinaryEvaluate(FOp ## OP_NAME(), DataContext.ResultsBuffer, Type, Dimension, ResultOffset, AOffset, BOffset, AInc, BInc);	\
 		break;
 
 // Fully custom evaluation
@@ -1021,10 +1210,11 @@ static void AppendOffset(TMemoryImageArray<uint8>& Opcodes, uint16 Offset)
 	Opcodes.Add((uint8)(Offset >> 8));
 }
 
-static void AppendImmediateConstant(TMemoryImageArray<uint8>& Opcodes, EPreshader2Type Type, uint8 Dimension, void* Data)
+static int32 AppendImmediateConstant(TMemoryImageArray<uint8>& Opcodes, EPreshader2Type Type, uint8 Dimension, const void* Data)
 {
 	int32 TypeSize = Type == EPreshader2Type::Double ? 8 : 4;
 	Opcodes.Append((const uint8*)Data, TypeSize * Dimension);
+	return TypeSize * Dimension;
 }
 
 static void EnablePreshader2(UE::Shader::FPreshaderData& PreshaderData)
@@ -1076,7 +1266,7 @@ void EmitUnaryOp(UE::Shader::FPreshaderData& PreshaderData, EPreshader2Opcode Op
 	}
 	else
 	{
-		if (ResultOffset != AOffset)
+		if (ResultOffset != AOffset || PreshaderData.bPreFixup)
 		{
 			PreshaderData.Data.Add((uint8)Op);							// Default variant
 			AppendOffset(PreshaderData.Data, ResultOffset);
@@ -1113,15 +1303,14 @@ void EmitBinaryOp(UE::Shader::FPreshaderData& PreshaderData, EPreshader2Opcode O
 
 	if (Op == EPreshader2Opcode::Dot || Op == EPreshader2Opcode::Cross)
 	{
-		// Dot and Cross don't support scalars for A or B.  Cross only works with dimension 3.
-		check(bAScalar == false && bBScalar == false);
+		// Cross only works with dimension 3.
 		check(Op != EPreshader2Opcode::Cross || Dimension == 3);
 
 		PreshaderData.Data.Add((uint8)Op);
 		AppendOffset(PreshaderData.Data, ResultOffset);
 		AppendOffset(PreshaderData.Data, AOffset);
 		AppendOffset(PreshaderData.Data, BOffset);
-		PreshaderData.Data.Add(Dimension | ((uint8)Type << 6));
+		PreshaderData.Data.Add(Dimension | ((uint8)Type << 6) | (bAScalar ? 0 : (uint8)1 << 5) | (bBScalar ? 0 : (uint8)1 << 4));
 	}
 	else if (bReuse)
 	{
@@ -1133,7 +1322,7 @@ void EmitBinaryOp(UE::Shader::FPreshaderData& PreshaderData, EPreshader2Opcode O
 	}
 	else
 	{
-		if (ResultOffset != AOffset)
+		if (ResultOffset != AOffset || PreshaderData.bPreFixup)
 		{
 			PreshaderData.Data.Add((uint8)Op);										// Default variant
 			AppendOffset(PreshaderData.Data, ResultOffset);
@@ -1148,8 +1337,10 @@ void EmitBinaryOp(UE::Shader::FPreshaderData& PreshaderData, EPreshader2Opcode O
 	}
 }
 
-void EmitBinaryConstOp(UE::Shader::FPreshaderData& PreshaderData, EPreshader2Opcode Op, EPreshader2Type Type, uint8 Dimension, uint16 ResultOffset, uint16 AOffset, void* BData, bool bAScalar, bool bBScalar, bool bReuse)
+int32 EmitBinaryConstOp(UE::Shader::FPreshaderData& PreshaderData, EPreshader2Opcode Op, EPreshader2Type Type, uint8 Dimension, uint16 ResultOffset, uint16 AOffset, const void* BData, bool bAScalar, bool bBScalar, bool bReuse)
 {
+	int32 ResultConstantSize;
+
 	EnablePreshader2(PreshaderData);
 
 	// Make sure this is a Binary opcode, and isn't a variant.  The first opcode in each group of 7 is the default version of the opcode.
@@ -1165,11 +1356,11 @@ void EmitBinaryConstOp(UE::Shader::FPreshaderData& PreshaderData, EPreshader2Opc
 		check(bAScalar == false);
 
 		PreshaderData.Data.Add((uint8)Op + (bBScalar ? 6 : 5));					// Const reuse previous variant, vector or scalar constant
-		AppendImmediateConstant(PreshaderData.Data, Type, bBScalar ? 1 : Dimension, BData);
+		ResultConstantSize = AppendImmediateConstant(PreshaderData.Data, Type, bBScalar ? 1 : Dimension, BData);
 	}
 	else
 	{
-		if (ResultOffset != AOffset)
+		if (ResultOffset != AOffset || PreshaderData.bPreFixup)
 		{
 			PreshaderData.Data.Add((uint8)Op + 3);									// Const default variant
 			AppendOffset(PreshaderData.Data, ResultOffset);
@@ -1183,11 +1374,13 @@ void EmitBinaryConstOp(UE::Shader::FPreshaderData& PreshaderData, EPreshader2Opc
 		// Emit an increment of zero for scalar, one for vector.
 		PreshaderData.Data.Add((uint8)Dimension | ((uint8)Type << 6) | (bAScalar ? 0 : 1u << 5) | (bBScalar ? 0 : 1u << 4));
 
-		AppendImmediateConstant(PreshaderData.Data, Type, bBScalar ? 1 : Dimension, BData);
+		ResultConstantSize = AppendImmediateConstant(PreshaderData.Data, Type, bBScalar ? 1 : Dimension, BData);
 	}
+
+	return ResultConstantSize;
 }
 
-void EmitConstantOp(UE::Shader::FPreshaderData& PreshaderData, EPreshader2Type Type, uint8 Dimension, uint16 ResultOffset, void* Data)
+void EmitConstantOp(UE::Shader::FPreshaderData& PreshaderData, EPreshader2Type Type, uint8 Dimension, uint16 ResultOffset, const void* Data)
 {
 	EnablePreshader2(PreshaderData);
 
@@ -1205,6 +1398,7 @@ void EmitConstantOp(UE::Shader::FPreshaderData& PreshaderData, EPreshader2Type T
 void EmitMoveOp(UE::Shader::FPreshaderData& PreshaderData, EPreshader2Type Type, uint16 ResultOffset, TConstArrayView<uint16> SourceOffsets)
 {
 	EnablePreshader2(PreshaderData);
+	check(!PreshaderData.bPreFixup);
 
 	check(SourceOffsets.Num() <= 4 && SourceOffsets.Num() >= 1);
 	check(Type != EPreshader2Type::LWC);
@@ -1262,13 +1456,35 @@ void EmitMoveOp(UE::Shader::FPreshaderData& PreshaderData, EPreshader2Type Type,
 	}
 }
 
-void EmitTextureOp(UE::Shader::FPreshaderData& PreshaderData, EPreshader2Opcode Op, uint16 ResultOffset, const FHashedMaterialParameterInfo& HashedSourceParameterInfo, int32 TextureIndex, int32 UniformIndex)
+void EmitMoveOpForFixup(UE::Shader::FPreshaderData& PreshaderData, EPreshader2Type Type, uint16 ResultRegister, uint8 ResultComponent, TConstArrayView<uint16> Registers, TConstArrayView<uint8> Components)
+{
+	EnablePreshader2(PreshaderData);
+	check(PreshaderData.bPreFixup);
+
+	check(Registers.Num() <= 4 && Registers.Num() >= 1 && Registers.Num() == Components.Num());
+	check(Type != EPreshader2Type::LWC);
+
+	PreshaderData.Data.Add((uint8)EPreshader2Opcode::MoveForFixup);
+	PreshaderData.Data.Add(((uint8)Type << 6) | Registers.Num());
+
+	AppendOffset(PreshaderData.Data, ResultRegister);
+	PreshaderData.Data.Add(ResultComponent);
+
+	for (int32 Index = 0; Index < Registers.Num(); ++Index)
+	{
+		AppendOffset(PreshaderData.Data, Registers[Index]);
+	}
+	PreshaderData.Data.Append(Components);
+}
+
+void EmitTextureOp(UE::Shader::FPreshaderData& PreshaderData, EPreshader2Opcode Op, EPreshader2Type Type, uint8 Dimension, uint16 ResultOffset, const FHashedMaterialParameterInfo& HashedSourceParameterInfo, int32 TextureIndex, int32 UniformIndex)
 {
 	EnablePreshader2(PreshaderData);
 
 	check(Op >= EPreshader2Opcode::TextureFirst && Op <= EPreshader2Opcode::TextureLast);
 	PreshaderData.Data.Add((uint8)Op);
 	AppendOffset(PreshaderData.Data, ResultOffset);
+	PreshaderData.Data.Add(Dimension | ((uint8)Type << 6));
 
 	// Material parameter info
 	PreshaderData.Write(HashedSourceParameterInfo);
@@ -1283,6 +1499,699 @@ void EmitTextureOp(UE::Shader::FPreshaderData& PreshaderData, EPreshader2Opcode
 }
 
 
+// --------------------------------------------------------------------------------------------------------------------------------------------------------
+//
+//	FixupAndCompress implementation.
+//
+// --------------------------------------------------------------------------------------------------------------------------------------------------------
+
+// FixupAndCompress emits new opcodes with remapped offsets, and attempts to emit compressed opcodes where possible.
+#define PRESHADER2_UNARY_OP(OP_NAME, ALLOW_INTEGER) \
+	case EPreshader2Opcode::OP_NAME:																														\
+		ResultOffset = Remap[FetchOffset(OpcodePtr)];																										\
+		AOffset = Remap[FetchOffset(OpcodePtr + 2)];																										\
+		Dimension = OpcodePtr[4];								/* Dimension value which may include extra payload */										\
+		Type = EPreshader2Type(Dimension >> 6);																												\
+		DestType = EPreshader2Type((Dimension >> 4) & 0x3);		/* DestType is included in payload for Cast operation */									\
+		Dimension = Dimension & 7;								/* Dimension with payload removed */														\
+		OpcodePtr += 5;																																		\
+		bReuse = (ResultOffset == AOffset) && (PreviousResultOffset == ResultOffset) && (PreviousType == Type) && (PreviousDimension == Dimension);			\
+		EmitUnaryOp(PreshaderOut, Opcode, Type, Dimension, ResultOffset, AOffset, bReuse, DestType);														\
+		break;
+
+// FixupAndCompress should be passed input with unique registers, where reuse cannot occur, so we can leave these macros empty.
+#define PRESHADER2_UNARY_OP_CAST_REUSE(OP_NAME)
+#define PRESHADER2_UNARY_OP_CAST_REUSE_LAST(OP_NAME)
+
+// Can reuse PRESHADER2_UNARY_OP
+#define PRESHADER2_UNARY_OP_CUSTOM(OP_NAME) PRESHADER2_UNARY_OP(OP_NAME, false)
+
+#define PRESHADER2_BINARY_OP(OP_NAME, ALLOW_INTEGER) \
+	case EPreshader2Opcode::OP_NAME:																						\
+		ResultOffset = Remap[FetchOffset(OpcodePtr)];																		\
+		AOffset = Remap[FetchOffset(OpcodePtr + 2)];																		\
+		BOffset = FetchOffset(OpcodePtr + 4);			/* Need to do Remap below after removing BInc from value */			\
+		Dimension = OpcodePtr[6];																							\
+		Type = (EPreshader2Type)(Dimension >> 6);																			\
+		AInc = (Dimension >> 5) & 1;																						\
+		BInc = (BOffset >> 15) & 1;																							\
+		BOffset = Remap[BOffset & 0x7fff];																					\
+		Dimension = Dimension & 7;																							\
+		OpcodePtr += 7;																										\
+		bReuse = (ResultOffset == AOffset) && (PreviousResultOffset == ResultOffset) && (PreviousType == Type) && (PreviousDimension == Dimension);		\
+		EmitBinaryOp(PreshaderOut, Opcode, Type, Dimension, ResultOffset, AOffset, BOffset, AInc==0, BInc==0, bReuse);		\
+		break;																												\
+	case EPreshader2Opcode::OP_NAME ## _Const:																				\
+		ResultOffset = Remap[FetchOffset(OpcodePtr)];																		\
+		AOffset = Remap[FetchOffset(OpcodePtr + 2)];																		\
+		Dimension = OpcodePtr[4];																							\
+		Type = (EPreshader2Type)(Dimension >> 6);																			\
+		AInc = (Dimension >> 5) & 1;																						\
+		BInc = (Dimension >> 4) & 1;																						\
+		Dimension = Dimension & 7;																							\
+		OpcodePtr += 5;																										\
+		bReuse = (ResultOffset == AOffset) && (PreviousResultOffset == ResultOffset) && (PreviousType == Type) && (PreviousDimension == Dimension) && AInc;			\
+		OpcodePtr += EmitBinaryConstOp(PreshaderOut, EPreshader2Opcode::OP_NAME, Type, Dimension, ResultOffset, AOffset, OpcodePtr, AInc==0, BInc==0, bReuse);		\
+		break;
+
+#define PRESHADER2_BINARY_OP_CUSTOM(OP_NAME) \
+	case EPreshader2Opcode::OP_NAME:																						\
+		ResultOffset = Remap[FetchOffset(OpcodePtr)];																		\
+		AOffset = Remap[FetchOffset(OpcodePtr + 2)];																		\
+		BOffset = Remap[FetchOffset(OpcodePtr + 4)];																		\
+		Dimension = OpcodePtr[6];																							\
+		Type = (EPreshader2Type)(Dimension >> 6);																			\
+		AInc = (Dimension >> 5) & 1;																						\
+		BInc = (Dimension >> 4) & 1;																						\
+		Dimension = Dimension & 7;																							\
+		OpcodePtr += 7;																										\
+		EmitBinaryOp(PreshaderOut, Opcode, Type, Dimension, ResultOffset, AOffset, BOffset, AInc==0, BInc==0);				\
+		break;
+
+#define PRESHADER2_OP(OP_NAME) \
+	case EPreshader2Opcode::OP_NAME:																						\
+		OpFixup(FOp ## OP_NAME(), PreshaderOut, Remap, OpcodePtr, Type, Dimension, ResultOffset);							\
+		break;
+
+
+void FixupAndCompress(UE::Shader::FPreshaderData& PreshaderOut, const UE::Shader::FPreshaderData& PreshaderIn, TConstArrayView<uint16> Remap)
+{
+	check(&PreshaderOut != &PreshaderIn && PreshaderOut.Data.IsEmpty());
+
+	// If empty (no opcodes generated), bPreshader2 won't have been set.
+	check(PreshaderIn.bPreshader2 || PreshaderIn.Data.IsEmpty());
+
+	const uint8* OpcodePtr = PreshaderIn.Data.GetData();
+	const uint8* OpcodeEnd = OpcodePtr + PreshaderIn.Data.Num();
+
+	// Variables that are preserved between loop iterations, for "Reuse" opcode variants
+	EPreshader2Type PreviousType = EPreshader2Type::Float;
+	int32 PreviousDimension = 0;
+	int32 PreviousResultOffset = 0;
+
+	while (OpcodePtr < OpcodeEnd)
+	{
+		const EPreshader2Opcode Opcode = (EPreshader2Opcode)*OpcodePtr;
+		OpcodePtr++;
+
+		// Temporaries, to reduce clutter in the case statement macros
+		EPreshader2Type Type = EPreshader2Type::Float;
+		int32 Dimension = 0;
+		int32 ResultOffset = 0;
+		int32 AOffset;
+		int32 BOffset;
+		int32 AInc;
+		int32 BInc;
+		bool bReuse;
+		EPreshader2Type DestType;
+
+		switch (Opcode)
+		{
+			#include "Shader/Preshader2Opcodes.inl"
+
+		default:
+			check(!"Invalid in-place or reuse opcode variant passed to UE::Preshader2::FixupAndCompress!");
+			UE_ASSUME(false);
+			break;
+		}
+
+		PreviousType = Type;
+		PreviousDimension = Dimension;
+		PreviousResultOffset = ResultOffset;
+	}
+
+	// Copy other fields.
+	PreshaderOut.bPreshader2 = PreshaderIn.bPreshader2;
+	PreshaderOut.Preshader2TemporarySize = PreshaderIn.Preshader2TemporarySize;
+	PreshaderOut.Names = PreshaderIn.Names;
+}
+
+#undef PRESHADER2_UNARY_OP
+#undef PRESHADER2_BINARY_OP
+#undef PRESHADER2_UNARY_OP_CAST_REUSE
+#undef PRESHADER2_UNARY_OP_CAST_REUSE_LAST
+#undef PRESHADER2_UNARY_OP_CUSTOM
+#undef PRESHADER2_BINARY_OP_CUSTOM
+#undef PRESHADER2_OP
+
+
+
+// --------------------------------------------------------------------------------------------------------------------------------------------------------
+//
+//	Print implementation.
+//
+// --------------------------------------------------------------------------------------------------------------------------------------------------------
+
+struct FPrintContext
+{
+	FPrintContext(TConstArrayView<FPrintParameterInfo> InParameterInfos, TConstArrayView<FScriptName> InNames, bool bInPreFixup, TArray<FString>& OutInstructions)
+		: ParameterInfos(InParameterInfos), Names(InNames), bPreFixup(bInPreFixup), Instructions(OutInstructions)
+	{
+		uint16 MaxParameterOffset = 0;
+		for (const FPrintParameterInfo& ParameterInfo : ParameterInfos)
+		{
+			MaxParameterOffset = FMath::Max(MaxParameterOffset, (uint16)(ParameterInfo.Offset + ParameterInfo.ComponentCount * (ParameterInfo.bDouble ? 2 : 1)));
+		}
+
+		ParameterLookup.SetNumUninitialized(MaxParameterOffset);
+		FMemory::Memset(ParameterLookup.GetData(), 0xff, ParameterLookup.Num() * sizeof(uint16));
+
+		for (int32 ParameterIndex = 0; ParameterIndex < ParameterInfos.Num(); ParameterIndex++)
+		{
+			const FPrintParameterInfo& ParameterInfo = ParameterInfos[ParameterIndex];
+			int32 WordCount = ParameterInfo.ComponentCount * (ParameterInfo.bDouble ? 2 : 1);
+
+			for (int32 WordIndex = 0; WordIndex < WordCount; WordIndex++)
+			{
+				ParameterLookup[ParameterInfo.Offset + WordIndex] = (uint16)ParameterIndex;
+			}
+		}
+	}
+
+	void MarkOffsetWritten(int32 Offset, int32 Dimension, bool bDouble)
+	{
+		if (bDouble)
+		{
+			Dimension *= 2;
+		}
+		for (int32 Index = 0; Index < Dimension; Index++)
+		{
+			if (Offset + Index < ParameterLookup.Num())
+			{
+				ParameterLookup[Offset + Index] = UINT16_MAX;
+			}
+		}
+	}
+
+	// OutComponentStart specifies the start component being fetched from the parameter
+	const FScriptName GetParameterAtOffset(int32 Offset, int32& OutComponentStart)
+	{
+		if (Offset < ParameterLookup.Num() && ParameterLookup[Offset] != UINT16_MAX)
+		{
+			FPrintParameterInfo ParameterInfo = ParameterInfos[ParameterLookup[Offset]];
+			OutComponentStart = ParameterInfo.bDouble ? (Offset - ParameterInfo.Offset) / 2 : Offset - ParameterInfo.Offset;
+			return ParameterInfo.Name;
+		}
+
+		OutComponentStart = 0;
+		return FScriptName();
+	}
+
+	void AppendUniformLocation(EPreshader2Type Type, int32 Dimension, int32 Offset, int32 ComponentStart, FString& InoutResult, FScriptName ParameterName = {})
+	{
+		if (bPreFixup)
+		{
+			// Pre-fixup, offsets are register indices, not offsets in the uniform buffer, so just print them as integers prefixed by underscore.
+			InoutResult.Appendf(TEXT("_%d%s"), Offset, GPreshader2LinearSwizzles[ComponentStart][Dimension]);
+		}
+		else
+		{
+			// Note that the Preshader system supports locations that span float4 buffer elements, such as LWC which uses two float4 vectors.
+			// Also, there's technically no alignment requirement for temporaries not accessed in the final buffer (the Preshader system sees
+			// a flat array of words), and those need to print in some way (although for readability when debugging and simplicty of temporary
+			// register allocation, we may choose not to use element spanning intermediates in practice).  We'll wrap to the next buffer
+			// element and keep printing words that are being accessed in the current operation, for example [0].xyzw[1].xyzw or [0].w[1].xy.
+			// It's not proper HLSL syntax, but it hopefully communicates the idea reasonably.
+			int32 WordCount = Type == EPreshader2Type::Double || Type == EPreshader2Type::LWC ? Dimension * 2 : Dimension;
+			bool bNeedComma = false;
+			while (WordCount > 0)
+			{
+				if (bNeedComma)
+				{
+					InoutResult.Append(TEXT(", "));
+				}
+				bNeedComma = true;
+
+				// Figure out how many words we need to print at the current offset
+				int32 WordsToPrint = FMath::Min(4 - Offset % 4, WordCount);
+				InoutResult.Appendf(TEXT("[%d]%s"), Offset / 4, GPreshader2LinearSwizzles[Offset % 4][WordsToPrint]);
+
+				WordCount -= WordsToPrint;
+				Offset += WordsToPrint;
+			}
+		}
+
+		if (!ParameterName.IsNone())
+		{
+			InoutResult.Append(TEXT(" \""));
+			InoutResult.Append(*ParameterName.ToString());
+			InoutResult.AppendChar(TEXT('\"'));
+		}
+	}
+
+	TConstArrayView<FPrintParameterInfo> ParameterInfos;
+
+	// For each buffer location, holds what parameter was written there for quick lookup.  If something else overwrites
+	// that offset, the values are cleared to UINT16_MAX.
+	TArray<uint16> ParameterLookup;
+
+	TConstArrayView<FScriptName> Names;
+
+	bool bPreFixup;
+
+	// Output list of instructions written to.
+	TArray<FString>& Instructions;
+};
+
+static void UnaryPrint(FPrintContext& Context, EPreshader2Opcode Opcode, EPreshader2Type& Type, int32 Dimension, int32 ResultOffset, int32 AOffset)
+{
+	// By design, the macros that implement unary ops pass the Dimension byte unmodified, to allow clients to store data in it.
+	// Mask the low bits to get the actual dimension.
+	Dimension &= 0x7;
+
+	int32 ComponentStart;
+	FScriptName ParameterName = Context.GetParameterAtOffset(AOffset, ComponentStart);
+	const TCHAR* TypeName = GPreshader2TypeNames[(int32)Type];
+	const TCHAR* DimensionName = GPreshader2DimensionNames[Dimension];
+
+	FString Result;
+	Context.AppendUniformLocation(Type, Opcode == EPreshader2Opcode::Length ? 1 : Dimension, ResultOffset, 0, Result);
+	Result.Appendf(TEXT(" = %s("), GPreshader2OpcodeNames[(int32)Opcode]);
+	Context.AppendUniformLocation(Type, Dimension, AOffset, 0, Result, ParameterName);
+	Result.Appendf(TEXT(")  %s%s"), TypeName, DimensionName);
+
+	Context.Instructions.Add(Result);
+	Context.MarkOffsetWritten(ResultOffset, Dimension, Type == EPreshader2Type::Double);
+}
+
+void BinaryPrint(FPrintContext& Context, EPreshader2Opcode Opcode, EPreshader2Type Type, int32 Dimension, int32 ResultOffset, int32 AOffset, int32 BOffset, int32 AInc, int32 BInc)
+{
+	int32 AComponent;
+	int32 BComponent;
+	FScriptName AParameter = Context.GetParameterAtOffset(AOffset, AComponent);
+	FScriptName BParameter = Context.GetParameterAtOffset(BOffset, BComponent);
+	const TCHAR* TypeName = GPreshader2TypeNames[(int32)Type];
+	const TCHAR* DimensionName = GPreshader2DimensionNames[Dimension];
+
+	FString Result;
+	Context.AppendUniformLocation(Type, Opcode == EPreshader2Opcode::Dot ? 1 : Dimension, ResultOffset, 0, Result);
+	Result.Appendf(TEXT(" = %s("), GPreshader2OpcodeNames[(int32)Opcode]);
+	Context.AppendUniformLocation(Type, AInc ? Dimension : 1, AOffset, 0, Result, AParameter);
+	Result.Append(TEXT(", "));
+	Context.AppendUniformLocation(Type, BInc ? Dimension : 1, BOffset, 0, Result, BParameter);
+	Result.Appendf(TEXT(")  %s%s"), TypeName, DimensionName);
+
+	Context.Instructions.Add(Result);
+	Context.MarkOffsetWritten(ResultOffset, Dimension, Type == EPreshader2Type::Double);
+}
+
+void BinaryPrintConst(FPrintContext& Context, EPreshader2Opcode Opcode, EPreshader2Type Type, int32 Dimension, int32 ResultOffset, int32 AOffset, const uint8*& BData, int32 AInc, int32 BInc)
+{
+	// Copy data to an aligned buffer
+	double BDataCopy[4];
+	int32 CopyNum = BInc ? Dimension : 1;
+	int32 CopySize = CopyNum * (Type == EPreshader2Type::Double ? 8 : 4);
+	FMemory::Memcpy(BDataCopy, BData, CopySize);
+
+	// Increment input pointer by copy size
+	BData += CopySize;
+
+	int32 AComponent;
+	FScriptName AParameter = Context.GetParameterAtOffset(AOffset, AComponent);
+	const TCHAR* TypeName = GPreshader2TypeNames[(int32)Type];
+	const TCHAR* DimensionName = GPreshader2DimensionNames[Dimension];
+	const TCHAR* ADimension = GPreshader2DimensionNames[AInc ? Dimension : 1];
+
+	FString Result;
+	Context.AppendUniformLocation(Type, Dimension, ResultOffset, 0, Result);
+	Result.Appendf(TEXT(" = %s("), GPreshader2OpcodeNames[(int32)Opcode]);
+	Context.AppendUniformLocation(Type, AInc ? Dimension : 1, AOffset, 0, Result, AParameter);
+	Result.Append(TEXT(", { "));
+
+	int32 BNum = BInc ? Dimension : 1;
+	for (int32 Index = 0; Index < BNum; Index++)
+	{
+		if (Index)
+		{
+			Result.Append(TEXT(", "));
+		}
+		switch (Type)
+		{
+		case EPreshader2Type::Float:    Result.Appendf(TEXT("%ff"), AsFloat(BDataCopy, 0, Index));  break;
+		case EPreshader2Type::Integer:  Result.Appendf(TEXT("%d"), AsInt(BDataCopy, 0, Index));  break;
+		case EPreshader2Type::Double:
+		case EPreshader2Type::LWC:		Result.Appendf(TEXT("%lf"), AsDouble(BDataCopy, 0, Index));  break;
+		}
+	}
+	Result.Appendf(TEXT(" })  %s%s"), TypeName, DimensionName);
+
+	Context.Instructions.Add(Result);
+	Context.MarkOffsetWritten(ResultOffset, Dimension, Type == EPreshader2Type::Double);
+}
+
+void OpPrint(const FOpMove&, FPrintContext& Context, EPreshader2Opcode Opcode, const uint8*& OpcodePtr, EPreshader2Type& Type, int32& Dimension, int32& ResultOffset)
+{
+	uint8 CountData = OpcodePtr[0];
+	ResultOffset = FetchOffset(OpcodePtr + 1);
+	OpcodePtr += 3;
+
+	// MoveNum and MoveSizes here are all one less than their actual count
+	uint8 MoveNum = (CountData & 0x3);
+	uint8 MoveSizes[4] =
+	{
+		((CountData >> 2) & 0x3),			// 2 bit first move size
+		((CountData >> 4) & 0x3),			// 2 bit second move size
+		((ResultOffset >> 15) & 0x1),		// 1 bit third move size
+		0									// 0 bit implicit fourth move size
+	};
+	Type = (EPreshader2Type)(CountData >> 6);
+	ResultOffset &= 0x7fff;
+
+	// Compute dimension first, rather than in loop below, so we can print it.
+	Dimension = 0;
+	for (int32 MoveIndex = 0; MoveIndex <= MoveNum; ++MoveIndex)
+	{
+		Dimension += MoveSizes[MoveIndex] + 1;
+	}
+
+	int32 ElementSizeInWords = Type == EPreshader2Type::Double ? 2 : 1;
+
+	const TCHAR* TypeName = GPreshader2TypeNames[(int32)Type];
+	const TCHAR* DimensionName = GPreshader2DimensionNames[MoveNum + 1];
+
+	FString Result;
+	Context.AppendUniformLocation(Type, Dimension, ResultOffset, 0, Result);
+	Result.Append(TEXT(" = Move("));
+
+	bool bFirstElement = true;
+	for (int32 MoveIndex = 0; MoveIndex <= MoveNum; ++MoveIndex, OpcodePtr += 2)
+	{
+		int32 SourceOffset = FetchOffset(OpcodePtr);
+
+		for (int32 MoveWord = 0; MoveWord <= MoveSizes[MoveIndex]; ++MoveWord)
+		{
+			if (!bFirstElement)
+			{
+				Result.Append(TEXT(", "));
+			}
+			bFirstElement = false;
+
+			Context.AppendUniformLocation(Type, ElementSizeInWords, SourceOffset + MoveWord * ElementSizeInWords, 0, Result);
+		}
+	}
+	Result.Appendf(TEXT(")  %s%s"), TypeName, DimensionName);
+
+	Context.Instructions.Add(Result);
+	Context.MarkOffsetWritten(ResultOffset, Dimension, Type == EPreshader2Type::Double);
+}
+
+void OpPrint(const FOpMoveForFixup&, FPrintContext& Context, EPreshader2Opcode Opcode, const uint8*& OpcodePtr, EPreshader2Type& Type, int32& Dimension, int32& ResultOffset)
+{
+	Dimension = OpcodePtr[0];
+	ResultOffset = FetchOffset(OpcodePtr + 1);
+	uint8 ComponentStart = OpcodePtr[3];
+	Type = (EPreshader2Type)(Dimension >> 6);
+	Dimension = Dimension & 0x7;
+	OpcodePtr += 4;
+
+	const TCHAR* TypeName = GPreshader2TypeNames[(int32)Type];
+	const TCHAR* DimensionName = GPreshader2DimensionNames[Dimension];
+
+	TArray<uint16, TFixedAllocator<4>> SourceOffsets;
+	SourceOffsets.SetNumUninitialized(Dimension);
+	for (int32 RegisterIndex = 0; RegisterIndex < Dimension; RegisterIndex++, OpcodePtr += 2)
+	{
+		SourceOffsets[RegisterIndex] = FetchOffset(OpcodePtr);
+	}
+	TConstArrayView<uint8> SourceRegisters(OpcodePtr, Dimension);
+	OpcodePtr += Dimension;
+
+	FString Result;
+	Context.AppendUniformLocation(Type, Dimension, ResultOffset, ComponentStart, Result);
+	Result.Append(TEXT(" = MoveForFixup("));
+
+	for (int32 Index = 0; Index < Dimension; ++Index)
+	{
+		if (Index)
+		{
+			Result.Append(TEXT(", "));
+		}
+		Result.Appendf(TEXT("_%d%s"), (int32)SourceOffsets[Index], GPreshader2LinearSwizzles[SourceRegisters[Index]][1]);
+	}
+	Result.Appendf(TEXT(")  %s%s"), TypeName, DimensionName);
+
+	Context.Instructions.Add(Result);
+	Context.MarkOffsetWritten(ResultOffset, Dimension, Type == EPreshader2Type::Double);
+}
+
+void OpPrint(const FOpConstant&, FPrintContext& Context, EPreshader2Opcode Opcode, const uint8*& OpcodePtr, EPreshader2Type& Type, int32& Dimension, int32& ResultOffset)
+{
+	// Opcode(1), ResultOffset(2), Dimension+Type(1), Data(N)
+	ResultOffset = FetchOffset(OpcodePtr);
+	Dimension = OpcodePtr[2];
+	Type = (EPreshader2Type)(Dimension >> 6);
+	Dimension = Dimension & 0x7;
+
+	// Copy data to an aligned buffer
+	double DataCopy[4];
+	const uint8* DataPtr = OpcodePtr + 3;
+	int32 ElementSize = Type == EPreshader2Type::Double ? 8 : 4;
+	int32 CopySize = Dimension * ElementSize;
+	FMemory::Memcpy(DataCopy, DataPtr, CopySize);
+
+	OpcodePtr = DataPtr + CopySize;
+
+	const TCHAR* TypeName = GPreshader2TypeNames[(int32)Type];
+	const TCHAR* DimensionName = GPreshader2DimensionNames[Dimension];
+
+	FString Result;
+	Context.AppendUniformLocation(Type, Dimension, ResultOffset, 0, Result);
+	Result.Append(TEXT(" = Constant({ "));
+
+	for (int32 Index = 0; Index < Dimension; Index++)
+	{
+		if (Index)
+		{
+			Result.Append(TEXT(", "));
+		}
+		switch (Type)
+		{
+		case EPreshader2Type::Float:    Result.Appendf(TEXT("%ff"), AsFloat(DataCopy, 0, Index));  break;
+		case EPreshader2Type::Integer:  Result.Appendf(TEXT("%d"), AsInt(DataCopy, 0, Index));  break;
+		case EPreshader2Type::Double:
+		case EPreshader2Type::LWC:		Result.Appendf(TEXT("%lf"), AsDouble(DataCopy, 0, Index));  break;
+		}
+	}
+	Result.Appendf(TEXT(" })  %s%s"), TypeName, DimensionName);
+
+	Context.Instructions.Add(Result);
+	Context.MarkOffsetWritten(ResultOffset, Dimension, Type == EPreshader2Type::Double);
+}
+
+void OpPrint(const FOpRuntimeVirtualTextureUniform&, FPrintContext& Context, EPreshader2Opcode Opcode, const uint8*& OpcodePtr, EPreshader2Type& Type, int32& Dimension, int32& ResultOffset)
+{
+	ResultOffset = FetchOffset(OpcodePtr);
+	Dimension = OpcodePtr[2];
+	Type = (EPreshader2Type)(Dimension >> 6);
+	Dimension = Dimension & 0x7;
+
+	// Parameter info.
+	FHashedMaterialParameterInfo ParameterInfo;
+	OpcodePtr = FetchHashedMaterialParameterInfo(Context.Names, OpcodePtr + 3, ParameterInfo);
+
+	// Other data.
+	int32 TextureIndex;
+	int32 VectorIndex;
+	FetchDataValue(OpcodePtr, TextureIndex);
+	FetchDataValue(OpcodePtr, VectorIndex);
+
+	FString Result;
+	Context.AppendUniformLocation(Type, Dimension, ResultOffset, 0, Result);
+	Result.Appendf(TEXT(" = RuntimeVirtualTexture(\"%s\", %d, %d)  %s%d"), *ParameterInfo.GetName().ToString(), TextureIndex, VectorIndex, GPreshader2TypeNames[(int32)Type], (int32)Dimension);
+
+	Context.Instructions.Add(Result);
+	Context.MarkOffsetWritten(ResultOffset, 8, false);
+}
+
+void PrintTextureOp(FPrintContext& Context, EPreshader2Opcode Opcode, const uint8*& OpcodePtr, EPreshader2Type& Type, int32& Dimension, int32& ResultOffset)
+{
+	ResultOffset = FetchOffset(OpcodePtr);
+	Dimension = OpcodePtr[2];
+	Type = (EPreshader2Type)(Dimension >> 6);
+	Dimension = Dimension & 0x7;
+
+	// Parameter info.
+	FHashedMaterialParameterInfo ParameterInfo;
+	OpcodePtr = FetchHashedMaterialParameterInfo(Context.Names, OpcodePtr + 3, ParameterInfo);
+
+	// Other data.
+	int32 TextureIndex;
+	FetchDataValue(OpcodePtr, TextureIndex);
+
+	// Set output type to float
+	Type = EPreshader2Type::Float;
+
+	FString Result;
+	Context.AppendUniformLocation(Type, Dimension, ResultOffset, 0, Result);
+	Result.Appendf(TEXT(" = %s(\"%s\", %d)  %s%d"), GPreshader2OpcodeNames[(int32)Opcode], *ParameterInfo.GetName().ToString(), TextureIndex, GPreshader2TypeNames[(int32)Type], (int32)Dimension);
+
+	Context.Instructions.Add(Result);
+	Context.MarkOffsetWritten(ResultOffset, Dimension, false);
+}
+
+void OpPrint(const FOpTexelSize&, FPrintContext& Context, EPreshader2Opcode Opcode, const uint8*& OpcodePtr, EPreshader2Type& Type, int32& Dimension, int32& ResultOffset)
+{
+	PrintTextureOp(Context, Opcode, OpcodePtr, Type, Dimension, ResultOffset);
+}
+
+void OpPrint(const FOpTextureSize&, FPrintContext& Context, EPreshader2Opcode Opcode, const uint8*& OpcodePtr, EPreshader2Type& Type, int32& Dimension, int32& ResultOffset)
+{
+	PrintTextureOp(Context, Opcode, OpcodePtr, Type, Dimension, ResultOffset);
+}
+
+void OpPrint(const FOpExternalTextureCoordinateScaleRotation&, FPrintContext& Context, EPreshader2Opcode Opcode, const uint8*& OpcodePtr, EPreshader2Type& Type, int32& Dimension, int32& ResultOffset) { checkNoEntry(); }
+void OpPrint(const FOpExternalTextureCoordinateOffset&, FPrintContext& Context, EPreshader2Opcode Opcode, const uint8*& OpcodePtr, EPreshader2Type& Type, int32& Dimension, int32& ResultOffset) { checkNoEntry(); }
+void OpPrint(const FOpSparseVolumeTextureUniform&, FPrintContext& Context, EPreshader2Opcode Opcode, const uint8*& OpcodePtr, EPreshader2Type& Type, int32& Dimension, int32& ResultOffset) { checkNoEntry(); }
+
+
+// Cut-and-paste from Evaluate, with print functions substituted -- see comments on the original macros for details.
+#define PRESHADER2_UNARY_OP(OP_NAME, ALLOW_INTEGER) \
+	case EPreshader2Opcode::OP_NAME:																																	\
+		OpcodePtr += 2;									/* Skip to AOffset */																							\
+	case EPreshader2Opcode::OP_NAME ## _InPlace:		/* Otherwise AOffset pulled from same address as ResultOffset */												\
+		ResultOffset = FetchOffset(ResultSrc);																															\
+		AOffset = FetchOffset(OpcodePtr);																																\
+		Dimension = OpcodePtr[2];						/* Dimension value which may include extra payload */															\
+		Type = (EPreshader2Type)(Dimension >> 6);																														\
+		OpcodePtr += 3;																																					\
+	case EPreshader2Opcode::OP_NAME ## _Reuse:																															\
+		UnaryPrint(Context, Opcode, Type, Dimension, ResultOffset, AOffset);																							\
+		Dimension = Dimension & 7;						/* Dimension that's preserved across loop iterations, with payload removed */									\
+		break;
+
+#define PRESHADER2_UNARY_OP_CAST_REUSE(OP_NAME) \
+	case EPreshader2Opcode::OP_NAME: [[fallthrough]];
+
+#define PRESHADER2_UNARY_OP_CAST_REUSE_LAST(OP_NAME) \
+	case EPreshader2Opcode::OP_NAME:																												\
+		UnaryPrint(Context, Opcode, Type, Dimension | (((uint8)Opcode - (uint8)EPreshader2Opcode::Cast_Reuse) << 4), ResultOffset, AOffset);		\
+		break;
+
+#define PRESHADER2_UNARY_OP_CUSTOM(OP_NAME) \
+	case EPreshader2Opcode::OP_NAME:																			\
+		ResultOffset = FetchOffset(OpcodePtr);																	\
+		AOffset = FetchOffset(OpcodePtr + 2);																	\
+		Dimension = OpcodePtr[4];																				\
+		OpcodePtr += 5;																							\
+		Type = (EPreshader2Type)(Dimension >> 6);																\
+		UnaryPrint(Context, Opcode, Type, Dimension, ResultOffset, AOffset);									\
+		break;
+
+#define PRESHADER2_BINARY_OP(OP_NAME, ALLOW_INTEGER) \
+	case EPreshader2Opcode::OP_NAME ## _Reuse:																				\
+		BOffset = FetchOffset(OpcodePtr);																					\
+		OpcodePtr += 2;									/* Skip over BOffset */												\
+		AInc = 1;																											\
+		goto OP_NAME ## _Evaluate;																							\
+	case EPreshader2Opcode::OP_NAME:																						\
+		OpcodePtr += 2;									/* Skip to AOffset */												\
+	case EPreshader2Opcode::OP_NAME ## _InPlace:		/* Otherwise AOffset pulled from same address as ResultOffset */	\
+		ResultOffset = FetchOffset(ResultSrc);																				\
+		AOffset = FetchOffset(OpcodePtr);																					\
+		BOffset = FetchOffset(OpcodePtr + 2);																				\
+		Dimension = OpcodePtr[4];																							\
+		Type = (EPreshader2Type)(Dimension >> 6);																			\
+		AInc = (Dimension >> 5) & 1;																						\
+		Dimension = Dimension & 7;																							\
+		OpcodePtr += 5;																										\
+	OP_NAME ## _Evaluate:																									\
+		BInc = (BOffset >> 15) & 1;																							\
+		BOffset &= 0x7fff;																									\
+		BinaryPrint(Context, Opcode, Type, Dimension, ResultOffset, AOffset, BOffset, AInc, BInc);							\
+		break;																												\
+	case EPreshader2Opcode::OP_NAME ## _ConstReuse:																			\
+		AInc = 1;																											\
+		BInc = 1;																											\
+		goto OP_NAME ## _ConstEvaluate;																						\
+	case EPreshader2Opcode::OP_NAME ## _ConstReuseScalar:																	\
+		AInc = 1;																											\
+		BInc = 0;																											\
+		goto OP_NAME ## _ConstEvaluate;																						\
+	case EPreshader2Opcode::OP_NAME ## _Const:																				\
+		OpcodePtr += 2;									/* Skip to AOffset */												\
+	case EPreshader2Opcode::OP_NAME ## _ConstInPlace:	/* Otherwise AOffset pulled from same address as ResultOffset */	\
+		ResultOffset = FetchOffset(ResultSrc);																				\
+		AOffset = FetchOffset(OpcodePtr);																					\
+		Dimension = OpcodePtr[2];																							\
+		Type = (EPreshader2Type)(Dimension >> 6);																			\
+		AInc = (Dimension >> 5) & 1;																						\
+		BInc = (Dimension >> 4) & 1;																						\
+		Dimension = Dimension & 7;																							\
+		OpcodePtr += 3;																										\
+	OP_NAME ## _ConstEvaluate:																								\
+		BinaryPrintConst(Context, Opcode, Type, Dimension, ResultOffset, AOffset, OpcodePtr, AInc, BInc);					\
+		break;
+
+#define PRESHADER2_BINARY_OP_CUSTOM(OP_NAME) \
+	case EPreshader2Opcode::OP_NAME:																						\
+		ResultOffset = FetchOffset(OpcodePtr);																				\
+		AOffset = FetchOffset(OpcodePtr + 2);																				\
+		BOffset = FetchOffset(OpcodePtr + 4);																				\
+		Dimension = OpcodePtr[6];																							\
+		Type = (EPreshader2Type)(Dimension >> 6);																			\
+		AInc = (Dimension >> 5) & 1;																						\
+		BInc = (Dimension >> 4) & 1;																						\
+		Dimension = Dimension & 7;																							\
+		OpcodePtr += 7;																										\
+		BinaryPrint(Context, Opcode, Type, Dimension, ResultOffset, AOffset, BOffset, AInc, BInc);							\
+		break;
+
+#define PRESHADER2_OP(OP_NAME) \
+	case EPreshader2Opcode::OP_NAME:																						\
+		OpPrint(FOp ## OP_NAME(), Context, Opcode, OpcodePtr, Type, Dimension, ResultOffset);								\
+		break;
+
+
+void Print(const UE::Shader::FPreshaderData& Preshader, TConstArrayView<FPrintParameterInfo> ParameterMap, TArray<FString>& OutInstructions)
+{
+	FPrintContext Context(ParameterMap, Preshader.Names, Preshader.bPreFixup, OutInstructions);
+
+	// Remaining code in this function verbatim copied from Evaluate (although macro definitions are different).
+	const uint8* OpcodePtr = Preshader.Data.GetData();
+	const uint8* OpcodeEnd = OpcodePtr + Preshader.Data.Num();
+
+	// Variables that are preserved between loop iterations, for "Reuse" opcode variants
+	EPreshader2Type Type = EPreshader2Type::Float;
+	int32 Dimension = 0;
+	int32 ResultOffset = 0;
+
+	while (OpcodePtr < OpcodeEnd)
+	{
+		const EPreshader2Opcode Opcode = (EPreshader2Opcode)*OpcodePtr;
+		OpcodePtr++;
+
+		// When an operation is set to "reuse", the result offset from the previous loop iteration is used as both input and output.
+		// Initialize these to the same, in case we hit that code path.
+		int32 AOffset = ResultOffset;
+
+		// When an operation is set to "in place", the result and first argument will both be fetched from the same address.
+		// Store a copy of the original address before it gets advanced, in case we hit that code path.  These two statements
+		// outside the switch statement avoid the need for extra branches for every switch case where fallthrough is used,
+		// reducing code size.
+		const uint8* ResultSrc = OpcodePtr;
+
+		// Temporaries, to reduce clutter in the case statement macros
+		int32 BOffset;
+		int32 AInc;
+		int32 BInc;
+
+		switch (Opcode)
+		{
+			#include "Shader/Preshader2Opcodes.inl"
+		}
+	}
+}
+
+#undef PRESHADER2_UNARY_OP
+#undef PRESHADER2_BINARY_OP
+#undef PRESHADER2_UNARY_OP_CAST_REUSE
+#undef PRESHADER2_UNARY_OP_CAST_REUSE_LAST
+#undef PRESHADER2_UNARY_OP_CUSTOM
+#undef PRESHADER2_BINARY_OP_CUSTOM
+#undef PRESHADER2_OP
+
+
+
 // --------------------------------------------------------------------------------------------------------------------------------------------------------
 //
 //	Unit tests
diff --git a/Engine/Source/Runtime/Engine/Public/MaterialShared.h b/Engine/Source/Runtime/Engine/Public/MaterialShared.h
--- a/Engine/Source/Runtime/Engine/Public/MaterialShared.h
+++ b/Engine/Source/Runtime/Engine/Public/MaterialShared.h
@@ -83,6 +83,7 @@ struct FMaterialShaders;
 struct FMaterialCachedExpressionData;
 class FShaderCompileJobStatus;
 class FShaderMapLayout;
+namespace MIR { struct FDebugDumpIRState; }
 #if WITH_EDITOR
 class FMaterialKeyGeneratorContext;
 #endif
@@ -469,16 +470,20 @@ public:
 };
 
 /**
-* Represents the evaluation of a material parameter directly to a byte offset in the preshader uniform buffer.
+* Represents the evaluation of a material parameter directly to a word offset in the preshader uniform buffer.
 * A UniformParameterEvaluation avoids the overhead of executing an entire preshader when all it does is 
-* evaluate a parameter with a single instruction.
+* evaluate a parameter with a single instruction, and is also used to efficiently preload parameters.
+* 
+* Note that during material translation, BufferOffset may contain a unique "register" index, which is later
+* fixed up to the buffer offset by a call to FixupNumericParameterEvaluations.
 */
 class FMaterialUniformParameterEvaluation
 {
 	DECLARE_TYPE_LAYOUT(FMaterialUniformParameterEvaluation, NonVirtual);
 public:
 	LAYOUT_FIELD(uint16, ParameterIndex);	/** Index of the parameter for UniformExpressionSet::GetNumericParameter() */
-	LAYOUT_FIELD(uint16, BufferOffset); 	/** Offset in the preshader uniform buffer (in bytes) where result will be written */
+	LAYOUT_FIELD(uint16, BufferOffset); 	/** Offset in words where result will be written. */
+	LAYOUT_FIELD(uint16, ComponentCount); 	/** Number of components to write */
 };
 
 class FMaterialUniformPreshaderHeader
@@ -755,15 +760,22 @@ public:
 	
 	/// Adds a new direct parameter evaluation for parameter with specified index (returned by FindOrAddNumericParameter)
 	/// BufferOffset is the location in the uniform buffer in which to evaluate the parameter (in multiple of float components).
-	void AddNumericParameterEvaluation(uint32 ParameterIndex, uint32 BufferOffset);
-	
-	/// Returns the offset of specified parameter given with AddNumericParameterEvaluation.
-	uint32 GetNumericParameterEvaluationOffset(uint32 ParameterIndex) const;
+	void AddNumericParameterEvaluation(uint32 ParameterIndex, uint32 BufferOffset, uint32 ComponentCount);
 	
+	// Fixes up numeric parameter evaluation indices to offsets in the output buffer, using the given remapping table.
+	// Also, optionally fixes up component counts.
+	void FixupNumericParameterEvaluations(TConstArrayView<uint16> Remap, TConstArrayView<uint8> UniformComponentCount);
+
 	/// Allocates the specified number of float4s in the uniform buffer and returns the index
 	/// of the first newly allocated float4 in the buffer.
 	/// Note: multiply returned value by sizeof(float4) to know the byte offset in the buffer.
 	uint32 AllocateFromUniformBuffer(uint32 NumOfFloat4s);
+
+	// Returns size of uniform preshader buffer in float4s.
+	uint32 GetUniformBufferSize() const { return UniformPreshaderBufferSize; }
+
+	// Returns size of temporary preshader buffer in float4s.  Temporary space should be added to the CPU side uniform buffer.
+	uint32 GetTemporaryBufferSize() const { return UniformPreshaderData.Preshader2TemporarySize; }
 	
 	TConstArrayView<FMaterialCacheTagStack> GetMaterialCacheTagStacks() const { return MaterialCacheTagStacks; }
 	
@@ -822,7 +834,7 @@ protected:
 	friend class FMaterialVirtualTextureStack;
 	friend class FDebugUniformExpressionSet;
 	friend struct FMaterialIRValueAnalyzer;
-
+	friend struct MIR::FDebugDumpIRState;
 	LAYOUT_FIELD(TMemoryImageArray<FMaterialUniformParameterEvaluation>, UniformParameterEvaluations);
 	LAYOUT_FIELD(TMemoryImageArray<FMaterialUniformPreshaderHeader>, UniformPreshaders);
 	LAYOUT_FIELD(TMemoryImageArray<FMaterialNumericParameterInfo>, UniformNumericParameters);
diff --git a/Engine/Source/Runtime/Engine/Public/Materials/MaterialIR.h b/Engine/Source/Runtime/Engine/Public/Materials/MaterialIR.h
--- a/Engine/Source/Runtime/Engine/Public/Materials/MaterialIR.h
+++ b/Engine/Source/Runtime/Engine/Public/Materials/MaterialIR.h
@@ -121,7 +121,7 @@ struct FValue
 	FType Type;
 
 	// Used to discern the concrete C++ type of this value (e.g. Subscript)
-	EValueKind Kind : 8;
+	TEnumAsByte<EValueKind> Kind;
 
 	// Set of fundamental flags true for this value.
 	EValueFlags Flags;
@@ -131,6 +131,13 @@ struct FValue
 	// See `EGraphProperties` for more information.
 	EGraphProperties GraphProperties;
 
+	// Offset where this value will be written in the Preshader buffer plus one, with zero
+	// indicating the value is not a Preshader.  During value analyze, before PreshaderFixup has
+	// been called, stores a unique "register" index.  PreshaderFixup translates these registers
+	// to offsets once locations in the output buffer have all been computed, taking into
+	// account lifetimes and aliasing of the registers.
+	int32 Analysis_PreshaderOffset;
+
 	// Returns whether this value has been analyzed for specified stage.
 	bool IsAnalyzed(EStage State) const;
 
@@ -368,6 +375,12 @@ struct FPrimitiveUniform : TValue<VK_PrimitiveUniform>
 
 	// Index of this uniform in the uniform expression set generated during analysis.
 	int16 Analysis_UniformIndex;
+
+	// Mask of which components of this uniform parameter are actually used.
+	int32 Analysis_ComponentMask;
+
+	// Number of components of this uniform parameter that are actually used.  Asserts if Analysis_ComponentMask hasn't been initialized.
+	int32 Analysis_NumComponents() const;
 };
 
 // It represents a texture uniform value.
@@ -746,6 +759,10 @@ int GetOperatorArity(EOperator Op);
 // Returns the string representation of the given operator.
 const TCHAR* LexToString(EOperator Op);
 
+// Translates the given MIR operator to a Preshader2 opcode, and specifies whether it is commutative.
+// Returns EPreshader2Opcode::Invalid if no corresponding opcode.
+EPreshader2Opcode ToPreshader2Opcode(EOperator Op, bool& bOutIsCommutative);
+
 // A mathematical operator instruction.
 //
 // This instruction identifies a built-in operation on one, two or three argument values.
@@ -1087,9 +1104,6 @@ struct FPreshaderParameter : TInstruction<VK_PreshaderParameter, 1>
 
 	// Payload data that depends on the opcode.
 	FPreshaderParameterPayload Payload;
-
-	// Offset into the preshader buffer of the material resource.
-	uint32 Analysis_PreshaderOffset;
 };
 
 
diff --git a/Engine/Source/Runtime/Engine/Public/Materials/MaterialIRCommon.h b/Engine/Source/Runtime/Engine/Public/Materials/MaterialIRCommon.h
--- a/Engine/Source/Runtime/Engine/Public/Materials/MaterialIRCommon.h
+++ b/Engine/Source/Runtime/Engine/Public/Materials/MaterialIRCommon.h
@@ -114,12 +114,18 @@ static inline const TCHAR* LexToString(EStage Stage)
 
 // A collection of graph properties used by a value.
 //
-// If a graph property is true, it means that either the value itself makes has of that
-// property, or one of its dependencies (direct or indirect) has it. This entails these
-// flags are automatically propagated to all dependant values (values that depend on a given one).
-// As an example, if "ReadsPixelNormal" is true on a specific value it means that either
-// that value itself or some other upstream value that value is dependent on reads
-// the pixel normal.
+// If a graph property is true, it generally means that either the value itself makes has of that
+// property, or one of its dependencies (direct or indirect) has it. This entails these flags are
+// automatically propagated to all dependant values (values that depend on a given one).  As an
+// example, if "ReadsPixelNormal" is true on a specific value it means that either that value
+// itself or some other upstream value that value is dependent on reads the pixel normal.
+// 
+// The properties starting with "Uniform" are exceptions.  Instead of being a union of child subgraph
+// properties, "Uniform" represents an intersection, indicating *all* dependencies have "Uniform" set,
+// as opposed to any.  "PreshaderOut" is set for "Uniform" nodes with an immediate non-Uniform parent,
+// and "PreshaderTemp" is set for nodes in a subtree of a "PreshaderOut" node, representing
+// intermediate values that don't need to be available on the GPU or in the output uniform buffer.
+// And "TrivialTailSwizzle" handles a special case optimization.
 enum class EGraphProperties : uint8
 {
 	None = 0,
@@ -129,6 +135,26 @@ enum class EGraphProperties : uint8
 
 	// Some value reads the pixel normal.
 	ReadsPixelNormal = 1 << 1,
+
+	// Some value in the subtree is a parameter.
+	HasParameter = 1 << 2,
+
+	// This value and all its dependencies have the Uniform property set.
+	Uniform = 1 << 3,
+
+	// This value is used as an output in a Preshader, meaning it's visible to the GPU and needs a uniform buffer location.
+	// Set if Uniform and HasParameter are both set, and the value is used by a non-Uniform.
+	PreshaderOut = 1 << 4,
+
+	// This value is used as a temporary in a Preshader, meaning it is only used by other Preshaders, and HLSL doesn't
+	// need to be generated for it.  An exception is if TrivialTailSwizzle below is also set.  Set for values in a
+	// subtree of a PreshaderOut, exclusive with PreshaderOut, which takes precedence.
+	PreshaderTemp = 1 << 5,
+
+	// This value is a trivial tail swizzle, meaning a subscript or composite of elements from a single source vector
+	// (and optionally any number of constant elements) at the root of a subexpression tree that can be expanded on the
+	// GPU for "free".  Forces the code required to expand the value to be emitted to the GPU, even if also a PreshaderTemp.
+	TrivialTailSwizzle = 1 << 6,
 };
 
 ENUM_CLASS_FLAGS(EGraphProperties);
diff --git a/Engine/Source/Runtime/Engine/Public/Materials/MaterialIRPreshader.cpp b/Engine/Source/Runtime/Engine/Public/Materials/MaterialIRPreshader.cpp
new file mode 100644
--- /dev/null
+++ b/Engine/Source/Runtime/Engine/Public/Materials/MaterialIRPreshader.cpp
@@ -0,0 +1,1137 @@
+// Copyright Epic Games, Inc. All Rights Reserved.
+
+#include "MaterialIRPreshader.h"
+#include "Materials/MaterialIRModule.h"
+
+#if WITH_EDITOR
+
+// Returns whether all components of the given composite are identical, and it could be aliased as a scalar.
+static bool CompositeIsScalar(const MIR::FComposite* Composite)
+{
+	TConstArrayView<MIR::FValue*> Components = Composite->GetComponents();
+	for (int32 ComponentIndex = 1; ComponentIndex < Components.Num(); ComponentIndex++)
+	{
+		if (Components[ComponentIndex] != Components[0])
+		{
+			return false;
+		}
+	}
+	return true;
+}
+
+// Checks if the composite is a subset of another vector, and could be aliased.  Returns index of component
+// where the alias starts in the alias source vector, or INDEX_NONE if not an alias.  For example, if a
+// Composite pointed to subscripts 1 and 2 of another vector (yz components), it would return "1".
+// Sets OutAlias to the underlying vector.
+static int32 CompositeIsAlias(MIR::FComposite* Composite, MIR::FValue*& OutAlias)
+{
+	TConstArrayView<MIR::FValue*> Components = Composite->GetComponents();
+	int32 FirstSubscript = INDEX_NONE;
+	MIR::FValue* FirstArg = nullptr;
+	OutAlias = nullptr;
+
+	for (int32 ComponentIndex = 0; ComponentIndex < Components.Num(); ComponentIndex++)
+	{
+		if (MIR::FSubscript* Subscript = Components[ComponentIndex]->As<MIR::FSubscript>())
+		{
+			if (ComponentIndex == 0)
+			{
+				FirstSubscript = Subscript->Index;
+				FirstArg = Subscript->Arg;
+			}
+			else if (Subscript->Index != FirstSubscript + ComponentIndex || Subscript->Arg != FirstArg)
+			{
+				return INDEX_NONE;
+			}
+		}
+		else
+		{
+			return INDEX_NONE;
+		}
+	}
+
+	OutAlias = FirstArg;
+	return FirstSubscript;
+}
+
+// Copy a scalar constant to output location.
+static void CopyConstant(const MIR::FConstant* Constant, uint32* OutData)
+{
+	if (Constant->Type.AsPrimitive()->ScalarKind == MIR::EScalarKind::Double)
+	{
+		FMemory::Memcpy(OutData, &Constant->Integer, 8);
+	}
+	else
+	{
+		FMemory::Memcpy(OutData, &Constant->Integer, 4);
+	}
+}
+
+// Copy all the components of a composite or scalar value to the output array.
+// Returns true if the constant has all components the same, and has been returned as a scalar.
+static bool UnpackConstantValue(const MIR::FValue* Value, TArray<uint32, TFixedAllocator<8>>& OutData)
+{
+	int32 WordsPerComponent = Value->Type.AsPrimitive()->ScalarKind == MIR::EScalarKind::Double ? 2 : 1;
+
+	if (const MIR::FConstant* Constant = Value->As<MIR::FConstant>())
+	{
+		OutData.SetNumUninitialized(WordsPerComponent);
+		CopyConstant(Constant, OutData.GetData());
+		return true;
+	}
+	else if (const MIR::FComposite* Composite = Value->As<MIR::FComposite>())
+	{
+		TConstArrayView<MIR::FValue*> Components = Composite->GetComponents();
+
+		if (CompositeIsScalar(Composite))
+		{
+			// By design, we leave zero for child components that aren't constants (rather than say an error).
+			// This is so constants can be generated for an output with some elements constant, and some requiring
+			// evaluation, with the evaluation to happen in place after the constant members have been filled in.
+			OutData.SetNumZeroed(WordsPerComponent);
+			if ((Constant = Components[0]->As<MIR::FConstant>()))
+			{
+				CopyConstant(Constant, OutData.GetData());
+			}
+			return true;
+		}
+		else
+		{
+			OutData.SetNumZeroed(Components.Num() * WordsPerComponent);
+			for (int32 ComponentIndex = 0; ComponentIndex < Components.Num(); ComponentIndex++)
+			{
+				if ((Constant = Components[ComponentIndex]->As<MIR::FConstant>()))
+				{
+					CopyConstant(Constant, &OutData[ComponentIndex * WordsPerComponent]);
+				}
+			}
+			return false;
+		}
+	}
+	else
+	{
+		// Return a scalar zero...
+		OutData.SetNumZeroed(WordsPerComponent);
+		return true;
+	}
+}
+
+// Expands a scalar constant to a vector constant with NumComponents.
+static void ExpandScalarConstant(int32 NumComponents, TArray<uint32, TFixedAllocator<8>>& InoutData)
+{
+	if (InoutData.Num() == 1)
+	{
+		InoutData.SetNum(NumComponents);
+		for (int32 ComponentIndex = 1; ComponentIndex < NumComponents; ComponentIndex++)
+		{
+			InoutData[ComponentIndex] = InoutData[0];
+		}
+	}
+	else
+	{
+		// Should be a double type, need to copy pairs of data elements
+		check(InoutData.Num() == 2);
+
+		InoutData.SetNum(NumComponents * 2);
+		for (int32 ComponentIndex = 1; ComponentIndex < NumComponents; ComponentIndex++)
+		{
+			InoutData[ComponentIndex * 2 + 0] = InoutData[0];
+			InoutData[ComponentIndex * 2 + 1] = InoutData[1];
+		}
+	}
+}
+
+void FMaterialIRPreshader::AddValueUse(const MIR::FValue* Source, const MIR::FValue* Dest)
+{
+	// If source isn't a preshader (meaning its a constant), we don't care about uses.  Also ignore nullptr
+	// for Source, to handle passing a null BArg for a unary operator, without needing the caller to check.
+	if (!Source || !Source->Analysis_PreshaderOffset)
+	{
+		return;
+	}
+
+	// Dest must be a preshader.
+	check(Dest->Analysis_PreshaderOffset);
+
+	// If source is an alias, we want the use to be on the underlying value.  There can be a chain of aliases
+	// for the case of a composite that's a scalar alias of a subscript that's also an alias.
+	while (ValueInfos[Source->Analysis_PreshaderOffset].Alias)
+	{
+		Source = ValueInfos[Source->Analysis_PreshaderOffset].Alias;
+		check(Source->Analysis_PreshaderOffset);
+	}
+
+	FValueInfo& SourceInfo = ValueInfos[Source->Analysis_PreshaderOffset];
+	FValueInfo& DestInfo = ValueInfos[Dest->Analysis_PreshaderOffset];
+
+	// AddUnique, so multiple uses of the same source value by a single destination still only count as one use.
+	DestInfo.Uses.AddUnique(Source->Analysis_PreshaderOffset);
+
+	if (!SourceInfo.FirstUsedBy)
+	{
+		SourceInfo.FirstUsedBy = Dest->Analysis_PreshaderOffset;
+	}
+	SourceInfo.LastUsedBy = Dest->Analysis_PreshaderOffset;
+	SourceInfo.UsedByCount++;
+}
+
+// Explicitly emit a primitive composite that's any combination of preshader values and constants as a full preshader value.
+void FMaterialIRPreshader::EmitComposite(MIR::FComposite* Composite)
+{
+	int32 PreshaderIndex = AddValueInfo(Composite);
+
+	MIR::FPrimitive PrimitiveType = *Composite->Type.AsPrimitive();
+	EPreshader2Type Preshader2Type = ToPreshader2Type(PrimitiveType.ScalarKind);
+
+	TConstArrayView<MIR::FValue*> Components = Composite->GetComponents();
+	int32 FirstConstantIndex = INDEX_NONE;
+	int32 LastConstantIndex = INDEX_NONE;
+
+	for (int32 ComponentIndex = 0; ComponentIndex < Components.Num(); ComponentIndex++)
+	{
+		if (Components[ComponentIndex]->IsA(MIR::VK_Constant))
+		{
+			if (FirstConstantIndex == INDEX_NONE)
+			{
+				FirstConstantIndex = ComponentIndex;
+			}
+			LastConstantIndex = ComponentIndex;
+		}
+	}
+
+	if (FirstConstantIndex != INDEX_NONE)
+	{
+		TArray<uint32, TFixedAllocator<8>> ConstantData;
+		bool bIsScalar = UnpackConstantValue(Composite, ConstantData);
+
+		// The standalone constant op doesn't support scalars, unlike the binary constant ops.
+		if (bIsScalar)
+		{
+			ExpandScalarConstant(Components.Num(), ConstantData);
+		}
+
+		if (FirstConstantIndex == 0)
+		{
+			UE::Preshader2::EmitConstantOp(Data, Preshader2Type, LastConstantIndex - FirstConstantIndex + 1, PreshaderIndex, &ConstantData[0]);
+		}
+		else
+		{
+			// We need to create an aliased location to write this constant to, as we are starting at an offset.  This alias
+			// doesn't correspond to any value -- Value is left NULL in FValueInfo.
+			int32 AliasIndex = ValueInfos.AddDefaulted();
+			FMaterialIRPreshader::FValueInfo& AliasInfo = ValueInfos[AliasIndex];
+			AliasInfo.Alias = Composite;
+			AliasInfo.AliasOffset = FirstConstantIndex;
+
+			UE::Preshader2::EmitConstantOp(Data, Preshader2Type, LastConstantIndex - FirstConstantIndex + 1, AliasIndex,
+				&ConstantData[FirstConstantIndex * (Preshader2Type == EPreshader2Type::Double ? 2 : 1)]);
+		}
+	}
+
+	// Now we can handle the moves.  Moves only support contiguous output, so potentially we need to emit multiple moves,
+	// but in practice, where there is a mix of constants and moves, the moves will usually be contiguous (for example,
+	// appending a one constant to XYZ to make a homogenous vector).
+	TArray<uint16, TFixedAllocator<4>> SourceRegisters;
+	TArray<uint8, TFixedAllocator<4>> SourceComponents;
+	int32 ComponentStart = 0;
+
+	for (int32 ComponentIndex = 0; ComponentIndex < Components.Num(); ComponentIndex++)
+	{
+		if (Components[ComponentIndex]->Analysis_PreshaderOffset)
+		{
+			AddValueUse(Components[ComponentIndex], Composite);
+
+			// Track the start of our run of values we need to move
+			if (SourceRegisters.IsEmpty())
+			{
+				ComponentStart = ComponentIndex;
+			}
+
+			FMaterialIRPreshader::FValueInfo& ComponentInfo = ValueInfos[Components[ComponentIndex]->Analysis_PreshaderOffset];
+			if (ComponentInfo.Alias)
+			{
+				SourceRegisters.Add(ComponentInfo.Alias->Analysis_PreshaderOffset);
+				SourceComponents.Add(ComponentInfo.AliasOffset);
+			}
+			else
+			{
+				SourceRegisters.Add(ComponentInfo.Value->Analysis_PreshaderOffset);
+				SourceComponents.Add(0);
+			}
+		}
+
+		// If this is the last element, or the current element is not a preshader (implying its a constant), emit what we have so far.
+		if (ComponentIndex == Components.Num() - 1 || !Components[ComponentIndex]->Analysis_PreshaderOffset)
+		{
+			if (!SourceRegisters.IsEmpty())
+			{
+				UE::Preshader2::EmitMoveOpForFixup(Data, Preshader2Type, Composite->Analysis_PreshaderOffset, ComponentStart, SourceRegisters, SourceComponents);
+				SourceRegisters.Empty();
+				SourceComponents.Empty();
+			}
+		}
+	}
+}
+
+// Emit SourceConstant to Destination.
+MIR::FValue* FMaterialIRPreshader::EmitConstant(const MIR::FValue* SourceConstant, EPreshader2Type Preshader2Type, MIR::FValue* Destination)
+{
+	check(Destination->Analysis_PreshaderOffset && !SourceConstant->Analysis_PreshaderOffset);
+
+	int32 NumComponents = Destination->Type.AsPrimitive()->NumComponents();
+
+	TArray<uint32, TFixedAllocator<8>> SourceData;
+	bool bIsScalar = UnpackConstantValue(SourceConstant, SourceData);
+
+	// Constant op doesn't support scalars, so expand it to the number of output value components
+	if (bIsScalar)
+	{
+		ExpandScalarConstant(NumComponents, SourceData);
+	}
+
+	// Emit the resulting constant data at Destination
+	UE::Preshader2::EmitConstantOp(Data, Preshader2Type, NumComponents, Destination->Analysis_PreshaderOffset, &SourceData[0]);
+
+	return Destination;
+}
+
+// Generates a new dummy constant for SourceConstant, which we can give a buffer location to, as constants
+// in the main MIR tree by design don't support preshader locations.  The only current example is a constant
+// to be passed to the Dot operation, where storing the constant in place in the output isn't possible, since
+// the output of Dot is a smaller dimension than the input.  Since these are allocated via the MemStack
+// context used by translation, they will get automatically freed when that goes out of scope.
+MIR::FValue* FMaterialIRPreshader::EmitDummyConstant(const MIR::FValue* SourceConstant, EPreshader2Type Preshader2Type, int32 NumComponents)
+{
+	MIR::FValue* DummyConstant = (MIR::FValue*)FMemStack::Get().Alloc(sizeof(MIR::FValue), alignof(MIR::FValue));
+	FMemory::Memzero(DummyConstant, sizeof(MIR::FValue));
+
+	DummyConstant->Type = MIR::FType::MakeVector(SourceConstant->Type.AsPrimitive()->ScalarKind, NumComponents);
+
+	int32 PreshaderIndex = AddValueInfo(DummyConstant);
+
+	TArray<uint32, TFixedAllocator<8>> SourceData;
+	bool bIsScalar = UnpackConstantValue(SourceConstant, SourceData);
+
+	// Make sure the source constant is either a scalar (to be expanded below), or that the number of its components matches.
+	// If the number of components doesn't match, it implies a bug, and we don't want to read past the end of SourceData.
+	check(bIsScalar || SourceConstant->Type.AsPrimitive()->NumComponents() == NumComponents);
+
+	// Constant op doesn't support scalars, so expand it to the number of output value components
+	if (bIsScalar)
+	{
+		ExpandScalarConstant(NumComponents, SourceData);
+	}
+
+	// Emit the resulting constant data
+	UE::Preshader2::EmitConstantOp(Data, Preshader2Type, NumComponents, PreshaderIndex, &SourceData[0]);
+
+	return DummyConstant;
+}
+
+void FMaterialIRPreshader::EmitValue(MIR::FValue* Value)
+{
+	// FUniformParameter and FPreshaderParameter will already have their preshader code emitted in FMaterialIRValueAnalyzer::Analyze,
+	// and will already have Analysis_PreshaderOffset set, so we can ignore those here.
+	if (Value->Analysis_PreshaderOffset)
+	{
+		return;
+	}
+
+	if (MIR::FSubscript* Subscript = Value->As<MIR::FSubscript>())
+	{
+		// Subscripts are always aliases of components of the underlying value, and don't generate opcodes.
+		FMaterialIRPreshader::FValueInfo& ValueInfo = AddValueInfo_GetRef(Value);
+		ValueInfo.Alias = Subscript->Arg;
+		ValueInfo.AliasOffset = Subscript->Index;
+	}
+	else if (MIR::FComposite* Composite = Value->As<MIR::FComposite>())
+	{
+		// Trivial swizzles (consecutive subscripts of the same underlying value) can be aliases.
+		MIR::FValue* Alias;
+		int32 AliasOffset = CompositeIsAlias(Composite, Alias);
+		if (AliasOffset != INDEX_NONE)
+		{
+			FMaterialIRPreshader::FValueInfo& ValueInfo = AddValueInfo_GetRef(Value);
+			ValueInfo.Alias = Alias;
+			ValueInfo.AliasOffset = AliasOffset;
+		}
+		else
+		{
+			// Composites that are a copy of the same value can be aliases.
+			if (CompositeIsScalar(Composite))
+			{
+				// In this case, AliasOffset is left as INDEX_NONE, indicating this is a Scalar alias
+				FMaterialIRPreshader::FValueInfo& ValueInfo = AddValueInfo_GetRef(Value);
+				ValueInfo.Alias = Composite->GetComponents()[0];
+			}
+			else
+			{
+				// Non-trivial case, will be instantiated on demand later.  A lot of composites may be dead stripped
+				// if they are trivial tail swizzles.
+			}
+		}
+	}
+	else if (MIR::FOperator* Operator = Value->As<MIR::FOperator>())
+	{
+		// Check if any composite children of this operator need to be expanded.  We defer this so we don't generate opcodes
+		// if it's not actually used (common for trivial tail swizzles that have been dead stripped).
+		if (!Operator->AArg->Analysis_PreshaderOffset && Operator->AArg->HasSubgraphProperties(MIR::EGraphProperties::HasParameter) && Operator->AArg->IsA(MIR::VK_Composite))
+		{
+			EmitComposite(Operator->AArg->As<MIR::FComposite>());
+		}
+		if (Operator->BArg && !Operator->BArg->Analysis_PreshaderOffset && Operator->BArg->HasSubgraphProperties(MIR::EGraphProperties::HasParameter) && Operator->BArg->IsA(MIR::VK_Composite))
+		{
+			EmitComposite(Operator->BArg->As<MIR::FComposite>());
+		}
+
+		// Generate preshader operator
+		int32 PreshaderIndex = AddValueInfo(Value);
+		FMaterialIRPreshader::FValueInfo& ValueInfo = ValueInfos[PreshaderIndex];
+
+		MIR::FPrimitive PrimitiveType = *Operator->Type.AsPrimitive();
+		EPreshader2Type Preshader2Type = ToPreshader2Type(PrimitiveType.ScalarKind);
+
+		AddValueUse(Operator->AArg, Value);
+		AddValueUse(Operator->BArg, Value);
+
+		// These two ops have scalar outputs for vector inputs -- replace PrimitiveType with the first arg's type
+		if (Operator->Op == MIR::UO_Length || Operator->Op == MIR::BO_Dot)
+		{
+			PrimitiveType = *Operator->AArg->Type.AsPrimitive();
+		}
+
+		bool bIsCommutative;
+		EPreshader2Opcode Opcode = ToPreshader2Opcode(Operator->Op, bIsCommutative);
+		check(Opcode != EPreshader2Opcode::Invalid);
+
+		if (MIR::IsUnaryOperator(Operator->Op))
+		{
+			// MIR shouldn't generate unary ops on constants, as those should be optimized out.
+			check(Operator->AArg && Operator->AArg->Analysis_PreshaderOffset);
+
+			if (ValueIsScalarAlias(Operator->AArg))
+			{
+				// Preshader2 unary ops don't support promoting scalars to vectors, so if this is a scalar alias, we need to expand it
+				// to the destination before applying the op to it in place.  For a 3 element vector, this adds 5 bytes (2 offsets
+				// plus an extra opcode), as the in-place second op only takes 1 byte.  We could eventually have MIR optimize these
+				// out at emit time if we wanted to (emit the unary op on the underlying scalar, then expand it as a composite).
+				TArray<uint16, TFixedAllocator<4>> SourceRegisters;
+				TArray<uint8, TFixedAllocator<4>> SourceComponents;
+
+				uint16 SourceRegister = (uint16)Operator->AArg->Analysis_PreshaderOffset;
+				for (int32 ComponentIndex = 0; ComponentIndex < PrimitiveType.NumComponents(); ComponentIndex++)
+				{
+					SourceRegisters.Add(SourceRegister);
+					SourceComponents.Add(0);
+				}
+
+				UE::Preshader2::EmitMoveOpForFixup(Data, Preshader2Type, (uint16)PreshaderIndex, 0, SourceRegisters, SourceComponents);
+				UE::Preshader2::EmitUnaryOp(Data, Opcode, Preshader2Type,
+					(uint8)PrimitiveType.NumComponents(),
+					(uint16)PreshaderIndex,
+					(uint16)PreshaderIndex);
+			}
+			else
+			{
+				UE::Preshader2::EmitUnaryOp(Data, Opcode, Preshader2Type,
+					(uint8)PrimitiveType.NumComponents(),
+					(uint16)PreshaderIndex,
+					(uint16)Operator->AArg->Analysis_PreshaderOffset);
+			}
+		}
+		else if (MIR::IsBinaryOperator(Operator->Op))
+		{
+			MIR::FValue* AArg = Operator->AArg;
+			MIR::FValue* BArg = Operator->BArg;
+
+			// MIR shouldn't generate binary ops between two constants, as those should be optimized out.
+			check(AArg && BArg && (AArg->Analysis_PreshaderOffset || BArg->Analysis_PreshaderOffset));
+
+			// If first argument is a constant (not emitted to the preshader), attempt to swap it to the second argument,
+			// so we can use an inline constant version of the opcode.
+			if (!AArg->Analysis_PreshaderOffset && bIsCommutative)
+			{
+				Swap(AArg, BArg);
+			}
+
+			// If that failed (first arg still doesn't have a preshader offset), we need to emit two ops -- one to generate
+			// the constant at the destination location, and a second for the op.  AArg is updated to the destination where
+			// the constant will be stored.
+			if (!AArg->Analysis_PreshaderOffset)
+			{
+				AArg = EmitConstant(AArg, Preshader2Type, Value);
+			}
+
+			// Cross and Dot are limited binary operators, which don't support immediate constants.  We need to emit a
+			// constant for BArg, similar to AArg above, if that's a constant.  Note that we don't have to worry
+			// about both being constant, as MIR would have optimized that out.
+			if ((Operator->Op == MIR::BO_Cross || Operator->Op == MIR::BO_Dot))
+			{
+				if (!BArg->Analysis_PreshaderOffset)
+				{
+					if (Operator->Op == MIR::BO_Dot)
+					{
+						// Allocate a dummy constant, since the constant can't fit in the destination, which is a scalar.
+						// Because Dot is commutative, we know AArg is not a constant (if it were constant, the arguments
+						// would have been swapped), and so we don't need this special case above where AArg is handled.
+						BArg = EmitDummyConstant(BArg, Preshader2Type, AArg->Type.AsPrimitive()->NumComponents());
+
+						// Allow the temporary memory for the dummy constant to be reclaimed by tracking its usage.
+						AddValueUse(BArg, Value);
+					}
+					else
+					{
+						// Write constant to destination, where it will be overwritten in place (no need to allocate separate space)
+						BArg = EmitConstant(BArg, Preshader2Type, Value);
+					}
+				}
+			}
+
+			if (BArg->Analysis_PreshaderOffset)
+			{
+				UE::Preshader2::EmitBinaryOp(Data, Opcode, Preshader2Type,
+					(uint8)PrimitiveType.NumComponents(),
+					(uint16)PreshaderIndex,
+					(uint16)AArg->Analysis_PreshaderOffset,
+					(uint16)BArg->Analysis_PreshaderOffset,
+					ValueIsScalarAlias(AArg),
+					ValueIsScalarAlias(BArg));
+			}
+			else
+			{
+				TArray<uint32, TFixedAllocator<8>> BData;
+				bool bBIsScalar = UnpackConstantValue(BArg, BData);
+
+				UE::Preshader2::EmitBinaryConstOp(Data, Opcode, Preshader2Type,
+					(uint8)PrimitiveType.NumComponents(),
+					(uint16)PreshaderIndex,
+					(uint16)AArg->Analysis_PreshaderOffset,
+					BData.GetData(),
+					ValueIsScalarAlias(AArg),
+					bBIsScalar);
+			}
+		}
+		else
+		{
+			// TODO:  handle clamp, as min/max?
+			UE_MIR_UNREACHABLE();
+		}
+	}
+	// FUniformParameter and FPreshaderParameter will already have an analysis offset defined.  Everything else shouldn't
+	// be a preshader, and shouldn't reach this code path.
+	else if (!Value->Analysis_PreshaderOffset)
+	{
+		UE_MIR_UNREACHABLE();
+	}
+}
+
+// Class to handle allocation of preshader vector components.  Supports efficient search for the best fit available
+// location of allocations of a given component size, using a table of bit arrays to track free space by available
+// component mask.
+struct FComponentAllocator
+{
+	// Allocate float4 buffer space
+	int32 AllocateFromBuffer(int32 NumOfFloat4s)
+	{
+		int32 Result = AllocatedFloat4s;
+		AllocatedFloat4s += NumOfFloat4s;
+		check(AllocatedFloat4s <= UINT16_MAX);
+		return Result;
+	}
+
+	// Accessors for mask bit arrays.  Bit arrays aren't guaranteed to be sized to include the current maximum
+	// number of elements -- it would be wasteful to grow all 16 mask bit arrays every time an element is added.
+	// So we need to treat out of bounds bit accesses as false on read, and add storage on write as needed.
+	// Sizing should be very cheap since it's a bit array with some inline storage.
+	void SetIsAvailable(uint8 Mask, int32 Index)
+	{
+		// Ignore zero masks, these have nothing available, and we don't need overhead writing to AvailableByMask.
+		if (Mask)
+		{
+			if (Index >= AvailableByMask[Mask].Num())
+			{
+				AvailableByMask[Mask].SetNum(Index + 1, false);
+			}
+			AvailableByMask[Mask][Index] = true;
+		}
+
+		// We still need to track which components are available at the given vector Index.
+		if (Index >= AvailableComponents.Num())
+		{
+			AvailableComponents.SetNumZeroed(Index + 1);
+		}
+		AvailableComponents[Index] = Mask;
+	}
+
+	// Gets the item at the given index if available, setting the fetched item as no longer available.
+	bool GetIfAvailable(uint8 Mask, int32 Index)
+	{
+		if (Index < AvailableByMask[Mask].Num() && AvailableByMask[Mask][Index])
+		{
+			AvailableByMask[Mask][Index] = false;
+			return true;
+		}
+		else
+		{
+			return false;
+		}
+	}
+
+	int32 FirstAvailable(uint8 Mask)
+	{
+		return AvailableByMask[Mask].Find(true);
+	}
+
+	int32 NextAvailable(uint8 Mask, int32 Index)
+	{
+		if (Index + 1 < AvailableByMask[Mask].Num())
+		{
+			return AvailableByMask[Mask].FindFrom(true, Index + 1);
+		}
+		return INDEX_NONE;
+	}
+
+	// Given a table of mask/offset pairs that can fit the desired element size, searches for the best available
+	// fit in AvailableByMask, or allocates one if that fails.
+	int32 FindOrAllocateBestFit(const uint8* MaskFitTable, int32 NumWords)
+	{
+		int32 VectorIndex = INDEX_NONE;
+		for (; MaskFitTable[0]; MaskFitTable += 2)
+		{
+			uint8 Mask = MaskFitTable[0];
+
+			VectorIndex = FirstAvailable(Mask);
+			if (VectorIndex != INDEX_NONE)
+			{
+				// Remove the item we are using from its bitvector
+				AvailableByMask[Mask][VectorIndex] = false;
+				break;
+			}
+		}
+
+		if (VectorIndex == INDEX_NONE)
+		{
+			// Back up to last element in table, which is the full vector mask element.
+			MaskFitTable -= 2;
+
+			VectorIndex = AllocateFromBuffer(1);
+		}
+
+		// Remove from mask elements we are occupying.
+		uint8 Mask = MaskFitTable[0];
+		uint8 Offset = MaskFitTable[1];
+		uint8 MaskToRemove = ((1 << NumWords) - 1) << Offset;
+
+		// Make sure there's not a bug in our lookup table
+		check((Mask & MaskToRemove) == MaskToRemove);
+		Mask &= ~MaskToRemove;
+
+		// Add the remainder of the vector to an available list.
+		SetIsAvailable(Mask, VectorIndex);
+
+		// Add the component offset to the vector's offset, to get the overall global component offset.
+		return VectorIndex * 4 + Offset;
+	}
+
+	// Returns the best available offset into the preshader buffer for a float vector with the specified number of components (1-4).
+	uint32 BestComponentOffset(const MIR::FPrimitive& PrimitiveType)
+	{
+		const uint32 NumComponents = PrimitiveType.NumComponents();
+		check(NumComponents >= 1 && NumComponents <= 4);
+
+		// Following code calculates GlobalComponentOffset.
+		// GlobalComponentOffset is the i-th component in the array of float4s that make the uniform buffer.
+		// For example a GlobalComponentOffset of 13 references PreshaderBuffer[3].y.
+		// Attempts to find a best fit of available components from previous allocations or freed elements,
+		// in order to reduce the number of allocations and thus preshader buffer memory footprint.
+
+		if (PrimitiveType.IsDouble())
+		{
+			// Double values must be 8-byte aligned, so only search for offsets at the correct alignment in the word.
+			static const uint8 MaskDouble1[] = {
+				0b0011,0,  0b1100,2,							// Single element free, filling exact space
+				0b1011,0,  0b0111,0,  0b1101,2,  0b1110,2,		// Single element free, some remaining space
+				0b1111,0,  0									// Full vector4 free, plus terminator
+			};
+			static const uint8 MaskDouble2[] = {
+				0b1111,0,  0				// Full vector4 free, plus terminator
+			};
+
+			if (NumComponents == 1)
+			{
+				return FindOrAllocateBestFit(MaskDouble1, NumComponents*2);
+			}
+			else if (NumComponents == 2)
+			{
+				// For 2 component double, we take up a whole word.
+				return FindOrAllocateBestFit(MaskDouble2, NumComponents*2);
+			}
+			else
+			{
+				// Beyond that, we need to take up more than a float4, which need to be consecutive.
+				// So this needs custom allocation logic.  Scan each available full vector, then check
+				// the vector following it to see if the rest will fit.
+				int32 FullVector;
+				uint8 Mask = 0;
+
+				for (FullVector = FirstAvailable(0b1111); FullVector != INDEX_NONE; FullVector = NextAvailable(0b1111, FullVector))
+				{
+					int32 NextVector = FullVector + 1;
+					if (NextVector == AllocatedFloat4s)
+					{
+						// The current full vector is the last allocated element.  Allocate a new one.
+						AllocateFromBuffer(1);
+						Mask = 0b1111;
+						break;
+					}
+
+					// For 3 elements, we can check if the next vector has its first half available
+					if (NumComponents == 3)
+					{
+						if (GetIfAvailable(0b0011, NextVector))
+						{
+							Mask = 0b0011;
+							break;
+						}
+						if (GetIfAvailable(0b0111, NextVector))
+						{
+							Mask = 0b0111;
+							break;
+						}
+						if (GetIfAvailable(0b1011, NextVector))
+						{
+							Mask = 0b1011;
+							break;
+						}
+					}
+
+					// Check if the full next vector is available
+					if (GetIfAvailable(0b1111, NextVector))
+					{
+						Mask = 0b1111;
+						break;
+					}
+				}
+
+				if (FullVector != INDEX_NONE)
+				{
+					// Mark the full vector we are using as unavailable.  Note that the next vector will already
+					// have been marked as unavailable in the calls to "GetIfAvailable" above.
+					AvailableByMask[0b1111][FullVector] = false;
+				}
+				else
+				{
+					// If we didn't find anything, allocate two fresh vectors
+					FullVector = AllocateFromBuffer(2);
+					Mask = 0b1111;
+				}
+
+				// Mark availability mask of FullVector and the following vector.
+				SetIsAvailable(0, FullVector);
+
+				if (NumComponents == 3)
+				{
+					SetIsAvailable(Mask & ~0b0011, FullVector + 1);
+				}
+				else
+				{
+					SetIsAvailable(0, FullVector + 1);
+				}
+
+				// And finally return the allocated offset!
+				return FullVector * 4;
+			}
+		}
+		else
+		{
+			// Tables storing pairs of a best fit mask, and where to put the relevant element if that mask is chosen.
+			// Lookup tables are simpler than writing complex logic to iterate over possibilities in the ideal order.
+			static const uint8 MaskFloat1[] = {
+				0b0001,0,  0b0010,1,  0b0100,2,  0b1000,3,					// Single element free
+				0b0101,0,  0b1001,0,  0b1010,1,  0b1011,3,  0b1101,0,		// Single element holes
+				0b0011,0,  0b0110,1,  0b0111,0,  0b1100,2,  0b1110,1,		// Miscellaneous
+				0b1111,3,  0												// Full vector4 free, plus terminator
+			};
+			static const uint8 MaskFloat2[] = {
+				0b0011,0,  0b0110,1,  0b1100,2,								// Pair of elements free
+				0b1011,0,  0b1101,2,										// Pair of elements hole
+				0b0111,0,  0b1110,1,										// Miscellaneous
+				0b1111,0,  0												// Full vector4 free, plus terminator
+			};
+			static const uint8 MaskFloat3[] = {
+				0b0111,0,  0b1110,1,										// Three elements free
+				0b1111,0,  0,												// Full vector4 free, plus terminator
+			};
+			static const uint8 MaskFloat4[] = {
+				0b1111,0,  0,												// Full vector4 free, plus terminator
+			};
+
+			static const uint8* MasksFloat[4] = { MaskFloat1, MaskFloat2, MaskFloat3, MaskFloat4 };
+
+			return FindOrAllocateBestFit(MasksFloat[NumComponents - 1], NumComponents);
+		}
+	}
+
+	void FreeComponentOffset(int32 Offset, int32 NumWords)
+	{
+		int32 VectorIndex = Offset / 4;
+		if (NumWords >= 4)
+		{
+			// Freeing at least a whole word -- it should be vector aligned, and we know it's not in an available list.
+			check((Offset & 3) == 0);
+			SetIsAvailable(0b1111, VectorIndex);
+
+			if (NumWords > 4)
+			{
+				// Double type that spans words -- NumWords should be even, and thus 6 or 8
+				check(NumWords == 6 || NumWords == 8);
+				if (NumWords == 6)
+				{
+					// Remove from its existing available list if necessary.
+					uint8 ExistingMask = AvailableComponents[VectorIndex + 1];
+					if (ExistingMask)
+					{
+						AvailableByMask[ExistingMask][VectorIndex + 1] = false;
+					}
+
+					// Mark the new availability
+					check(!(ExistingMask & 0b0011));
+					SetIsAvailable(ExistingMask | 0b0011, VectorIndex + 1);
+				}
+				else
+				{
+					// Another whole word -- we know it's not in an available list.
+					SetIsAvailable(0b1111, VectorIndex + 1);
+				}
+			}
+		}
+		else
+		{
+			// Remove from its existing available list if necessary
+			uint8 ExistingMask = AvailableComponents[VectorIndex];
+			if (ExistingMask)
+			{
+				AvailableByMask[ExistingMask][VectorIndex] = false;
+			}
+
+			// Mark the new availability
+			uint8 FreeingMask = ((1 << NumWords) - 1) << (Offset % 4);
+			check(!(ExistingMask & FreeingMask));
+			SetIsAvailable(ExistingMask | FreeingMask, VectorIndex);
+		}
+	}
+
+	void Reset()
+	{
+		for (auto& AvailableItem : AvailableByMask)
+		{
+			AvailableItem.Empty();
+		}
+		AvailableComponents.Empty();
+		AllocatedFloat4s = 0;
+	}
+
+	// Encodes available space in vectors by mask of which components are currently used.  The goal is to
+	// allow efficient search for vectors with sufficient space to hold a new allocation, with minimal
+	// space and dynamic allocation requirements.
+	TStaticArray<TBitArray<TInlineAllocator<1>>, 16> AvailableByMask;
+
+	// By vector offset, mask of which components are available -- used when freeing space to find which
+	// AvailableByMask element to remove the item from.
+	TArray<uint8> AvailableComponents;
+
+	// Number of float4s allocated so far
+	int32 AllocatedFloat4s = 0;
+};
+
+// Gets the primitive type for a preshader value, taking into account the uniform component counts override table.
+static MIR::FPrimitive GetPrimitive(const FMaterialIRPreshader* Preshader, const TArray<uint8>& UniformComponentCounts, int32 PreshaderIndex)
+{
+	MIR::FPrimitive Primitive = Preshader->ValueInfos[PreshaderIndex].Value->Type.GetPrimitive();
+
+	// Component count may be overridden for uniform parameters based on what's actually used -- generate a
+	if (UniformComponentCounts[PreshaderIndex] && UniformComponentCounts[PreshaderIndex] != Primitive.NumComponents())
+	{
+		return MIR::FType::MakeVector(Primitive.ScalarKind, UniformComponentCounts[PreshaderIndex]).GetPrimitive();
+	}
+	else
+	{
+		return Primitive;
+	}
+}
+
+void FMaterialIRPreshader::PreshaderFixup(UE::Shader::FPreshaderData& PreshaderOut, FMaterialIRModule* Module, FMaterialInsights* Insights)
+{
+	// One is added to Remap offsets during processing, so zero represents uninitialized, then subtracted after processing is finished.
+	TArray<uint16> Remap;
+	Remap.SetNumZeroed(ValueInfos.Num());
+
+	// For values that are uniform parameters, tracks component count overrides based on how many components of the parameter are actually used.
+	TArray<uint8> UniformComponentCounts;
+	UniformComponentCounts.SetNumZeroed(ValueInfos.Num());
+
+	// Generate UniformComponentCounts (based on used components) for all Parameters.  Needed for fixup of component
+	// counts in FUniformExpressionSet::FixupNumericParameterEvaluations.  Zero indicates it's not a uniform component.
+	// We always skip the first item, which is an unused dummy.
+	for (int32 PreshaderIndex = 1; PreshaderIndex < ValueInfos.Num(); PreshaderIndex++)
+	{
+		// Value can be null for constant aliases, which should be ignored here.
+		MIR::FValue* Value = ValueInfos[PreshaderIndex].Value;
+		if (!Value)
+		{
+			continue;
+		}
+
+		if (MIR::FPrimitiveUniform* Parameter = Value->As<MIR::FPrimitiveUniform>())
+		{
+			if (Parameter->Type.IsVector())
+			{
+				UniformComponentCounts[PreshaderIndex] = (uint8)Parameter->Analysis_NumComponents();
+			}
+		}
+	}
+
+	// Identify cases of output location propagation.  See comments in header for broader explanation of this optimization.
+	// Scan the list of values in reverse, and identify cases where one of the inputs to a value can reuse the output
+	// location of another value.
+	for (int32 PreshaderIndex = ValueInfos.Num() - 1; PreshaderIndex >= 1; PreshaderIndex--)
+	{
+		// Value can be null for constant aliases, which should be ignored here.
+		MIR::FValue* Value = ValueInfos[PreshaderIndex].Value;
+		if (!Value)
+		{
+			continue;
+		}
+
+		MIR::FPrimitive ValuePrimitive = GetPrimitive(this, UniformComponentCounts, PreshaderIndex);
+
+		for (uint16 UseIndex : ValueInfos[PreshaderIndex].Uses)
+		{
+			// We can reuse the output location if this is the last use, and types match.
+			if (ValueInfos[UseIndex].LastUsedBy == PreshaderIndex)
+			{
+				if (ValuePrimitive == GetPrimitive(this, UniformComponentCounts, UseIndex))
+				{
+					// OutputLocation can be chained across multiple ops, so check if the destination already has an OutputLocation, and propagate that.
+					uint16 DestOutputLocation = ValueInfos[PreshaderIndex].OutputLocation;
+					ValueInfos[UseIndex].OutputLocation = DestOutputLocation ? DestOutputLocation : PreshaderIndex;
+					break;
+				}
+			}
+		}
+	}
+
+	FComponentAllocator OutAllocator;
+	FComponentAllocator TempAllocator;
+
+	// Allocate all PreshaderOut values first, so we know the starting offset for PreshaderTemp.
+	// We don't have to worry about lifetime tracking for these, since they are permanent.
+	for (int32 PreshaderIndex = 1; PreshaderIndex < ValueInfos.Num(); PreshaderIndex++)
+	{
+		FMaterialIRPreshader::FValueInfo& ValueInfo = ValueInfos[PreshaderIndex];
+
+		// Value can be null for constant aliases, which should be ignored here.
+		MIR::FValue* Value = ValueInfo.Value;
+		if (!Value)
+		{
+			continue;
+		}
+
+		if (Value->HasSubgraphProperties(MIR::EGraphProperties::PreshaderOut) && !ValueInfo.Alias)
+		{
+			MIR::FPrimitive PrimitiveType = GetPrimitive(this, UniformComponentCounts, PreshaderIndex);
+
+			uint32 GlobalComponentOffset = OutAllocator.BestComponentOffset(PrimitiveType);
+			check(GlobalComponentOffset < UINT16_MAX);
+
+			Remap[PreshaderIndex] = (uint16)GlobalComponentOffset + 1;		// Add one!
+		}
+	}
+
+	Module->GetCompilationOutput().UniformExpressionSet.AllocateFromUniformBuffer(OutAllocator.AllocatedFloat4s);
+
+	uint32 UniformBufferWords = OutAllocator.AllocatedFloat4s * 4;
+	check(UniformBufferWords <= UINT16_MAX);
+
+	// Then allocate space for all uniform parameters -- these are loaded in a batch before bytecode execution, so
+	// all need to start out in the temporary buffer, or in an OutputLocation already allocated as a PreshaderOut.
+	// The assumption is that most preshader expressions will fall under that category, but as a potential future
+	// optimization, we could allow space in the output buffer to be reused as input or temporary storage before it
+	// finally gets written to as an output.  It adds complexity to the allocator for only a savings in temporary
+	// space, though.
+	//
+	// In a large project, some categories of preshader in the original system, and their efficiency in Preshader2:
+	//		45.0% parameter, with no other opcodes   -- no bytecode,   param to PreshaderOut
+	//		25.6% parameter, tail swizzle            -- no bytecode,   param to PreshaderOut
+	//		 4.6% parameter, unary op                -- 1 op,          param to PreshaderOut via OutputLocation
+	//		 4.3% parameter, binary op with const    -- 1 op,          param to PreshaderOut via OutputLocation
+	//		 3.2% parameter, 2 unary ops             -- 2 ops,         param to PreshaderOut via OutputLocation
+	//		 0.9% runtime virtual texture            -- 1 op,          op to PreshaderOut
+	//
+	// That's 83.6% of preshaders that trivially don't require temporary space to evaluate, and many more won't.
+	// Probably around 10% do, so you're looking at around 10% savings in temporary buffer space, perhaps around
+	// 100 bytes for a complex shader.  The original preshader system stored temporaries on the stack, as
+	// 4-element doubles, so an old preshader with just a binary op would already use 64 bytes of temporary storage.
+	//
+	for (int32 PreshaderIndex = 1; PreshaderIndex < ValueInfos.Num(); PreshaderIndex++)
+	{
+		FMaterialIRPreshader::FValueInfo& ValueInfo = ValueInfos[PreshaderIndex];
+
+		// If this already has a Remap location, it's a PreshaderOut, and can be skipped, as well as null Values.
+		if (Remap[PreshaderIndex] || !ValueInfo.Value)
+		{
+			continue;
+		}
+
+		if (MIR::FPrimitiveUniform* Parameter = ValueInfo.Value->As<MIR::FPrimitiveUniform>())
+		{
+			if (ValueInfo.OutputLocation && Remap[ValueInfo.OutputLocation])
+			{
+				// Already allocated as a PreshaderOut OutputLocation!
+				Remap[PreshaderIndex] = Remap[ValueInfo.OutputLocation];
+			}
+			else
+			{
+				// Allocate a temporary location.
+				MIR::FPrimitive PrimitiveType = GetPrimitive(this, UniformComponentCounts, PreshaderIndex);
+				int32 GlobalComponentOffset = TempAllocator.BestComponentOffset(PrimitiveType) + UniformBufferWords;
+				check(GlobalComponentOffset < UINT16_MAX);
+
+				Remap[PreshaderIndex] = (uint16)GlobalComponentOffset + 1;		// Add one!
+
+				// Propagate the temporary location to OutputLocation, if present, so other intermediates
+				// using the OutputLocation later can see it.
+				if (ValueInfo.OutputLocation)
+				{
+					Remap[ValueInfo.OutputLocation] = Remap[PreshaderIndex];
+				}
+			}
+		}
+	}
+
+	// Then allocate all remaining intermediate values and assign locations to aliases
+	for (int32 PreshaderIndex = 1; PreshaderIndex < ValueInfos.Num(); PreshaderIndex++)
+	{
+		FMaterialIRPreshader::FValueInfo& ValueInfo = ValueInfos[PreshaderIndex];
+
+		// Mark temporary items as free that have encountered their last use, if allocated from the temporary buffer.
+		// We do this before emitting the current item, so the current item can potentially use that space immediately.
+		// Aliases are never marked used, so we can ignore those.
+		if (!ValueInfo.Alias)
+		{
+			for (uint16 UseIndex : ValueInfos[PreshaderIndex].Uses)
+			{
+				// Ignore uses with an OutputLocation, as they won't become free until the last value in a chain,
+				// which won't have OutputLocation set.
+				if (ValueInfos[UseIndex].OutputLocation)
+				{
+					continue;
+				}
+
+				// Temporary buffer starts at offset "UniformBufferWords".
+				if ((ValueInfos[UseIndex].LastUsedBy == PreshaderIndex) && ((uint16)(Remap[UseIndex] - 1) >= UniformBufferWords))		// Subtract one!
+				{
+					MIR::FPrimitive UsePrimitiveType = GetPrimitive(this, UniformComponentCounts, UseIndex);
+					int32 WordsPerComponent = UsePrimitiveType.IsDouble() ? 2 : 1;
+					int32 NumComponents = UsePrimitiveType.NumComponents();
+
+					TempAllocator.FreeComponentOffset(Remap[UseIndex] - 1 - UniformBufferWords, NumComponents * WordsPerComponent);		// Subtract one!
+				}
+			}
+		}
+
+		if (!Remap[PreshaderIndex])
+		{
+			// Check if there's an output location with a remap already set -- if so, pull remap from that.  The remap
+			// might already exist if it's a PreshaderOut that's already been allocated above, or it's used in a chain,
+			// and the first item using the OutputLocation in the chain already allocated it.
+			if (ValueInfo.OutputLocation && Remap[ValueInfo.OutputLocation])
+			{
+				Remap[PreshaderIndex] = Remap[ValueInfo.OutputLocation];
+			}
+			else
+			{
+				int32 GlobalComponentOffset = 0;
+				if (ValueInfo.Alias)
+				{
+					// Aliases should always appear after their dependency, and derive their offset from that.
+					check(ValueInfo.Alias->Analysis_PreshaderOffset < PreshaderIndex);
+					GlobalComponentOffset = Remap[ValueInfo.Alias->Analysis_PreshaderOffset] - 1 + FMath::Max(ValueInfo.AliasOffset, 0);		// Subtract one!
+				}
+				else
+				{
+					// Allocate a new offset.
+					MIR::FPrimitive PrimitiveType = GetPrimitive(this, UniformComponentCounts, PreshaderIndex);
+					GlobalComponentOffset = TempAllocator.BestComponentOffset(PrimitiveType) + UniformBufferWords;
+				}
+				check(GlobalComponentOffset < UINT16_MAX);
+
+				Remap[PreshaderIndex] = (uint16)GlobalComponentOffset + 1;		// Add one!
+
+				// If there's an OutputLocation corresponding to the value we just allocated, mirror the remapping
+				// to the OutputLocation, so it can be picked up by all values sharing that OutputLocation.
+				if (ValueInfo.OutputLocation)
+				{
+					Remap[ValueInfo.OutputLocation] = Remap[PreshaderIndex];
+				}
+			}
+		}
+
+		// Now that we know the buffer location, we can write parameter Insights.
+		if (ValueInfo.Value)
+		{
+			MIR::FPrimitiveUniform* PrimitiveUniform = ValueInfo.Value->As<MIR::FPrimitiveUniform>();
+			if (PrimitiveUniform && Insights)
+			{
+				// Push information about this parameter allocation to the insights.
+				int32 GlobalComponentOffset = Remap[PreshaderIndex] - 1;		// Subtract one!
+				MIR::FPrimitive PrimitiveType = GetPrimitive(this, UniformComponentCounts, PreshaderIndex);
+
+				FMaterialInsights::FUniformParameterAllocationInsight ParamInsight;
+				ParamInsight.BufferSlotIndex = GlobalComponentOffset / 4;
+				ParamInsight.BufferSlotOffset = GlobalComponentOffset % 4;
+				ParamInsight.ComponentsCount = PrimitiveType.NumComponents();
+				ParamInsight.ParameterName = FName(PrimitiveUniform->Name);
+
+				switch (PrimitiveType.ScalarKind)
+				{
+					case MIR::EScalarKind::Integer:    ParamInsight.ComponentType = FMaterialInsights::FUniformBufferSlotComponentType::CT_Int; break;
+					case MIR::EScalarKind::Float:  ParamInsight.ComponentType = FMaterialInsights::FUniformBufferSlotComponentType::CT_Float; break;
+					case MIR::EScalarKind::Double: ParamInsight.ComponentType = FMaterialInsights::FUniformBufferSlotComponentType::CT_LWC; break;
+					default: UE_MIR_UNREACHABLE();
+				}
+
+				Insights->UniformParameterAllocationInsights.Push(ParamInsight);
+			}
+		}
+	}
+
+	Data.Preshader2TemporarySize = (uint16)TempAllocator.AllocatedFloat4s;
+
+	// Subtract one from all offsets
+	for (int32 PreshaderIndex = 1; PreshaderIndex < ValueInfos.Num(); PreshaderIndex++)
+	{
+		Remap[PreshaderIndex]--;		// Subtract one!
+	}
+
+	// Fix up uniform parameter evaluations.
+	Module->GetCompilationOutput().UniformExpressionSet.FixupNumericParameterEvaluations(Remap, UniformComponentCounts);
+
+	// Fix up preshader bytecode.
+	UE::Preshader2::FixupAndCompress(PreshaderOut, Data, Remap);
+
+	// Fix up Analysis_PreshaderOffset -- needs to happen in separate loop, as the original Analysis_PreshaderOffset
+	// values are used when generating the remapping.
+	for (int32 PreshaderIndex = 1; PreshaderIndex < ValueInfos.Num(); PreshaderIndex++)
+	{
+		MIR::FValue* Value = ValueInfos[PreshaderIndex].Value;
+
+		// One is added to this to guarantee the value is non-zero, and subtracted later when used.
+		if (Value)
+		{
+			Value->Analysis_PreshaderOffset = Remap[Value->Analysis_PreshaderOffset] + 1;
+		}
+	}
+}
+
+#endif  // WITH_EDITOR
\ No newline at end of file
diff --git a/Engine/Source/Runtime/Engine/Public/Materials/MaterialIRPreshader.h b/Engine/Source/Runtime/Engine/Public/Materials/MaterialIRPreshader.h
new file mode 100644
--- /dev/null
+++ b/Engine/Source/Runtime/Engine/Public/Materials/MaterialIRPreshader.h
@@ -0,0 +1,272 @@
+// Copyright Epic Games, Inc. All Rights Reserved.
+#pragma once
+
+#include "Materials/MaterialIR.h"
+#include "Materials/MaterialInsights.h"
+
+#if WITH_EDITOR
+
+// Manages data related to preshaders (expressions evaluated on the CPU, outside the GPU).  Expressions can be evaluated on the CPU if they
+// are only dependent on uniform values, meaning values that don't vary per vertex or pixel.  The types of leaf values that are uniform are
+// parameters and constants.  Constants by themselves don't need to be evaluated as preshaders, since they can just be inlined in the shader
+// code.  So any preshader must include some sort of parameter in its expression tree, either FUniformParameter or FPreshaderParameter.
+// 
+// A leaf value by itself isn't an expression -- for something to become an expression, it must include operators (FOperator) or perhaps a
+// swizzle or merge of some other values (combinations of FComposite and FSubscript nodes).  Not all MIR operators are supported by the
+// runtime Preshader2 bytecode, so the FOperatorSignature struct in MIR includes markup for which are supported, and only those are treated
+// as uniform.  To sum up, there are 5 MIR::FValue kinds supported by preshaders:
+//
+//		FUniformParameter, FPreshaderParameter, FOperator, FComposite, FSubscript
+//
+// In the future FScalar (used by MIR for type casting) could also be supported, but this is considered low priority, as currently shaders
+// are unable to produce integer based uniform expressions at all, and uniform math on double values is exceedingly rare.  We do need to
+// treat operations on double parameters as non-uniform in the meantime.  See MarkUniformIfSupportedPreshaderOperator in MaterialIREmitter.cpp.
+//
+// The MIR value analyzer step tracks subtrees where any child is a parameter (EGraphProperties::HasParameter), and subtrees where all
+// children are uniform (EGraphProperties::Uniform).  Any value with both flags set is a Preshader.  The value analyzer further must
+// differentiate between values that are visible to the GPU (EGraphProperties::PreshaderOut), and temporaries which are only used by other
+// uniform calculations (EGraphProperties::PreshaderTemp).  We want to avoid wasting GPU buffer space for values not actually used on the GPU,
+// while on the CPU, we can allocate a larger buffer with space at the end for temporaries.  It's possible for PreshaderOut values to also be
+// an intermediate value to other calculations, but only the PreshaderOut flag will be set on such values.  To figure out whether a value is
+// visible to the GPU requires checking its parent to see if it's non-uniform, handled by having non-uniform nodes flag their Uniform children
+// with PreshaderOut.  Children of PreshaderOut are then recursively marked as PreshaderTemp, if not already marked as PreshaderOut.
+//
+// There's a special case optimization in the value analyzer where trivial tail swizzles (swizzles of components of a single source value,
+// plus any number of constants) are NOT executed as Preshaders.  Such swizzles and constant loads are "free" on all modern scalar op GPUs,
+// and there's no reason to waste CPU time cobbling them together into a contiguous vector.  For example, expressions like these:
+//
+//			MaterialOutput = Composite( Parameter.y, Parameter.x, Parameter.z );
+//			MaterialOutput = Composite( Value.x, Value.y, 1.0 );
+//
+// Optimizing out trivial tail swizzles involves checking if a composite used by a non-uniform meets this criteria, and if so, marking the
+// underlying Parameter or Value as the starting point of the Preshader, rather than the Composite.  This pattern is actually quite common,
+// with trivial tail swizzles on a Parameter representing 25% of all Preshaders in the old material translator in an instrumented replay
+// run on a large project.  This optimization also handles cases where the old material translator wouldn't generate swizzles, but the new
+// material translator would try to, representing another 45% of Preshaders.  And it will frequently save additional ops on many other
+// Preshaders.  That's why it's worthwhile to add this complexity.
+// 
+// It's necessary to separately mark up these composites as "TrivialTailSwizzle", to inform the GPU backend that it's OK to emit them, even
+// if they are also marked as PreshaderTemp, which otherwise prevents emitting instructions to the GPU.  Put in a different way,
+// "TrivialTailSwizzle" is similar to "PreshaderOut", in that both imply a value accessible on the GPU, just with the former involving the
+// final Composite step being handled by the GPU backend.
+//
+// Now that we've identified Preshaders, we need to translate them to the Preshader2 bytecode format.  We don't know which values are
+// outputs versus temporaries, which affects where the values should go in the buffer (temporaries go at the end), until the entire tree
+// has been parsed.  We also need to handle "register" allocation and lifetime tracking.  "Register" is in quotes, because the Preshader2
+// runtime doesn't technically have registers -- all operations are on buffer offsets (in words).  The MIR HLSL translator creates a unique
+// variable for each intermediate value, but Preshader2 bytecode benefits from reusing registers, both in terms of using less memory for
+// temporaries, and allowing opcodes to be compressed if an output matches the input or the previous instruction's output.
+//
+// To handle buffer allocation and register assignment, initial non-executable Preshader2 bytecode is generated with unique register indices
+// per value, instead of buffer offsets.  These can then be replaced with buffer offsets with a call to UE::Preshader2::FixupAndCompress.
+// The "bPreFixup" flag on the output data indicates whether we are generating such bytecode.  A special case is needed for the Preshader2
+// MoveOp, which handles swizzling or concatenating components of vectors.  The runtime version of the bytecode can directly access individual
+// components of vectors by word offset, but when they are in the register format, we need separate component indices within the registers,
+// so a different MoveOpForFixup opcode is used.
+//
+// Preshader2 supports word based indexing, and most binary ops support scalar arguments.  This means that we can avoid explicit MoveOps
+// in most cases, and directly access a subset of another intermediate result (either an individual component or sequential set of
+// components).  This is handled using an Alias feature, where a given intermediate value can declare itself as a subset of another value.
+// Vector aliases include an AliasOffset specifying the component where the Alias starts, while an AliasOffset of -1 (INDEX_NONE) indicates
+// one component should be used for all, with the number of components coming from the MIR::FValue.  Aliases are automatically set up for all
+// FSubscript, since those have only one component and are thus always representable as an alias.  For Composite, an alias will be set up where
+// possible (sequential subset of components or all components the same), but in non-trivial cases where MoveOps would be required, generation
+// of bytecode is deferred until used by another op, as it will be fairly common that the non-trivial composite won't actually be used, if it
+// represents a TrivialTailSwizzle that gets optimized out.  ExpandComposite handles this case, and also supports mixing in components
+// initialized from constants.
+//
+// Preshader2 includes opcodes with inline constants, so constants are usually rolled into binary ops, although some edge cases can't use
+// an inline constant, and a constant load op is generated.  For example, a subtract (or any non-commutative op) where the first argument
+// is a constant, as the constant is always implicitly the second argument.  In the future, if it would help optimization, we could add
+// support for reversed order opcodes to Preshader2, to allow those cases to use inline constants, although it would increase code size.
+//
+// Aliases are a form of space optimization (avoids generating a separate buffer location).  As another space optimization, temporary
+// buffer space can be freed and reused (see FComponentAllocator class in the cpp file).  A third important space optimization involves
+// output location propagation.  This involves scanning the list of values in reverse, and identifying cases where one of the inputs to
+// a value can reuse the output location of another value.  To illustrate where this is valuable, consider the following real world example,
+// without the optimization:
+//
+//		UniformBuffer[2].xyz = #1 Name = "Vector" Type = "Vector" Mask = .xyz
+//		UniformBuffer[1].xyz = Add_Const(UniformBuffer[2].xyz "Vector", { 0.500000f, 0.100000f, 0.100000f })  float3
+//
+// Here we have a vector added to a constant.  UniformBuffer[1] is an output location visible to the GPU, while UniformBuffer[2] (a raw
+// uniform parameter value) is in a temporary location.  We could make this take less space by initially writing the input value to the
+// same location as the output, completely eliminating the temporary location:
+//
+//		UniformBuffer[1].xyz = #1 Name = "Vector" Type = "Vector" Mask = .xyz
+//		UniformBuffer[1].xyz = Add_ConstInPlace(UniformBuffer[1].xyz "Vector", { 0.500000f, 0.100000f, 0.100000f })  float3
+//
+// Not only does this save space in the temporary buffer, it also results in the more compact "InPlace" opcode being generated, saving
+// space in the opcode data.  For chains of multiple such operations, the even more compact "Reuse" opcode variants can potentially be
+// generated.  Here you can see why it's important to scan the values in reverse, as output location propagation can happen across
+// multiple chained opcodes, and a value can't get the correct final location unless ones after it have already been processed.
+//
+// FUTURE OPTIMIZATIONS:  Output location propagation is only currently supported for whole vectors.  It's possible in theory to implement
+// output location propagation for components cobbled together into a single vector.  For example:
+// 
+//		UniformBuffer[2].x = #1 Name = "ValueA" Type = "Scalar" Mask = .x
+//		UniformBuffer[2].y = #1 Name = "ValueB" Type = "Scalar" Mask = .x
+//		UniformBuffer[1].xy = Move(UniformBuffer[2].x "ValueA", UniformBuffer[2].y "ValueB");
+//
+// Here, a move operation is generated to cobble two vector components together into a float2 accessible on the GPU.  Trivial tail swizzle
+// optimizations don't apply here, because the input values aren't from the same source vector (here, they happen to be adjacent in
+// the temporary buffer, but this is not guaranteed).  Because Preshader2 supports flat access to individual words, it would be possible
+// to just write the words into the correct place in the output vector, cutting out the Move:
+//
+// 		UniformBuffer[1].x = #1 Name = "ValueA" Type = "Scalar" Mask = .x
+//		UniformBuffer[1].y = #1 Name = "ValueB" Type = "Scalar" Mask = .x
+//
+// The complexity of implementing this is that the move op will have already been emitted to the bytecode, and needs to be stripped out,
+// and LastUsedBy on each component needs to be propagated forward to anything that uses the compound value.  It's a little bit
+// difficult to unit test the possible permutations of this, compared to the simpler case of whole vector output location propagation.
+//
+// An even more specific case of output location propagation that could be optimized would be parallel in-place operations.  Consider
+// a case like this:
+//
+//		UniformBuffer[2].xyzw = #1 Name = "Vector" Type = "Vector" Mask = .xyzw
+//		UniformBuffer[1].xyz = Multiply_Const(UniformBuffer[2].xyz "Vector", { 0.500000f, 0.100000f, 0.100000f })  float3
+//		UniformBuffer[1].w = Add_Const(UniformBuffer[2].w "Vector", { 1.00000f })  float
+//
+// In this case, different preshader math is applied more or less in place to a vector and scalar portion of the same float4.  This
+// can happen much more frequently than you might think, because MIR does folding of operations involving constants independently per
+// vector component, even if the original vector was strictly used as a whole in the original graph.  But it's also common for end
+// user math to treat "w" differently in an example like the above by design.  The temporary could be eliminated:
+//
+//		UniformBuffer[1].xyzw = #1 Name = "Vector" Type = "Vector" Mask = .xyzw
+//		UniformBuffer[1].xyz = Multiply_ConstInPlace(UniformBuffer[1].xyz "Vector", { 0.500000f, 0.100000f, 0.100000f })  float3
+//		UniformBuffer[1].w = Add_ConstInPlace(UniformBuffer[1].w "Vector", { 1.00000f })  float
+//
+// To make this work, you would do output propagation like normal, and when you reach a fetch from an alias, you would check if the
+// components being fetched are only used by that subtree (perhaps handled by extending UsedBy tracking to aliases).  And then check
+// if the "components only used in subtree" property was true for all the components of the underlying value.  In that case, you
+// would merge the two values into a single contiguous allocation.  It could be considered the logical opposite of cobbling together
+// OUTPUT components to avoid move as described above -- instead it's cobbling together INPUT components.
+//
+// A final similar optimization is merging of parallel scalar math.  As mentioned above, MIR does constant operation folding
+// independently per component, which can generate preshader bytecode like this:
+// 
+//		UniformBuffer[3].x = Add_Const(UniformBuffer[2].x "Vector", { 0.500000f })  float
+//		UniformBuffer[3].y = Add_Const(UniformBuffer[2].y "Vector", { 0.100000f })  float
+//		UniformBuffer[2].xyz = Move(UniformBuffer[3].x, UniformBuffer[3].y, UniformBuffer[2].z)  float3
+//		UniformBuffer[1].xyz = Multiply_Const(UniformBuffer[2].xyz, { 0.300000f })  float3
+//
+// The original graph included the math:  (Vector + float3(0.5, 0.1, 0.0)) * 0.3
+//
+// MIR decided the add by zero could be optimized out for the Z component, and so split the operation for each individual scalar
+// component, producing a separate constant multiply for the x and y components, even though they are from the same source vector.
+// The simplest way to handle this would probably be to disable per-component constant operation folding in the first place for
+// uniform subtrees, to produce this code, saving several bytes and two ops:
+//
+//		UniformBuffer[1].xyz = Add_ConstInPlace(UniformBuffer[1].xyz "Vector", { 0.500000f, 0.100000f, 0.000000f })  float3
+//		UniformBuffer[1].xyz = Multiply_ConstReuse(UniformBuffer[1].xyz, { 0.300000f })  float3
+//
+// The challenge is that uniform subtrees right now are not identified until the analyze step, so you would need to move that
+// earlier (which may or may not be difficult).  Merging the ops after the fact is also possible, but messier.  Perhaps
+// merely adding a flag to operations that were split would allow the Preshader logic to merge ops easily enough (so it doesn't
+// have to do a deep, and therefore complicated to debug, analysis of where merging subtrees of ops is possible).
+// 
+// Another intermediate option would be to have MIR not split adjacent parallel ops into individual scalars, saving one op:
+//
+//		UniformBuffer[3].xy = Add_Const(UniformBuffer[2].xy "Vector", { 0.500000f, 0.100000f })  float2
+//		UniformBuffer[2].xyz = Move(UniformBuffer[3].x, UniformBuffer[3].y, UniformBuffer[2].z)  float3
+//		UniformBuffer[1].xyz = Multiply_Const(UniformBuffer[2].xyz, { 0.300000f })  float3
+//
+// And the second future optimization mentioned above (parallel in-place ops) could optimize that down to this:
+//
+//		UniformBuffer[1].xy = Add_ConstInPlace(UniformBuffer[1].xy "Vector", { 0.500000f, 0.100000f })  float2
+//		UniformBuffer[1].xyz = Multiply_ConstInPlace(UniformBuffer[1].xyz, { 0.300000f })  float3
+//
+// Achieving the minimal number of ops for this particular case in a different way.  Although this approach breaks down in
+// cases where the components modified are not adjacent.
+//
+struct FMaterialIRPreshader
+{
+	FMaterialIRPreshader()
+	{
+		// Add an invalid dummy element to the array.
+		ValueInfos.AddDefaulted();
+
+		// Mark the preshader data as pre-fixup.
+		Data.bPreFixup = true;
+	}
+
+	// Stores data for values in the MIR tree that are used for preshader evaluation.
+	struct FValueInfo
+	{
+		MIR::FValue* Value = nullptr;
+		MIR::FValue* Alias = nullptr;				// Value is an alias of this other value.
+		int16 AliasOffset = INDEX_NONE;				// Offset of Value in its alias.  If INDEX_NONE, it's an alias of a scalar.
+		uint16 FirstUsedBy = 0;						// First value this is used by
+		uint16 LastUsedBy = 0;						// Last value this is used by
+		uint16 UsedByCount = 0;						// Number of items this value is used by
+		TArray<uint16, TFixedAllocator<4>> Uses;	// List of FValueInfo this uses
+		uint16 OutputLocation = 0;					// Value pulls its output location from this other value.
+	};
+
+	// Stores list of preshader values encountered.  "Analysis_PreshaderOffset" in the FValue indexes into this array. Note that the first item is an
+	// unused null dummy value, so valid Analysis_PreshaderOffset indices are non-zero, and zero initialized FValue data represents no preshader.
+	TArray<FValueInfo> ValueInfos;
+
+	// Data here is pre-fixup, with offsets representing indices into the ValueInfo array above.  Later UE::Preshader2::FixupAndCompress remaps the
+	// offsets when generating the final output preshader data.
+	UE::Shader::FPreshaderData Data;
+
+	int32 OutputCount = 0;
+
+	static EPreshader2Type ToPreshader2Type(MIR::EScalarKind Kind)
+	{
+		switch (Kind)
+		{
+		case MIR::EScalarKind::Float:	return EPreshader2Type::Float;
+		case MIR::EScalarKind::Double:	return EPreshader2Type::Double;
+		case MIR::EScalarKind::Integer:	return EPreshader2Type::Integer;
+		}
+		return EPreshader2Type::Float;
+	}
+
+	int32 AddValueInfo(MIR::FValue* Value)
+	{
+		check(!Value->Analysis_PreshaderOffset);
+
+		int32 PreshaderIndex = ValueInfos.AddDefaulted();
+		Value->Analysis_PreshaderOffset = PreshaderIndex;
+
+		FMaterialIRPreshader::FValueInfo& ValueInfo = ValueInfos[PreshaderIndex];
+		ValueInfo.Value = Value;
+
+		return PreshaderIndex;
+	}
+
+	FMaterialIRPreshader::FValueInfo& AddValueInfo_GetRef(MIR::FValue* Value)
+	{
+		return ValueInfos[AddValueInfo(Value)];
+	}
+
+	bool ValueIsScalarAlias(MIR::FValue* Value) const
+	{
+		check(Value->Analysis_PreshaderOffset);
+		const FValueInfo& ValueInfo = ValueInfos[Value->Analysis_PreshaderOffset];
+		return ValueInfo.Alias && ValueInfo.AliasOffset == INDEX_NONE;
+	}
+
+	// Track that "Dest" uses "Source".
+	void AddValueUse(const MIR::FValue* Source, const MIR::FValue* Dest);
+
+	// Explicitly emit a primitive composite that's any combination of preshader values and constants as a full preshader value.
+	void EmitComposite(MIR::FComposite* Composite);
+
+	// Emits a constant to a different destination location, useful to follow up with an in-place op on the destination.  Returns Destination.
+	MIR::FValue* EmitConstant(const MIR::FValue* SourceConstant, EPreshader2Type Preshader2Type, MIR::FValue* Destination);
+
+	// Emits a constant to a dummy value (not part of the MIR tree).  Used for rare case of constants that can't be inline or in-place.
+	MIR::FValue* EmitDummyConstant(const MIR::FValue* SourceConstant, EPreshader2Type Preshader2Type, int32 NumComponents);
+
+	// Emit a value as a preshader.  Non-trivial composites are not emitted here (as they can often be dead stripped), and instead
+	// are handled by EmitComposite as needed.  Also doesn't support emitting constants -- use EmitConstant instead.
+	void EmitValue(MIR::FValue* Value);
+
+	// Final fixup after entire value tree has been processed.
+	void PreshaderFixup(UE::Shader::FPreshaderData& PreshaderOut, FMaterialIRModule* Module, FMaterialInsights* Insights);
+};
+
+#endif  // WITH_EDITOR
\ No newline at end of file
diff --git a/Engine/Source/Runtime/Engine/Public/Shader/Preshader.h b/Engine/Source/Runtime/Engine/Public/Shader/Preshader.h
--- a/Engine/Source/Runtime/Engine/Public/Shader/Preshader.h
+++ b/Engine/Source/Runtime/Engine/Public/Shader/Preshader.h
@@ -96,7 +96,8 @@ class FPreshaderData
 public:
 	friend inline bool operator==(const FPreshaderData& Lhs, const FPreshaderData& Rhs)
 	{
-		return Lhs.bPreshader2 == Rhs.bPreshader2 && Lhs.Names == Rhs.Names && Lhs.Data == Rhs.Data;
+		return Lhs.bPreshader2 == Rhs.bPreshader2 && Lhs.bPreFixup == Rhs.bPreFixup && Lhs.Preshader2TemporarySize == Rhs.Preshader2TemporarySize &&
+			Lhs.Names == Rhs.Names && Lhs.Data == Rhs.Data;
 	}
 
 	friend inline bool operator!=(const FPreshaderData& Lhs, const FPreshaderData& Rhs)
@@ -138,6 +139,8 @@ public:
 	inline FPreshaderData& WriteOpcode(EPreshaderOpcode Op) { ensure(Op != EPreshaderOpcode::Nop && bPreshader2 == false); return Write<uint8>((uint8)Op); }
 
 	LAYOUT_FIELD_INITIALIZED(bool, bPreshader2, false);
+	LAYOUT_FIELD_INITIALIZED(bool, bPreFixup, false);
+	LAYOUT_FIELD_INITIALIZED(uint16, Preshader2TemporarySize, 0);		// Number of extra temporary float4s required for Preshader2 evaluation
 	LAYOUT_FIELD(TMemoryImageArray<FScriptName>, Names);
 	LAYOUT_FIELD(TMemoryImageArray<uint8>, Data);
 };
diff --git a/Engine/Source/Runtime/Engine/Public/Shader/Preshader2.h b/Engine/Source/Runtime/Engine/Public/Shader/Preshader2.h
--- a/Engine/Source/Runtime/Engine/Public/Shader/Preshader2.h
+++ b/Engine/Source/Runtime/Engine/Public/Shader/Preshader2.h
@@ -8,6 +8,7 @@
 
 #include "Materials/MaterialParameters.h"
 #include "Shader/Preshader.h"
+#include "UObject/NameTypes.h"
 
 namespace UE::Shader
 {
@@ -76,6 +77,14 @@ enum class EPreshader2Type : uint8
 namespace UE::Preshader2
 {
 
+struct FPrintParameterInfo
+{
+	FScriptName Name;
+	uint16 Offset;
+	uint16 ComponentCount;
+	bool bDouble;
+};
+
 // Preshader2 uses buffer offsets specified in 4-byte words.  Operations on Double types must use even aligned word offsets, so the 8-byte type ends up aligned.
 
 // "DestType" only applies to the Cast op.  Reuse option uses the result offset of the previous instruction for both ResultOffset and AOffset, but some opcodes
@@ -86,21 +95,37 @@ void EmitUnaryOp(UE::Shader::FPreshaderData& PreshaderData, EPreshader2Opcode Op
 void EmitBinaryOp(UE::Shader::FPreshaderData& PreshaderData, EPreshader2Opcode Op, EPreshader2Type Type, uint8 Dimension, uint16 ResultOffset, uint16 AOffset, uint16 BOffset, bool bAScalar, bool bBScalar, bool bReuse = false);
 
 // Similar to above, but second argument is a constant inlined into the bytecode.  Note that "custom" binary opcodes (ones in the numeric range BinaryCustom to BinaryLast)
-// don't work with this function.  Currently this includes Cross and Dot opcodes.
-void EmitBinaryConstOp(UE::Shader::FPreshaderData& PreshaderData, EPreshader2Opcode Op, EPreshader2Type Type, uint8 Dimension, uint16 ResultOffset, uint16 AOffset, void* BData, bool bAScalar, bool bBScalar, bool bReuse = false);
+// don't work with this function.  Currently this includes Cross and Dot opcodes.  Returns size of the constant data that was emitted.
+int32 EmitBinaryConstOp(UE::Shader::FPreshaderData& PreshaderData, EPreshader2Opcode Op, EPreshader2Type Type, uint8 Dimension, uint16 ResultOffset, uint16 AOffset, const void* BData, bool bAScalar, bool bBScalar, bool bReuse = false);
 
 // Loads an inlined constant to the given ResultOffset.
-void EmitConstantOp(UE::Shader::FPreshaderData& PreshaderData, EPreshader2Type Type, uint8 Dimension, uint16 ResultOffset, void* Data);
+void EmitConstantOp(UE::Shader::FPreshaderData& PreshaderData, EPreshader2Type Type, uint8 Dimension, uint16 ResultOffset, const void* Data);
 
-// SourceOffsets contains the offsets of individual components to be moved.  If any are adjacent, the moves can be combined.
+// SourceOffsets contains the offsets of individual components to be moved.  If any are adjacent, the moves can be combined.  EmitMoveOpForFixup
+// below must be used instead if PreshaderData.bPreFixup is true.
 void EmitMoveOp(UE::Shader::FPreshaderData& PreshaderData, EPreshader2Type Type, uint16 ResultOffset, TConstArrayView<uint16> SourceOffsets);
 
+// When using FixupAndCompress (where unique "registers" are defined for intermediates, and later remapped to offsets), a different version of the move op must be
+// emitted, as register indices by themselves don't allow for selection of which components in those registers are being moved.  A separate array of component
+// indices is provided, and it is stored in a less compressed format than the normal move op.  Requires PreshaderData.bPreFixup to be true.
+void EmitMoveOpForFixup(UE::Shader::FPreshaderData& PreshaderData, EPreshader2Type Type, uint16 ResultRegister, uint8 ResultComponent, TConstArrayView<uint16> Registers, TConstArrayView<uint8> Components);
+
 // Emit texture related opcode (TexelSize, TextureSize, etc).
-void EmitTextureOp(UE::Shader::FPreshaderData& PreshaderData, EPreshader2Opcode Op, uint16 ResultOffset, const FHashedMaterialParameterInfo& HashedSourceParameterInfo, int32 TextureIndex, int32 UniformIndex = INDEX_NONE);
+void EmitTextureOp(UE::Shader::FPreshaderData& PreshaderData, EPreshader2Opcode Op, EPreshader2Type Type, uint8 Dimension, uint16 ResultOffset, const FHashedMaterialParameterInfo& HashedSourceParameterInfo, int32 TextureIndex, int32 UniformIndex = INDEX_NONE);
+
+// Fix up offsets and compress opcodes.  Useful for initially emitting opcodes with unique "registers" per intermediate, then once register lifetimes and usages
+// are known (after analyzing the whole tree), go back and remap registers and compress opcodes where in-place or reused destination execution is possible.  When
+// using this function, in-place and reuse variations of opcodes are not supported!  These won't be possible if each intermediate is unique.  PreshaderOut should
+// be a different empty buffer than PreshaderIn.
+void FixupAndCompress(UE::Shader::FPreshaderData& PreshaderOut, const UE::Shader::FPreshaderData& PreshaderIn, TConstArrayView<uint16> Remap);
 
 // Evaluate opcodes.  FMaterialRenderContext is required if any texture ops are used (TexelSize, TextureSize, etc).
 void Evaluate(UE::Shader::FPreshaderDataContext& DataContext, const FMaterialRenderContext* MaterialContext = nullptr);
 
+// Each Preshader instruction is printed on its own line to the output array.  ParameterMap is an optional list of parameter offsets and sizes to improve readability.
+// 
+void Print(const UE::Shader::FPreshaderData& Preshader, TConstArrayView<FPrintParameterInfo> ParameterInfos, TArray<FString>& OutInstructions);
+
 // Unit test of system.
 void UnitTest();
 
diff --git a/Engine/Source/Runtime/Engine/Public/Shader/Preshader2Opcodes.inl b/Engine/Source/Runtime/Engine/Public/Shader/Preshader2Opcodes.inl
--- a/Engine/Source/Runtime/Engine/Public/Shader/Preshader2Opcodes.inl
+++ b/Engine/Source/Runtime/Engine/Public/Shader/Preshader2Opcodes.inl
@@ -53,8 +53,9 @@ PRESHADER2_BINARY_OP_CUSTOM(Cross)
 PRESHADER2_BINARY_OP_CUSTOM(Dot)
 
 // Ops that don't fall under unary or binary categories
-PRESHADER2_OP(Move)			// Opcode(1), MoveCount+Type+FirstMoveSize+SecondMoveSize(1), ResultOffset+ThirdMoveSize(2), SourceOffset(2*N)
-PRESHADER2_OP(Constant)		// Opcode(1), ResultOffset(2), Dimension+Type(1), Data(N)
+PRESHADER2_OP(Move)				// Opcode(1), MoveCount+Type+FirstMoveSize+SecondMoveSize(1), ResultOffset+ThirdMoveSize(2), SourceOffset(2*N)
+PRESHADER2_OP(MoveForFixup)		// Opcode(1), MoveCount+Type(1), ResultOffset(2), ComponentStart(1), SourceRegister(2*N), SourceComponent(1*N)
+PRESHADER2_OP(Constant)			// Opcode(1), ResultOffset(2), Dimension+Type(1), Data(N)
 
 // Texture inputs
 PRESHADER2_OP(RuntimeVirtualTextureUniform)
