Fixed issue where level editor viewport cursor was being set to wrong position in some cases due to incorrect math operation on cursor position

#Synced-CL 51479348

diff --git a/Engine/Source/Runtime/Engine/Private/Slate/SceneViewport.cpp b/Engine/Source/Runtime/Engine/Private/Slate/SceneViewport.cpp
--- a/Engine/Source/Runtime/Engine/Private/Slate/SceneViewport.cpp
+++ b/Engine/Source/Runtime/Engine/Private/Slate/SceneViewport.cpp
@@ -126,8 +126,8 @@ void FSceneViewport::ShowCursor( bool bVisible )
 		if( bIsSoftwareCursorVisible )
 		{
 			const FVector2D ToViewportPixelScale = UE::Slate::Viewport::GetViewportLocalToPixelScale(CachedGeometry, FVector2D(SizeX, SizeY));
-			const int32 ClampedMouseX = FMath::Clamp<int32>(SoftwareCursorPosition.X * ToViewportPixelScale.X, 0, SizeX);
-			const int32 ClampedMouseY = FMath::Clamp<int32>(SoftwareCursorPosition.Y * ToViewportPixelScale.Y, 0, SizeY);
+			const int32 ClampedMouseX = FMath::Clamp<int32>(SoftwareCursorPosition.X / ToViewportPixelScale.X, 0, SizeX);
+			const int32 ClampedMouseY = FMath::Clamp<int32>(SoftwareCursorPosition.Y / ToViewportPixelScale.Y, 0, SizeY);
 
 			CurrentReplyState.SetMousePos( CachedGeometry.LocalToAbsolute( FVector2D(ClampedMouseX, ClampedMouseY) ).IntPoint() );
 		}
