TextureDerivedData : VT texture format name must have no platform prefix.
That wasn't being done consistently before FinalizeBuildSettingsForLayer causing the checkslow to trigger.  Instead of checking, just do it there so all callers are gauranteed  to have done it correctly.

#Synced-CL 51463632

diff --git a/Engine/Source/Runtime/Engine/Private/TextureDerivedData.cpp b/Engine/Source/Runtime/Engine/Private/TextureDerivedData.cpp
--- a/Engine/Source/Runtime/Engine/Private/TextureDerivedData.cpp
+++ b/Engine/Source/Runtime/Engine/Private/TextureDerivedData.cpp
@@ -901,18 +901,23 @@ static void FinalizeBuildSettingsForLayer(
 	{
 		// note : FinalizeVirtualTextureLayerFormat is run outside of the normal TextureFormatName set up ; fix?
 		//	should be done inside GetPlatformTextureFormatNamesWithPrefix
-		//	this is only used by Android & iOS
+		//	this is only used by Android & iOS (changes DXT to ASTC/ETC)
 		//  the reason to do it here is we now have bVirtualStreamable, which is not available at the earlier call
 		
-		// FinalizeVirtualTextureLayerFormat assumes (incorrectly) that it gets non-prefixed names, so remove them :
+		// The platform FinalizeVirtualTextureLayerFormat assumes (incorrectly) that it gets non-prefixed names, so remove them :
 
-		// VT does not tile so should never have a platform prefix, but could have an Oodle prefix
-		checkSlow( OutSettings.TextureFormatName == UE::TextureBuildUtilities::TextureFormatRemovePlatformPrefixFromName(OutSettings.TextureFormatName) );
-		
+		// Virtual textures always strip the child format prefix prior to actual encode since VTs never tile.
+		// must match VirtualTextureDataBuilder.cpp
+		OutSettings.TextureFormatName = UE::TextureBuildUtilities::TextureFormatRemovePlatformPrefixFromName(OutSettings.TextureFormatName);
+
+		// May still have an Oodle prefix, remove that :
 		FName NameWithoutPrefix = UE::TextureBuildUtilities::TextureFormatRemovePrefixFromName(OutSettings.TextureFormatName);
 		FName ModifiedName = TargetPlatform->FinalizeVirtualTextureLayerFormat(NameWithoutPrefix);
 		if ( NameWithoutPrefix != ModifiedName )
 		{
+			// if the platform FinalizeVirtualTextureLayerFormat changed our format, take that change (with no Oodle) prefix)
+			//	otherwise, leave TextureFormatName along (with Oodle prefix)
+			// POSSIBLE BUG : if the Platform changed us to a different DXT format that should have Oodle prefix, we lost it here!
 			OutSettings.TextureFormatName = ModifiedName;
 		}
 	}
@@ -1461,6 +1466,9 @@ static void GetBuildSettingsForTargetPlatform(
 	TArray<FTexturePlatformData::FTextureEncodeResultMetadata>* OutResultMetadataPerLayer // can be nullptr if not needed
 )
 {
+	// Beware GetTargetPlatformBuildSettings and GetBuildSettingsForTargetPlatform are almost code dupes of each other; keep in sync!
+	//	todo: refactor to dedupe
+
 	check(TargetPlatform != NULL);
 
 	const UTextureLODSettings* LODSettings = (UTextureLODSettings*)UDeviceProfileManager::Get().FindProfile(TargetPlatform->PlatformName());
@@ -1588,18 +1596,12 @@ static void GetBuildSettingsPerFormat(
 			FTextureBuildSettings& OutSettings = OutSettingPerLayer.Add_GetRef(SourceBuildSettings);
 			OutSettings.TextureFormatName = PlatformFormatsPerLayer[LayerIndex];
 
-			if (OutSettings.bVirtualStreamable)
-			{
-				// Virtual textures always strip the child format prefix prior to actual encode since VTs never tile.
-				// must match VirtualTextureDataBuilder.cpp
-				OutSettings.TextureFormatName = UE::TextureBuildUtilities::TextureFormatRemovePlatformPrefixFromName(OutSettings.TextureFormatName);
-			}
-
 			FTexturePlatformData::FTextureEncodeResultMetadata* OutResultMetadata = nullptr;
 			if (OutResultMetadataPerLayer && SourceResultMetadata)
 			{
 				OutResultMetadata = &OutResultMetadataPerLayer->Add_GetRef(*SourceResultMetadata);			
 			}
+
 			FinalizeBuildSettingsForLayer(Texture, LayerIndex, TargetPlatform, InEncodeSpeed, OutSettings, OutResultMetadata);
 		}
 	}
@@ -4561,6 +4563,12 @@ bool UTexture::DownsizeImageUsingTextureSettings(const ITargetPlatform* TargetPl
 
 void UTexture::GetTargetPlatformBuildSettings(const ITargetPlatform* TargetPlatform, TArray<TArray<FTextureBuildSettings>>& OutSettingPerSupportedFormatPerLayer) const
 {
+	// OutSettingPerSupportedFormatPerLayer is an array for platforms that have multiple different texture format modes (eg. Android_Multi)
+
+	// Beware GetTargetPlatformBuildSettings and GetBuildSettingsForTargetPlatform are almost code dupes of each other; keep in sync!
+	//	todo: refactor to dedupe
+	//GetBuildSettingsForTargetPlatform(*this,TargetPlatform,ETextureEncodeSpeed::Final,...);
+
 	ETextureEncodeSpeed EncodeSpeed = ETextureEncodeSpeed::Final;
 
 	if (!TargetPlatform)
