[Selection] Typed Element Common actions now supports a CanDuplicate so asset editors can customize if they allow for duplicating of the selection

#Synced-CL 51343360

diff --git a/Engine/Source/Runtime/Engine/Private/Elements/Framework/TypedElementCommonActions.cpp b/Engine/Source/Runtime/Engine/Private/Elements/Framework/TypedElementCommonActions.cpp
--- a/Engine/Source/Runtime/Engine/Private/Elements/Framework/TypedElementCommonActions.cpp
+++ b/Engine/Source/Runtime/Engine/Private/Elements/Framework/TypedElementCommonActions.cpp
@@ -57,6 +57,19 @@ bool FTypedElementCommonActionsCustomization::CanCopyElements(ITypedElementWorld
 	return bCanCopy;
 }
 
+bool FTypedElementCommonActionsCustomization::CanDuplicateElements(ITypedElementWorldInterface* InWorldInterface, TArrayView<const FTypedElementHandle> InElementHandles) const
+{
+	for (const FTypedElementHandle& Handle : InElementHandles)
+	{
+		if (!InWorldInterface->CanDuplicateElement(Handle))
+		{
+			return false;
+		}
+	}
+
+	return true;
+}
+
 void FTypedElementCommonActionsCustomization::CopyElements(ITypedElementWorldInterface* InWorldInterface, TArrayView<const FTypedElementHandle> InElementHandles, FOutputDevice& Out)
 {
 	InWorldInterface->CopyElements(InElementHandles, Out);
@@ -121,10 +134,40 @@ bool UTypedElementCommonActions::DeleteNormalizedElements(const FTypedElementLis
 
 TArray<FTypedElementHandle> UTypedElementCommonActions::DuplicateSelectedElements(const UTypedElementSelectionSet* SelectionSet, UWorld* World, const FVector& LocationOffset)
 {
-	FTypedElementListRef NormalizedElements = SelectionSet->GetNormalizedSelection(FTypedElementSelectionNormalizationOptions());
+	FTypedElementListRef NormalizedElements = SelectionSet->GetNormalizedSelection(FTypedElementSelectionNormalizationOptions().SetNameForTEDSIntegration(SelectionSet->GetNameForTedsIntegration()));
 	return DuplicateNormalizedElements(NormalizedElements, World, LocationOffset);
 }
 
+bool UTypedElementCommonActions::CanDuplicateSelectedElements(const UTypedElementSelectionSet* SelectionSet)
+{
+	if (SelectionSet)
+	{
+		if (SelectionSet->GetNumSelectedElements() == 0)
+		{
+			//  return true for legacy duplication code in mode managers to operate on types that may not be hooked into typed elements yet
+			return true;
+		}
+		TMap<FTypedHandleTypeId, TArray<FTypedElementHandle>> ElementsToDuplicateByType;
+		TypedElementUtil::BatchElementsByType(SelectionSet->GetNormalizedSelection(FTypedElementSelectionNormalizationOptions().SetNameForTEDSIntegration(SelectionSet->GetNameForTedsIntegration())), ElementsToDuplicateByType);
+
+		UTypedElementRegistry* Registry = UTypedElementRegistry::GetInstance();
+		for (const auto& ElementsByTypePair : ElementsToDuplicateByType)
+		{
+			FTypedElementCommonActionsCustomization* CommonActionsCustomization = GetInterfaceCustomizationByTypeId(ElementsByTypePair.Key);
+			ITypedElementWorldInterface* WorldInterface = Registry->GetElementInterface<ITypedElementWorldInterface>(ElementsByTypePair.Key);
+			if (CommonActionsCustomization && WorldInterface)
+			{
+				if (!CommonActionsCustomization->CanDuplicateElements(WorldInterface, ElementsByTypePair.Value))
+				{
+					return false;
+				}
+			}
+		}
+	}
+
+	return true;
+}
+
 TArray<FTypedElementHandle> UTypedElementCommonActions::DuplicateNormalizedElements(const FTypedElementListConstPtr& ElementListPtr, UWorld* World, const FVector& LocationOffset)
 {
 	TArray<FTypedElementHandle> NewElements;
diff --git a/Engine/Source/Runtime/Engine/Public/Elements/Framework/TypedElementCommonActions.h b/Engine/Source/Runtime/Engine/Public/Elements/Framework/TypedElementCommonActions.h
--- a/Engine/Source/Runtime/Engine/Public/Elements/Framework/TypedElementCommonActions.h
+++ b/Engine/Source/Runtime/Engine/Public/Elements/Framework/TypedElementCommonActions.h
@@ -28,6 +28,7 @@ public:
 	
 	ENGINE_API virtual bool IsCopyCapable() const;
 	ENGINE_API virtual bool CanCopyElements(ITypedElementWorldInterface* InWorldInterface, TArrayView<const FTypedElementHandle> InElementHandles);
+	ENGINE_API virtual bool CanDuplicateElements(ITypedElementWorldInterface* InWorldInterface, TArrayView<const FTypedElementHandle> InElementHandles) const;
 	ENGINE_API virtual void CopyElements(ITypedElementWorldInterface* InWorldInterface,TArrayView<const FTypedElementHandle> InElementHandles, FOutputDevice& Out);
 	// By default this just calls CopyElements followed by DeleteElements
 	ENGINE_API virtual void CutElements(ITypedElementWorldInterface* InWorldInterface, TArrayView<const FTypedElementHandle> InElementHandles,
@@ -140,6 +141,7 @@ public:
 	 * @note Internally this just calls DuplicateNormalizedElements on the result of UTypedElementSelectionSet::GetNormalizedSelection.
 	 */
 	ENGINE_API TArray<FTypedElementHandle> DuplicateSelectedElements(const UTypedElementSelectionSet* SelectionSet, UWorld* World, const FVector& LocationOffset);
+	ENGINE_API bool CanDuplicateSelectedElements(const UTypedElementSelectionSet* SelectionSet);
 	
 	/**
 	 * Duplicate any elements from the given list that can be duplicated.
