Removed call to do incremental reachability if it has already done so in the same frame

#Synced-CL 51441789

diff --git a/Engine/Source/Runtime/Engine/Private/UnrealEngine.cpp b/Engine/Source/Runtime/Engine/Private/UnrealEngine.cpp
--- a/Engine/Source/Runtime/Engine/Private/UnrealEngine.cpp
+++ b/Engine/Source/Runtime/Engine/Private/UnrealEngine.cpp
@@ -2272,10 +2272,6 @@ void UEngine::ConditionalCollectGarbage()
 
 		LastGCFrame = GFrameCounter;
 	}
-	else if (IsIncrementalReachabilityAnalysisPending())
-	{
-		PerformIncrementalReachabilityAnalysis(GetReachabilityAnalysisTimeLimit());
-	}
 }
 
 void UEngine::PerformGarbageCollectionAndCleanupActors()
