Use the material instance if available to check displacement value on material resource (can be instance overridden)

#Synced-CL 51487274

diff --git a/Engine/Source/Runtime/Engine/Private/Materials/MaterialShared.cpp b/Engine/Source/Runtime/Engine/Private/Materials/MaterialShared.cpp
--- a/Engine/Source/Runtime/Engine/Private/Materials/MaterialShared.cpp
+++ b/Engine/Source/Runtime/Engine/Private/Materials/MaterialShared.cpp
@@ -2192,12 +2192,12 @@ FDisplacementScaling FMaterialResource::GetDisplacementScaling() const
 
 bool FMaterialResource::IsDisplacementFadeEnabled() const
 {
-	return Material->IsDisplacementFadeEnabled();
+	return GetMaterialInterface()->IsDisplacementFadeEnabled();
 }
 
 FDisplacementFadeRange FMaterialResource::GetDisplacementFadeRange() const
 {
-	return Material->GetDisplacementFadeRange();
+	return GetMaterialInterface()->GetDisplacementFadeRange();
 }
 
 uint32 FMaterialResource::GetMaterialDecalResponse() const
