PR #14477: fix missing dllexport/dllimport for GPUSort

#Synced-CL 51499580

#Synced-CL 51499580

diff --git a/Engine/Source/Runtime/Engine/Public/GPUSort.h b/Engine/Source/Runtime/Engine/Public/GPUSort.h
--- a/Engine/Source/Runtime/Engine/Public/GPUSort.h
+++ b/Engine/Source/Runtime/Engine/Public/GPUSort.h
@@ -38,7 +38,7 @@ struct FGPUSortBuffers
 /**
  * Get the number of passes we will need to make in order to sort
  */
-int32 GetGPUSortPassCount(uint32 KeyMask);
+ENGINE_API int32 GetGPUSortPassCount(uint32 KeyMask);
 
 /**
  * Sort a buffer on the GPU.
