[WorldComposition] Fix Rescan break with AsyncLoadingThread, by splitting it into discrete calls in the different paths (initializing a world, enabling composition, cooking)
- Rescan iterates actors, which is not allowed outside the main thread
- It also executes BP code via the AssetRegistry / GetComponentsBoundingBox / BP_GetComponentsBoundingBox, which can trigger checks if it was simply moved to PostLoad
- Duplicate does not require any code change because Tiles are transferred during duplicate.

#tests PIE, local cook/pak/run
#Synced-CL 51373330

diff --git a/Engine/Source/Editor/UnrealEd/Private/UnrealEdMisc.cpp b/Engine/Source/Editor/UnrealEd/Private/UnrealEdMisc.cpp
--- a/Engine/Source/Editor/UnrealEd/Private/UnrealEdMisc.cpp
+++ b/Engine/Source/Editor/UnrealEd/Private/UnrealEdMisc.cpp
@@ -995,10 +995,11 @@ bool FUnrealEdMisc::EnableWorldComposition(UWorld* InWorld, bool bEnable)
 			return false;
 		}
 			
-		auto WorldCompostion = NewObject<UWorldComposition>(InWorld);
+		auto WorldComposition = NewObject<UWorldComposition>(InWorld);
 		// All map files found in the same and folder and all sub-folders will be added ass sub-levels to this map
 		// Make sure user understands this
-		int32 NumFoundSublevels = WorldCompostion->GetTilesList().Num();
+		WorldComposition->Rescan();
+		int32 NumFoundSublevels = WorldComposition->GetTilesList().Num();
 		if (NumFoundSublevels)
 		{
 			FFormatNamedArguments Arguments;
@@ -1009,13 +1010,12 @@ bool FUnrealEdMisc::EnableWorldComposition(UWorld* InWorld, bool bEnable)
 			auto AppResult = FMessageDialog::Open(EAppMsgType::OkCancel, Message);
 			if (AppResult != EAppReturnType::Ok)
 			{
-				WorldCompostion->MarkAsGarbage();
+				WorldComposition->MarkAsGarbage();
 				return false;
 			}
 		}
-			
-		// 
-		InWorld->WorldComposition = WorldCompostion;
+
+		InWorld->WorldComposition = WorldComposition;
 		UWorldComposition::WorldCompositionChangedEvent.Broadcast(InWorld);
 	}
 	
diff --git a/Engine/Source/Editor/WorldBrowser/Private/Tiles/WorldTileCollectionModel.cpp b/Engine/Source/Editor/WorldBrowser/Private/Tiles/WorldTileCollectionModel.cpp
--- a/Engine/Source/Editor/WorldBrowser/Private/Tiles/WorldTileCollectionModel.cpp
+++ b/Engine/Source/Editor/WorldBrowser/Private/Tiles/WorldTileCollectionModel.cpp
@@ -823,7 +823,7 @@ void FWorldTileCollectionModel::OnLevelsCollectionChanged()
 	if (WorldComposition)
 	{
 		// Force rescanning world composition tiles
-		WorldComposition->Rescan();
+		WorldComposition->Rescan(true);
 		
 		// Initialize root level
 		TSharedPtr<FWorldTileModel> RootLevelModel = MakeShareable(new FWorldTileModel(*this, INDEX_NONE));
diff --git a/Engine/Source/Runtime/Engine/Classes/Engine/WorldComposition.h b/Engine/Source/Runtime/Engine/Classes/Engine/WorldComposition.h
--- a/Engine/Source/Runtime/Engine/Classes/Engine/WorldComposition.h
+++ b/Engine/Source/Runtime/Engine/Classes/Engine/WorldComposition.h
@@ -133,7 +133,7 @@ class UWorldComposition : public UObject
 	ENGINE_API FBox GetLevelBounds(ULevel* InLevel) const;
 
 	/** Scans world root folder for relevant packages and initializes world composition structures */
-	ENGINE_API void Rescan();
+	ENGINE_API void Rescan(bool bForce = false);
 
 	/**  */
 	ENGINE_API void ReinitializeForPIE();
@@ -169,7 +169,6 @@ class UWorldComposition : public UObject
 
 private:
 	// UObject interface
-	ENGINE_API virtual void PostInitProperties() override;
 	ENGINE_API virtual void Serialize(FArchive& Ar) override;
 	ENGINE_API virtual void PostDuplicate(bool bDuplicateForPIE) override;
 	ENGINE_API virtual void PostLoad() override;
@@ -214,6 +213,10 @@ public:
 #endif //WITH_EDITOR
 
 private:
+#if WITH_EDITOR
+	bool						bRescanned = false;
+#endif // WITH_EDITOR
+
 	// Path to current world composition (long PackageName)
 	FString						WorldRoot;
 	
diff --git a/Engine/Source/Runtime/Engine/Private/World.cpp b/Engine/Source/Runtime/Engine/Private/World.cpp
--- a/Engine/Source/Runtime/Engine/Private/World.cpp
+++ b/Engine/Source/Runtime/Engine/Private/World.cpp
@@ -2413,6 +2413,12 @@ void UWorld::InitWorld(const InitializationValues IVS)
 	{
 		return;
 	}
+	
+	// Skip Rescan on any cooked build, and also on PIE worlds in editor as Rescan is already computed and its result duplicated in the PIE world
+	if (WorldComposition && WorldType != EWorldType::PIE && !FPlatformProperties::RequiresCookedData())
+	{
+		WorldComposition->Rescan();
+	}
 
 	// This condition has to match the path that leads to GetRendererModule().AllocateScene below as that clobbers the Scene pointer
 	// However, we must call this before starting to create subsystem uobjects and such.
diff --git a/Engine/Source/Runtime/Engine/Private/WorldComposition.cpp b/Engine/Source/Runtime/Engine/Private/WorldComposition.cpp
--- a/Engine/Source/Runtime/Engine/Private/WorldComposition.cpp
+++ b/Engine/Source/Runtime/Engine/Private/WorldComposition.cpp
@@ -42,17 +42,6 @@ UWorldComposition::UWorldComposition(const FObjectInitializer& ObjectInitializer
 {
 }
 
-void UWorldComposition::PostInitProperties()
-{
-	Super::PostInitProperties();
-
-	if (!IsTemplate() && !GetOutermost()->HasAnyPackageFlags(PKG_PlayInEditor) && !FPlatformProperties::RequiresCookedData())
-	{
-		// Tiles information is not serialized to disk in the editor, and should be regenerated on world composition object construction
-		Rescan();
-	}
-}
-
 void UWorldComposition::Serialize( FArchive& Ar )
 {
 	Super::Serialize(Ar);
@@ -73,6 +62,11 @@ void UWorldComposition::PostDuplicate(bool bDuplicateForPIE)
 {
 	Super::PostDuplicate(bDuplicateForPIE);
 
+#if WITH_EDITOR
+	// Tiles have been deserialized from original, so we should be considered up-to-date
+	bRescanned = true;
+#endif // WITH_EDITOR
+
 	if (bDuplicateForPIE)
 	{
 		FixupForPIE(GetOutermost()->GetPIEInstanceID());	
@@ -240,10 +234,17 @@ struct FWorldTilesGatherer
 	}
 };
 
-void UWorldComposition::Rescan()
+void UWorldComposition::Rescan(bool bForce)
 {
 	checkfSlow(!FPlatformProperties::RequiresCookedData(), TEXT("Rescan should not be called in a cooked game with serialized tile info."));
 
+#if WITH_EDITOR
+	if (bRescanned && !bForce)
+	{
+		return;
+	}
+#endif // WITH_EDITOR
+
 	// Save tiles state, so we can restore it for dirty tiles after rescan is done
 	FTilesList SavedTileList = Tiles;
 		
@@ -326,11 +327,14 @@ void UWorldComposition::Rescan()
 
 	// Calculate absolute positions since they are not serialized to disk
 	CaclulateTilesAbsolutePositions();
+
+#if WITH_EDITOR
+	bRescanned = true;
+#endif // WITH_EDITOR
 }
 
 void UWorldComposition::ReinitializeForPIE()
 {
-	Rescan();
 	FixupForPIE(GetOutermost()->GetPIEInstanceID());
 	GetWorld()->SetStreamingLevels(TilesStreaming);
 }
@@ -527,6 +531,9 @@ void UWorldComposition::RestoreDirtyTilesInfo(const FTilesList& TilesPrevState)
 
 void UWorldComposition::CollectTilesToCook(TArray<FString>& PackageNames)
 {
+	// UWorld::InitWorld has not been invoked yet, but we need the tiles so they get packaged. Scan now.
+	Rescan();
+
 	for (const FWorldCompositionTile& Tile : Tiles)
 	{
 		PackageNames.AddUnique(Tile.PackageName.ToString());
