The FRenderAssetStreamingManager::UpdateStreamingRenderAssets Parallel code path should be usable in builds with STATS defined.

#Synced-CL 51502459

diff --git a/Engine/Source/Runtime/Engine/Private/Streaming/StreamingManagerTexture.cpp b/Engine/Source/Runtime/Engine/Private/Streaming/StreamingManagerTexture.cpp
--- a/Engine/Source/Runtime/Engine/Private/Streaming/StreamingManagerTexture.cpp
+++ b/Engine/Source/Runtime/Engine/Private/Streaming/StreamingManagerTexture.cpp
@@ -5,6 +5,8 @@
 =============================================================================*/
 
 #include "Streaming/StreamingManagerTexture.h"
+
+#include "Async/ParallelFor.h"
 #include "Engine/Level.h"
 #include "Engine/World.h"
 #include "Engine/StaticMesh.h"
@@ -32,10 +34,6 @@
 #include "SceneInterface.h"
 #include "Engine/Engine.h"
 
-#if !STATS
-#include "Async/ParallelFor.h"
-#endif
-
 CSV_DECLARE_CATEGORY_MODULE_EXTERN(CORE_API, Basic);
 
 CSV_DEFINE_CATEGORY(TextureStreaming, true);
@@ -1458,26 +1456,31 @@ void FRenderAssetStreamingManager::UpdateStreamingRenderAssets(int32 StageIndex,
 	int32 StartIndex = CurrentUpdateStreamingRenderAssetIndex;
 	int32 EndIndex = AsyncUnsafeStreamingRenderAssets.Num() * (StageIndex + 1) / NumUpdateStages;
 
-#if !STATS
 	if (GAllowParallelUpdateStreamingRenderAssets)
 	{
 		struct FPacket
 		{
-			FPacket(int32 InStartIndex,
-				int32 InEndIndex,
-				TArray<FStreamingRenderAsset>& InStreamingRenderAssets)
-				: StartIndex(InStartIndex),
-				EndIndex(InEndIndex),
-				StreamingRenderAssets(InStreamingRenderAssets)
+			FPacket(int32 InStartIndex, int32 InEndIndex, TArrayView<FStreamingRenderAsset> InStreamingRenderAssets)
+				: StartIndex(InStartIndex)
+				, EndIndex(InEndIndex)
+				, StreamingRenderAssets(InStreamingRenderAssets)
 			{
 			}
+
+			// Inputs
 			int32 StartIndex;
 			int32 EndIndex;
-			TArray<FStreamingRenderAsset>& StreamingRenderAssets;
+			TArrayView<FStreamingRenderAsset> StreamingRenderAssets;
+
+			// Outputs
 			TArray<int32> LocalInflightRenderAssets;
 			TArray<UStreamableRenderAsset*> LocalDeferredTickCBAssets;
+
+			STAT(int64 MipIOBandwidth = 0;)
+
 			char FalseSharingSpacerBuffer[PLATFORM_CACHE_LINE_SIZE];		// separate Packets to avoid false sharing
 		};
+
 		// In order to pass the InflightRenderAssets and DeferredTickCBAssets to each parallel for loop we need to break this into a number of workgroups.
 		// Cannot be too large or the overhead of consolidating all the arrays will take too long.  Cannot be too small or the parallel for will not have
 		// enough work.  Can be adjusted by CVarStreamingParallelRenderAssetsNumWorkgroups.
@@ -1493,7 +1496,8 @@ void FRenderAssetStreamingManager::UpdateStreamingRenderAssets(int32 StageIndex,
 		for (int32 i = 0; i < NumThreadTasks; ++i)
 		{
 			int32 NumAssetsToProcess = FMath::Min<int32>(NumRemaining, NumItemsPerGroup);
-			Packets.Add(FPacket(Start, Start + NumAssetsToProcess, AsyncUnsafeStreamingRenderAssets));
+			Packets.Emplace(Start, Start + NumAssetsToProcess, AsyncUnsafeStreamingRenderAssets);
+
 			Start += NumAssetsToProcess;
 			NumRemaining -= NumAssetsToProcess;
 			if (NumRemaining <= 0)
@@ -1502,7 +1506,7 @@ void FRenderAssetStreamingManager::UpdateStreamingRenderAssets(int32 StageIndex,
 			}
 		}
 
-		ParallelFor(TEXT("RenderAssetStreaming"), Packets.Num(), 1, [this, &Packets, &bWaitForMipFading, &bAsync](int32 PacketIndex)
+		ParallelFor(TEXT("RenderAssetStreaming"), Packets.Num(), 1, [this, &Packets, &bWaitForMipFading](int32 PacketIndex)
 		{
 			FTaskTagScope Scope(ETaskTag::EParallelGameThread);
 
@@ -1516,6 +1520,8 @@ void FRenderAssetStreamingManager::UpdateStreamingRenderAssets(int32 StageIndex,
 					continue;
 				}
 
+				STAT(int32 PreviousResidentMips = StreamingRenderAsset.ResidentMips;)
+
 				const int32* NumStreamedMips;
 				const int32 NumLODGroups = GetNumStreamedMipsArray(StreamingRenderAsset.RenderAssetType, NumStreamedMips);
 
@@ -1526,20 +1532,24 @@ void FRenderAssetStreamingManager::UpdateStreamingRenderAssets(int32 StageIndex,
 				{
 					Packets[PacketIndex].LocalInflightRenderAssets.Add(Index);
 				}
-			}
 
+#if STATS
+				if (StreamingRenderAsset.ResidentMips > PreviousResidentMips)
+				{
+					Packets[PacketIndex].MipIOBandwidth += StreamingRenderAsset.GetSize(StreamingRenderAsset.ResidentMips) - StreamingRenderAsset.GetSize(PreviousResidentMips);
+				}
+#endif
+			}
 		}, EParallelForFlags::BackgroundPriority);
 
-		for (FPacket Packet : Packets) {
+		for (const FPacket& Packet : Packets)
+		{
 			InflightRenderAssets.Append(Packet.LocalInflightRenderAssets);
 			DeferredTickCBAssets.Append(Packet.LocalDeferredTickCBAssets);
-			Packet.LocalInflightRenderAssets.Empty();
-			Packet.LocalDeferredTickCBAssets.Empty();
+			STAT(GatheredStats.MipIOBandwidth += Packet.MipIOBandwidth;)
 		}
-		Packets.Empty();
 	}
 	else
-#endif
 	{
 		for (int32 Index = StartIndex; Index < EndIndex; ++Index)
 		{
