RemoteObjects:  Proper fix for ActorComponents needing their bIsReadyForReplication reset (so they will register for replication upon the server receiving a migrated ActorComponent)

#Synced-CL 51377531

diff --git a/Engine/Source/Runtime/Engine/Private/Components/ActorComponent.cpp b/Engine/Source/Runtime/Engine/Private/Components/ActorComponent.cpp
--- a/Engine/Source/Runtime/Engine/Private/Components/ActorComponent.cpp
+++ b/Engine/Source/Runtime/Engine/Private/Components/ActorComponent.cpp
@@ -771,6 +771,9 @@ void UActorComponent::PostMigrate(const struct FUObjectMigrationContext& Migrati
 			GetWorld()->GetNetDriver()->RemoveObjectFromNetworkIdCache(this);
 		}
 
+		// On the next receive, force components to re-register their replicated subobjects when the owning actor calls ResetOwnedComponents.
+		bIsReadyForReplication = false;
+
 		// If we're loaning, we're keeping ownership; in all other cases we either don't own, or won't own the object after it's migrated
 		const bool bKeepOwnership = (SendType == EObjectMigrationSendType::Loan);
 		if (bKeepOwnership)
@@ -825,19 +828,6 @@ void UActorComponent::PostMigrate(const struct FUObjectMigrationContext& Migrati
 			GetWorld()->GetNetDriver()->AddObjectToNetworkIdCache(this);
 		}
 
-		// Force components to re-register their replicated subobjects when the owning actor calls ResetOwnedComponents,
-		// since the replicated component/subobject lists will be reset in that function.
-		// Even though we only replicate when ownership is assigned, ResetOwnedComponents is called on every receive,
-		// so we also have to rebuild replicated subobjects on every receive.
-		if (AActor* Owner = GetOwner())
-		{
-			// Has ResetOwnedComponents been called yet?
-			if (!Owner->IsReplicatedActorComponentRegistered(this))
-			{
-				bIsReadyForReplication = false;
-			}
-		}
-
 		// We are expecting this flow from AActor, but now we must do this per-component in such a way
 		// that we are not mutating our previously deserialized data.
 		// RegisterAllComponents();
