Crash fix for Incremental CI Full / CanVerseVKEditMin SEW

#Synced-CL 51495466
#Synced-CL 51495466

diff --git a/Engine/Source/Runtime/Engine/Private/LevelInstance/LevelInstanceSubsystem.cpp b/Engine/Source/Runtime/Engine/Private/LevelInstance/LevelInstanceSubsystem.cpp
--- a/Engine/Source/Runtime/Engine/Private/LevelInstance/LevelInstanceSubsystem.cpp
+++ b/Engine/Source/Runtime/Engine/Private/LevelInstance/LevelInstanceSubsystem.cpp
@@ -946,7 +946,7 @@ void ULevelInstanceSubsystem::Tick()
 		OnUpdateStreamingState();
 
 		// Update Editor Mode if we are the GEditor world's Subsystem
-		if (GEditor->GetEditorWorldContext().World() == GetWorld())
+		if (GEditor && GEditor->GetEditorWorldContext().World() == GetWorld())
 		{
 			if (ILevelInstanceEditorModule* EditorModule = (ILevelInstanceEditorModule*)FModuleManager::Get().GetModule("LevelInstanceEditor"))
 			{
