[HWRT] Remove Status member from FRegisteredGeometry
- currently we only check EStatus::Streaming which can be done using StreamingRequestIndex != INDEX_NONE.
- this simplifies logic and avoids issues with out of sync status.

#Synced-CL 51462794

diff --git a/Engine/Source/Runtime/Engine/Private/Rendering/RayTracingGeometryManager.cpp b/Engine/Source/Runtime/Engine/Private/Rendering/RayTracingGeometryManager.cpp
--- a/Engine/Source/Runtime/Engine/Private/Rendering/RayTracingGeometryManager.cpp
+++ b/Engine/Source/Runtime/Engine/Private/Rendering/RayTracingGeometryManager.cpp
@@ -629,11 +629,6 @@ void FRayTracingGeometryManager::RefreshRegisteredGeometry(RayTracing::FGeometry
 		}
 
 		checkf(!AlwaysResidentGeometries.Contains(Handle) || !RegisteredGeometry.bEvicted || !IsRayTracingEnabled(), TEXT("Always resident geometries can't be evicted"));
-
-		if (GeometryInitializer.Type == ERayTracingGeometryInitializerType::StreamingDestination)
-		{
-			RegisteredGeometry.Status = FRegisteredGeometry::EStatus::StreamedOut;
-		}
 	}
 }
 
@@ -1134,7 +1129,7 @@ void FRayTracingGeometryManager::DoesRayTracingGeometryRequireStreaming(const FR
 	bOutStreamBVH = false;
 	bOutStreamBuffers = false;
 
-	if (RegisteredGeometry.Status == FRegisteredGeometry::EStatus::Streaming)
+	if (RegisteredGeometry.StreamingRequestIndex != INDEX_NONE)
 	{
 		// skip if there's already a streaming request in-flight for this geometry
 		return;
@@ -1183,7 +1178,23 @@ bool FRayTracingGeometryManager::RequestRayTracingGeometryStreamIn(FRHICommandLi
 	{
 		// no offline data -> build from VB/IB at runtime
 
-		RegisteredGeometry.Status = FRegisteredGeometry::EStatus::StreamedIn;
+		{
+			FRHIResourceReplaceBatcher Batcher(RHICmdList, 1);
+			FRayTracingGeometryInitializer IntermediateInitializer = Geometry->GetInitializer();
+			IntermediateInitializer.Type = ERayTracingGeometryInitializerType::StreamingSource;
+			IntermediateInitializer.OfflineData = nullptr;
+
+			FRayTracingGeometryRHIRef IntermediateRayTracingGeometry = RHICmdList.CreateRayTracingGeometry(IntermediateInitializer);
+
+			Geometry->SetRequiresBuild(IntermediateInitializer.OfflineData == nullptr || IntermediateRayTracingGeometry->IsCompressed());
+
+			Geometry->InitRHIForStreaming(IntermediateRayTracingGeometry, Batcher);
+
+			// When Batcher goes out of scope it will add commands to copy the BLAS buffers on RHI thread.
+			// We need to do it before we build the current geometry (also on RHI thread).
+		}
+
+		Geometry->RequestBuildIfNeeded(RHICmdList, ERTAccelerationStructureBuildPriority::Normal);
 	}
 	else
 	{
@@ -1222,32 +1233,9 @@ bool FRayTracingGeometryManager::RequestRayTracingGeometryStreamIn(FRHICommandLi
 
 		RegisteredGeometry.StreamableAsset->IssueRequest(StreamingRequest.Request, StreamingRequest.RequestBuffer, StreamingRequest.bBuffersOnly);
 
-		RegisteredGeometry.Status = FRegisteredGeometry::EStatus::Streaming;
-
 		TotalStreamingSize += StreamingRequest.GeometrySize;
 	}
 
-	if (RegisteredGeometry.Status == FRegisteredGeometry::EStatus::StreamedIn)
-	{
-		{
-			FRHIResourceReplaceBatcher Batcher(RHICmdList, 1);
-			FRayTracingGeometryInitializer IntermediateInitializer = Geometry->GetInitializer();
-			IntermediateInitializer.Type = ERayTracingGeometryInitializerType::StreamingSource;
-			IntermediateInitializer.OfflineData = nullptr;
-
-			FRayTracingGeometryRHIRef IntermediateRayTracingGeometry = RHICmdList.CreateRayTracingGeometry(IntermediateInitializer);
-
-			Geometry->SetRequiresBuild(IntermediateInitializer.OfflineData == nullptr || IntermediateRayTracingGeometry->IsCompressed());
-
-			Geometry->InitRHIForStreaming(IntermediateRayTracingGeometry, Batcher);
-
-			// When Batcher goes out of scope it will add commands to copy the BLAS buffers on RHI thread.
-			// We need to do it before we build the current geometry (also on RHI thread).
-		}
-
-		Geometry->RequestBuildIfNeeded(RHICmdList, ERTAccelerationStructureBuildPriority::Normal);
-	}
-
 	return true;
 }
 
@@ -1319,8 +1307,6 @@ void FRayTracingGeometryManager::ProcessCompletedStreamingRequests(FRHICommandLi
 						RegisteredGeometry.StreamableAsset->ReleaseBuffersForStreaming(Batcher);
 					}
 				}
-
-				RegisteredGeometry.Status = FRegisteredGeometry::EStatus::StreamedIn;
 			}
 
 			PendingRequest.Reset();
diff --git a/Engine/Source/Runtime/Engine/Public/Rendering/RayTracingGeometryManager.h b/Engine/Source/Runtime/Engine/Public/Rendering/RayTracingGeometryManager.h
--- a/Engine/Source/Runtime/Engine/Public/Rendering/RayTracingGeometryManager.h
+++ b/Engine/Source/Runtime/Engine/Public/Rendering/RayTracingGeometryManager.h
@@ -121,14 +121,6 @@ private:
 
 		int16 StreamingRequestIndex = INDEX_NONE;
 
-		enum class EStatus : uint8
-		{
-			StreamedOut,
-			Streaming,
-			StreamedIn,
-		};
-
-		EStatus Status = EStatus::StreamedOut;
 		bool bEvicted : 1 = false;
 		bool bAlwaysResident : 1 = false;
 	};
