TangentSpline: Add SetVersion

#Synced-CL 51345399

diff --git a/Engine/Source/Runtime/Engine/Classes/Curves/Splines/TangentSpline.h b/Engine/Source/Runtime/Engine/Classes/Curves/Splines/TangentSpline.h
--- a/Engine/Source/Runtime/Engine/Classes/Curves/Splines/TangentSpline.h
+++ b/Engine/Source/Runtime/Engine/Classes/Curves/Splines/TangentSpline.h
@@ -167,6 +167,7 @@ public:
 	const FInterpCurveVector& GetSplinePointsScale() const { return ScaleCurve; }
 
 	uint32 GetVersion() const { return Version; }
+	void SetVersion(uint32 InVersion) { Version = InVersion; }
 
 private:
 
