Add basic curve table CRUD APIs to DataTableFunctionLibrary.

#Synced-CL 51402579

diff --git a/Engine/Source/Runtime/Engine/Classes/Kismet/DataTableFunctionLibrary.h b/Engine/Source/Runtime/Engine/Classes/Kismet/DataTableFunctionLibrary.h
--- a/Engine/Source/Runtime/Engine/Classes/Kismet/DataTableFunctionLibrary.h
+++ b/Engine/Source/Runtime/Engine/Classes/Kismet/DataTableFunctionLibrary.h
@@ -12,6 +12,7 @@
 #include "Blueprint/BlueprintExceptionInfo.h"
 #include "DataTableFunctionLibrary.generated.h"
 
+struct FRealCurve;
 class UCurveTable;
 
 /** Enum used to indicate success or failure of EvaluateCurveTableRow. */
@@ -34,6 +35,70 @@ class UDataTableFunctionLibrary : public UBlueprintFunctionLibrary
 
 	UFUNCTION(BlueprintCallable, Category = "DataTable", meta = (ExpandEnumAsExecs="OutResult", DataTablePin="CurveTable"))
 	static ENGINE_API void EvaluateCurveTableRow(UCurveTable* CurveTable, FName RowName, float InXY, TEnumAsByte<EEvaluateCurveTableResult::Type>& OutResult, float& OutXY,const FString& ContextString);
+
+	/**
+	 * Convert a curve table to JSON.
+	 * @param CurveTable The Curve Table to export.
+	 * @return The CurveTable contents encoded as a string.
+	 */
+	UFUNCTION(BlueprintCallable, Category = "DataTable")
+	static ENGINE_API FString CurveTableToJson(UCurveTable* CurveTable);
+
+	/**
+	 * Convert a curve table to JSON.
+	 */
+	UFUNCTION(BlueprintCallable, Category = "DataTable")
+	static ENGINE_API TArray<FString> JsonToCurveTable(UCurveTable* CurveTable, const FString& InString);
+
+	/**
+	 * Adds a simple curve to a curve table.
+	 * @param CurveTable The Curve Table to modify
+	 * @param RowName The name of the new row
+	 * @return True if the operation was successful.
+	 */
+	UFUNCTION(BlueprintCallable, Category = "DataTable")
+	static ENGINE_API bool AddSimpleCurveToTable(UCurveTable* CurveTable, FName RowName);
+
+	/**
+	 * Adds a simple curve to a curve table.
+	 * @param CurveTable The Curve Table to modify
+	 * @param RowName The name of the row to modify
+	 * @param DefaultValue The default value of the row.
+	 * @return True if the operation was successful.
+	 */
+	UFUNCTION(BlueprintCallable, Category = "DataTable")
+	static ENGINE_API bool SetCurveTableRowDefault(UCurveTable* CurveTable, FName RowName, float DefaultValue);
+	
+	/**
+	 * Adds a key to a row in a CurveTable
+	 * @param CurveTable The Curve Table to modify.
+	 * @param RowName The name of the row to modify.
+	 * @param InTime The time of the new key.
+	 * @param InValue The value of the new key.
+	 * @return True if the operation was successful.
+	 */
+	UFUNCTION(BlueprintCallable, Category = "DataTable")
+	static ENGINE_API bool AddKey(UCurveTable* CurveTable, FName RowName, float InTime, float InValue);
+
+	/**
+	 * Removes a curve from a curve table.
+	 * @param CurveTable The Curve Table to modify
+	 * @param RowName The name of the row to remove
+	 * @return True if the operation was successful.
+	 */
+	UFUNCTION(BlueprintCallable, Category = "DataTable")
+	static ENGINE_API bool RemoveCurveTableRow(UCurveTable* CurveTable, FName RowName);
+	
+	/**
+	 * Renames a curve table row.
+	 * @param CurveTable The Curve Table to modify
+	 * @param RowName The name of the row to rename
+	 * @param NewRowName The new name of the row
+	 * @return True if the operation was successful.
+	 */
+	UFUNCTION(BlueprintCallable, Category = "DataTable")
+	static ENGINE_API bool RenameCurveTableRow(UCurveTable* CurveTable, FName& RowName, FName& NewRowName);
+
     
 	/** Get the row struct used by the given Data Table, if any */
 	UFUNCTION(BlueprintPure, Category = "DataTable", meta=(ScriptMethod="GetRowStruct"))
diff --git a/Engine/Source/Runtime/Engine/Private/DataTableFunctionLibrary.cpp b/Engine/Source/Runtime/Engine/Private/DataTableFunctionLibrary.cpp
--- a/Engine/Source/Runtime/Engine/Private/DataTableFunctionLibrary.cpp
+++ b/Engine/Source/Runtime/Engine/Private/DataTableFunctionLibrary.cpp
@@ -35,6 +35,88 @@ void UDataTableFunctionLibrary::EvaluateCurveTableRow(UCurveTable* CurveTable, F
 	}
 }
 
+UFUNCTION(BlueprintCallable, Category = "DataTable")
+FString UDataTableFunctionLibrary::CurveTableToJson(UCurveTable* CurveTable)
+{
+	if (!CurveTable)
+	{
+		return FString(TEXT(""));
+	}
+	return CurveTable->GetTableAsJSON();
+}
+
+TArray<FString> UDataTableFunctionLibrary::JsonToCurveTable(UCurveTable* CurveTable, const FString& InString)
+{
+	if (!CurveTable)
+	{
+		return TArray<FString>({ FString(TEXT("CurveTable is null")) });
+	}
+	return CurveTable->CreateTableFromJSONString(InString);
+}
+
+bool UDataTableFunctionLibrary::AddSimpleCurveToTable(UCurveTable* CurveTable, FName RowName)
+{
+	if (!CurveTable)
+	{
+		return false;
+	}
+	CurveTable->AddSimpleCurve(RowName);
+	return true;
+}
+
+bool UDataTableFunctionLibrary::SetCurveTableRowDefault(UCurveTable* CurveTable, FName RowName, float DefaultValue)
+{
+	if (!CurveTable)
+	{
+		return false;
+	}
+	FString Errors;
+	FRealCurve* Curve = CurveTable->FindCurve(RowName, Errors, false);
+	if (!Curve)
+	{
+		return false;
+	}
+	Curve->DefaultValue = DefaultValue;
+	return true;
+}
+
+bool UDataTableFunctionLibrary::AddKey(UCurveTable* CurveTable, FName RowName, float InTime, float InValue)
+{
+	if (!CurveTable)
+	{
+		return false;
+	}
+	FString Errors;
+	FRealCurve* Curve = CurveTable->FindCurve(RowName, Errors, false);
+	if (!Curve)
+	{
+		return false;
+	}
+	Curve->AddKey(InTime, InValue);
+	return true;
+}
+
+bool UDataTableFunctionLibrary::RemoveCurveTableRow(UCurveTable* CurveTable, FName RowName)
+{
+	if (!CurveTable)
+	{
+		return false;
+	}
+	CurveTable->RemoveRow(RowName);
+	return true;
+
+}
+
+bool UDataTableFunctionLibrary::RenameCurveTableRow(UCurveTable* CurveTable, FName& RowName, FName& NewRowName)
+{
+	if (!CurveTable)
+	{
+		return false;
+	}
+	CurveTable->RenameRow(RowName, NewRowName);
+	return true;
+}
+
 const UScriptStruct* UDataTableFunctionLibrary::GetDataTableRowStruct(const UDataTable* Table)
 {
 	return Table
