Making GetNumResidentMips() an API function on Texture2DArray, to match Texture2D

#Synced-CL 51480380

diff --git a/Engine/Source/Runtime/Engine/Classes/Engine/Texture2DArray.h b/Engine/Source/Runtime/Engine/Classes/Engine/Texture2DArray.h
--- a/Engine/Source/Runtime/Engine/Classes/Engine/Texture2DArray.h
+++ b/Engine/Source/Runtime/Engine/Classes/Engine/Texture2DArray.h
@@ -132,9 +132,8 @@ public:
 
 	//~ Begin UStreamableRenderAsset Interface
 	virtual int32 CalcCumulativeLODSize(int32 NumLODs) const final override { return CalcTextureMemorySize(NumLODs); }
+	ENGINE_API int32 GetNumResidentMips() const;
 	virtual bool StreamOut(int32 NewMipCount) final override;
 	virtual bool StreamIn(int32 NewMipCount, bool bHighPrio) final override;
 	//~ End UStreamableRenderAsset Interface
-	int32 GetNumResidentMips() const;
-
 };
\ No newline at end of file
