Fixing minor streaming TSAN issues. The races were on stat variables that can be read while resources are being created.

#Synced-CL 51502459

diff --git a/Engine/Source/Runtime/Engine/Private/Rendering/StreamableTextureResource.cpp b/Engine/Source/Runtime/Engine/Private/Rendering/StreamableTextureResource.cpp
--- a/Engine/Source/Runtime/Engine/Private/Rendering/StreamableTextureResource.cpp
+++ b/Engine/Source/Runtime/Engine/Private/Rendering/StreamableTextureResource.cpp
@@ -15,8 +15,8 @@
 #include "Engine/TextureAllMipDataProviderFactory.h"
 
 #if STATS
-int64 GUITextureMemory = 0;
-int64 GNeverStreamTextureMemory = 0;
+TAtomic<int64> GUITextureMemory = 0;
+TAtomic<int64> GNeverStreamTextureMemory = 0;
 #endif
 
 static TAutoConsoleVariable<int32> CVarPartiallyResidentTextureEnabled(
@@ -175,11 +175,11 @@ void FStreamableTextureResource::IncrementTextureStats() const
 
 	if (LODGroup == TEXTUREGROUP_UI)
 	{
-		GUITextureMemory += TextureSize;
+		GUITextureMemory += static_cast<int64>(TextureSize);
 	}
 	else if (bIsNeverStream)
 	{
-		GNeverStreamTextureMemory += TextureSize;
+		GNeverStreamTextureMemory += static_cast<int64>(TextureSize);
 	}
 }
 
@@ -190,11 +190,11 @@ void FStreamableTextureResource::DecrementTextureStats() const
 
 	if (LODGroup == TEXTUREGROUP_UI)
 	{
-		GUITextureMemory -= TextureSize;
+		GUITextureMemory -= static_cast<int64>(TextureSize);
 	}
 	else if (bIsNeverStream)
 	{
-		GNeverStreamTextureMemory -= TextureSize;
+		GNeverStreamTextureMemory -= static_cast<int64>(TextureSize);
 	}
 }
 #endif
diff --git a/Engine/Source/Runtime/Engine/Private/Streaming/TextureStreamingHelpers.cpp b/Engine/Source/Runtime/Engine/Private/Streaming/TextureStreamingHelpers.cpp
--- a/Engine/Source/Runtime/Engine/Private/Streaming/TextureStreamingHelpers.cpp
+++ b/Engine/Source/Runtime/Engine/Private/Streaming/TextureStreamingHelpers.cpp
@@ -408,8 +408,8 @@ float GShadowmapStreamingFactor = 0.09f;
 bool GNeverStreamOutRenderAssets = false;
 
 #if STATS
-extern int64 GUITextureMemory;
-extern int64 GNeverStreamTextureMemory;
+extern TAtomic<int64> GUITextureMemory;
+extern TAtomic<int64> GNeverStreamTextureMemory;
 extern volatile int64 GPending2DUpdateCount;
 
 int64 GRequiredPoolSizeSum = 0;
@@ -471,8 +471,8 @@ void FRenderAssetStreamingStats::Apply()
 	SET_MEMORY_STAT(STAT_StreamingOverview03_Cached, CachedMips); 
 	SET_MEMORY_STAT(STAT_StreamingOverview04_StreamingOverbudget, FMath::Max<int64>(RequiredPool - StreamingPool, 0)); 
 	SET_MEMORY_STAT(STAT_StreamingOverview05_UnstreamableRenderAssets, NonStreamingMips);
-	SET_MEMORY_STAT(STAT_StreamingOverview06_NeverStream, GNeverStreamTextureMemory);
-	SET_MEMORY_STAT(STAT_StreamingOverview07_UIGroup, GUITextureMemory);
+	SET_MEMORY_STAT(STAT_StreamingOverview06_NeverStream, GNeverStreamTextureMemory.Load(EMemoryOrder::Relaxed));
+	SET_MEMORY_STAT(STAT_StreamingOverview07_UIGroup, GUITextureMemory.Load(EMemoryOrder::Relaxed));
 	SET_MEMORY_STAT(STAT_StreamingOverview08_AverageRequiredPool, GAverageRequiredPool); 
 
 	SET_DWORD_STAT(STAT_StreamingOverview09_NumStreamedMeshes, NumStreamedMeshes);
