[SkeletalMesh]
- Add support for async cook platformdata instead of doing everything on the gamethread
- Suppress warnings about morphtargets during async loading since this is safe

#Synced-CL 51263423

diff --git a/Engine/Source/Runtime/Engine/Classes/Engine/SkeletalMesh.h b/Engine/Source/Runtime/Engine/Classes/Engine/SkeletalMesh.h
--- a/Engine/Source/Runtime/Engine/Classes/Engine/SkeletalMesh.h
+++ b/Engine/Source/Runtime/Engine/Classes/Engine/SkeletalMesh.h
@@ -724,6 +724,7 @@ public:
     /** Warn if the platform supports the minimal number of per vertex bone weights */
 	ENGINE_API void ValidateBoneWeights(const ITargetPlatform* TargetPlatform);
 	ENGINE_API virtual void BeginCacheForCookedPlatformData(const ITargetPlatform* TargetPlatform) override;
+	ENGINE_API virtual bool IsCachedCookedPlatformDataLoaded(const ITargetPlatform* TargetPlatform) override;
 	ENGINE_API virtual void ClearAllCachedCookedPlatformData() override;
 #endif
 
@@ -2868,6 +2869,11 @@ private:
 #if WITH_EDITOR
 	/** Generate SkeletalMeshRenderData from ImportedModel */
 	ENGINE_API void CacheDerivedData(FSkinnedAssetCompilationContext* ContextPtr);
+	
+	/**
+	 * Caches derived renderable for cooked platforms currently active.
+	 */
+	ENGINE_API void PrepareDerivedDataForActiveTargetPlatforms();
 
 	/**
 	 * Initial step for the building process - Can't be done in parallel. USkinnedAsset Interface.
diff --git a/Engine/Source/Runtime/Engine/Private/MorphTarget.cpp b/Engine/Source/Runtime/Engine/Private/MorphTarget.cpp
--- a/Engine/Source/Runtime/Engine/Private/MorphTarget.cpp
+++ b/Engine/Source/Runtime/Engine/Private/MorphTarget.cpp
@@ -583,7 +583,8 @@ void FFinishBuildMorphTargetData::ApplyEditorData(USkeletalMesh * SkeletalMesh,
 	//  which will cause a dead lock since the garbage collector already have the gc lock.
 	//
 	// TODO: remove the UMorphTarget object and replace the data in the import data of the skeletal mesh.
-	if (!IsInGameThread())
+	// During async loading, GC is always inactive so concerns above do not apply. Also, objects aren't yet visible to the game thread so everything should be safe.
+	if (!IsInGameThread() && !IsInAsyncLoadingThread())
 	{
 		UE_ASSET_LOG(LogSkeletalMesh, Display, SkeletalMesh, TEXT("Calling function FFinishBuildMorphTargetData::ApplyEditorData outside of the game thread is not safe and can deadlock or crash."));
 	}
diff --git a/Engine/Source/Runtime/Engine/Private/SkeletalMesh.cpp b/Engine/Source/Runtime/Engine/Private/SkeletalMesh.cpp
--- a/Engine/Source/Runtime/Engine/Private/SkeletalMesh.cpp
+++ b/Engine/Source/Runtime/Engine/Private/SkeletalMesh.cpp
@@ -4137,6 +4137,27 @@ void USkeletalMesh::PostLoad()
 	Super::PostLoad();
 }
 
+#if WITH_EDITOR
+void USkeletalMesh::PrepareDerivedDataForActiveTargetPlatforms()
+{
+	TRACE_CPUPROFILER_EVENT_SCOPE(USkeletalMesh::PrepareDerivedDataForActiveTargetPlatforms);
+
+	ITargetPlatformManagerModule& TargetPlatformManager = GetTargetPlatformManagerRef();
+	ITargetPlatform* RunningPlatform = TargetPlatformManager.GetRunningTargetPlatform();
+	const TArray<ITargetPlatform*>& TargetPlatforms = TargetPlatformManager.GetActiveTargetPlatforms();
+	for (int32 PlatformIndex = 0; PlatformIndex < TargetPlatforms.Num(); ++PlatformIndex)
+	{
+		ITargetPlatform* Platform = TargetPlatforms[PlatformIndex];
+		if (Platform != RunningPlatform)
+		{
+			// Make sure to cache platform data so it doesn't happen lazily during serialization of the skeletal mesh
+			constexpr bool bIsSerializeSaving = false;
+			::GetPlatformSkeletalMeshRenderData(this, Platform, bIsSerializeSaving);
+		}
+	}
+}
+#endif
+
 void USkeletalMesh::ExecutePostLoadInternal(FSkinnedAssetPostLoadContext& Context)
 {
 #if WITH_EDITOR
@@ -4187,6 +4208,9 @@ void USkeletalMesh::ExecutePostLoadInternal(FSkinnedAssetPostLoadContext& Contex
 			CacheDerivedData(&Context);
 			Context.bHasCachedDerivedData = true;
 		}
+
+		// Additionally cache derived data for any other platforms we care about.
+		PrepareDerivedDataForActiveTargetPlatforms();
 	}
 
 	// check the MinLOD values are all within range
@@ -5682,11 +5706,21 @@ void USkeletalMesh::ValidateBoneWeights(const ITargetPlatform* TargetPlatform)
 
 void USkeletalMesh::BeginCacheForCookedPlatformData(const ITargetPlatform* TargetPlatform)
 {
+}
+
+bool USkeletalMesh::IsCachedCookedPlatformDataLoaded(const ITargetPlatform* TargetPlatform)
+{
+	if (IsCompiling())
+	{
+		return false;
+	}
+
 	LLM_SCOPE(ELLMTag::SkeletalMesh);
 	// Make sure to cache platform data so it doesn't happen lazily during serialization of the skeletal mesh
 	constexpr bool bIsSerializeSaving = false;
 	::GetPlatformSkeletalMeshRenderData(this, TargetPlatform, bIsSerializeSaving);
 	ValidateBoneWeights(TargetPlatform);
+	return true;
 }
 
 void USkeletalMesh::ClearAllCachedCookedPlatformData()
