- Fixing issue where materials set to the UI domain would render as the default material even if set to work with static meshes

diff --git a/Engine/Source/Runtime/Engine/Private/Materials/MaterialInterface.cpp b/Engine/Source/Runtime/Engine/Private/Materials/MaterialInterface.cpp
--- a/Engine/Source/Runtime/Engine/Private/Materials/MaterialInterface.cpp
+++ b/Engine/Source/Runtime/Engine/Private/Materials/MaterialInterface.cpp
@@ -2864,7 +2864,7 @@ bool UMaterialInterface::NeedsSetMaterialUsage_Concurrent(bool& bOutHasUsage, EM
 	}
 
 	const EMaterialDomain MaterialDomain = BaseMaterial->MaterialDomain;
-	if (MaterialDomain != MD_Surface && MaterialDomain != MD_DeferredDecal && MaterialDomain != MD_Volume)
+	if (MaterialDomain != MD_Surface && MaterialDomain != MD_DeferredDecal && MaterialDomain != MD_Volume && MaterialDomain != MD_UI)
 	{
 		return false;
 	}
