New material translator Preshader2 integration code improvements.
- Moved all preshader logic to the MaterialIRValueAnalysis.cpp into a single location
- Split some functions into local functions for better readability

#rnx
#Synced-CL 51538672

diff --git a/Engine/Source/Runtime/Engine/Private/Materials/MaterialIRDebug.cpp b/Engine/Source/Runtime/Engine/Private/Materials/MaterialIRDebug.cpp
--- a/Engine/Source/Runtime/Engine/Private/Materials/MaterialIRDebug.cpp
+++ b/Engine/Source/Runtime/Engine/Private/Materials/MaterialIRDebug.cpp
@@ -579,7 +579,7 @@ struct FDebugDumpIRState
 	// Prints Preshader data
 	void DumpPreshader()
 	{
-		const UE::Shader::FPreshaderData& PreshaderData = Module->GetCompilationOutput().UniformExpressionSet.UniformPreshaderData;
+		const UE::Shader::FPreshaderData& PreshaderData = Module->GetCompilationOutput().UniformExpressionSet.GetPreshaderData();
 		if (PreshaderData.Data.IsEmpty())
 		{
 			return;
diff --git a/Engine/Source/Runtime/Engine/Private/Materials/MaterialIRModuleBuilder.cpp b/Engine/Source/Runtime/Engine/Private/Materials/MaterialIRModuleBuilder.cpp
--- a/Engine/Source/Runtime/Engine/Private/Materials/MaterialIRModuleBuilder.cpp
+++ b/Engine/Source/Runtime/Engine/Private/Materials/MaterialIRModuleBuilder.cpp
@@ -88,6 +88,7 @@ struct FMaterialIRModuleBuilderImpl
 	FMaterialIRModuleBuilder* Builder;
 	FMaterialIRModule* Module;
 	MIR::FEmitter Emitter;
+	FMaterialIRPreshader Preshader;
 	TArray<FAnalysisContext> AnalysisContextStack;
 	FMaterialIRValueAnalyzer ValueAnalyzer;
 	MIR::FValue* DefaultMaterialAggregate;
@@ -101,7 +102,7 @@ struct FMaterialIRModuleBuilderImpl
 		this->Module = InModule;
 
 		// Setup the Builder implementation
-		ValueAnalyzer.Setup(Builder->MaterialInterface, Module, &Module->CompilationOutput,  Builder->TargetInsights);
+		ValueAnalyzer.Setup(Builder->MaterialInterface, Module, &Preshader, Builder->TargetInsights);
 
 		// Empty the module and set it up
 		Module->Empty();
@@ -1076,22 +1077,6 @@ struct FMaterialIRModuleBuilderImpl
 				// Note that individual value processing is independent from the stage it runs on so we can perform it only once.
 				if ((Value->Flags & MIR::EValueFlags::AnalyzedInAllStagesMask) == MIR::EValueFlags::None)
 				{
-					// Flow the graph properties downstream from the value's uses into this value.
-					for (MIR::FValue* Use : Value->GetUses())
-					{
-						if (Use)
-						{
-							// Note that we don't put this UnionMask of flags in the enum class, because it breaks Natvis flag display.
-							const MIR::EGraphProperties UnionMask = MIR::EGraphProperties::PixelStageOnly | MIR::EGraphProperties::ReadsPixelNormal | MIR::EGraphProperties::HasParameter;
-
-							Value->GraphProperties |= Use->GraphProperties & UnionMask;
-
-							// Uniform property is a logical intersection (AND) rather than union (OR), but we also need to make sure we don't
-							// clear any of the unioned properties, so always include those in the AND mask.
-							Value->GraphProperties &= (Use->GraphProperties & MIR::EGraphProperties::Uniform) | UnionMask;
-						}
-					}
-
 					// Allocate the entry points linkage information for this instruction.
 					if (MIR::FInstruction* Instr = MIR::AsInstruction(Value))
 					{
@@ -1116,13 +1101,6 @@ struct FMaterialIRModuleBuilderImpl
 			}
 		}
 
-		// After the entire tree has been analyzed, we can run fixup, which allocates output buffer offsets and fixes up
-		// parameter evaluations and preshader bytecode with those.
-		if (Module->IsValid())
-		{
-			ValueAnalyzer.Preshader.PreshaderFixup(ValueAnalyzer.GetPreshaderData(), Module, ValueAnalyzer.Insights);
-		}
-
 		return Module->IsValid();
 	}
 
@@ -1143,6 +1121,11 @@ struct FMaterialIRModuleBuilderImpl
 
 	bool Step_PerformModuleWideAnalysis()
 	{
+		// After the entire tree has been analyzed, we can run fixup, which allocates output buffer offsets and fixes up
+		// parameter evaluations and preshader bytecode with those.
+		UE::Shader::FPreshaderData& PreshaderData = Module->GetCompilationOutput().UniformExpressionSet.GetPreshaderData();
+		Preshader.PreshaderFixup(PreshaderData, Module, ValueAnalyzer.Insights);
+
 		// If the material doesn't use expression shading models, or they aren't valid, initialize the shading models to the ones from the material.
 		if (!Builder->MaterialInterface->IsShadingModelFromMaterialExpression() || !Module->ShadingModelsFromCompilation.IsValid())
 		{
diff --git a/Engine/Source/Runtime/Engine/Private/Materials/MaterialIRValueAnalyzer.cpp b/Engine/Source/Runtime/Engine/Private/Materials/MaterialIRValueAnalyzer.cpp
--- a/Engine/Source/Runtime/Engine/Private/Materials/MaterialIRValueAnalyzer.cpp
+++ b/Engine/Source/Runtime/Engine/Private/Materials/MaterialIRValueAnalyzer.cpp
@@ -245,7 +245,7 @@ static void PrimitiveUniform_Analyze(FMaterialIRValueAnalyzer& Analyzer, MIR::FP
 	PrimitiveUniform->Analysis_UniformIndex = UniformExpressionSet.FindOrAddNumericParameter(MaterialParameterType, { FName(PrimitiveUniform->Name) }, DefaultValueOffset);
 
 	// Allocate the space for the parameter on the uniform buffer.
-	int32 PreshaderIndex = Analyzer.Preshader.AddValueInfo(PrimitiveUniform);
+	int32 PreshaderIndex = Analyzer.Preshader->AddValueInfo(PrimitiveUniform);
 	 
 	// Add the parameter evaluation to the uniform data.  Note that we don't know the used component count until the entire
 	// tree has been processed, so the final component count is filled in during fixup.
@@ -316,12 +316,12 @@ static void PreshaderParameter_Analyze(FMaterialIRValueAnalyzer& Analyzer, MIR::
 		Preshader2Type = EPreshader2Type::LWC;
 	}
 
-	int32 PreshaderIndex = Analyzer.Preshader.AddValueInfo(Parameter);
+	int32 PreshaderIndex = Analyzer.Preshader->AddValueInfo(Parameter);
 
 	const FHashedMaterialParameterInfo HashedSourceParameterInfo{ SourceParameterName, EMaterialParameterAssociation::GlobalParameter, INDEX_NONE };
 
 	// UniformIndex argument only applies for RuntimeVirtualTextureUniform, but it's harmless to pass it in regardless.
-	UE::Preshader2::EmitTextureOp(Analyzer.Preshader.Data, Parameter->Opcode,
+	UE::Preshader2::EmitTextureOp(Analyzer.Preshader->Data, Parameter->Opcode,
 		Preshader2Type, (uint8)PrimitiveType.NumComponents(), (uint16)PreshaderIndex,
 		HashedSourceParameterInfo, Parameter->TextureIndex, Parameter->Payload.UniformIndex);
 }
@@ -390,142 +390,165 @@ static void SetVertexInterpolator_Analyze(FMaterialIRValueAnalyzer& Analyzer, MI
 	}
 }
 
-//////////////////////////////////////////////////////////////////////////////////////////
-// FMaterialIRValueAnalyzer
-//////////////////////////////////////////////////////////////////////////////////////////
-
-void FMaterialIRValueAnalyzer::Setup(UMaterialInterface* InMaterialInterface, FMaterialIRModule* InModule, FMaterialCompilationOutput* InCompilationOutput, FMaterialInsights* InInsights)
+static void RecurseMarkPreshaderTemp(MIR::FValue* Value)
 {
-	MaterialInterface = InMaterialInterface;
-	Module = InModule;
-	Insights = InInsights;
-	CompilationOutput = InCompilationOutput;
-	UniformPrimitiveDefaultValueOffsets = {};
-	EnvironmentDefines = {};
-	VTStacks = {};
+	for (MIR::FValue* Use : Value->GetUses())
+	{
+		// Check if this value has already been marked as preshader (either Out or Temp).
+		// If so, we can also assume its children have already been marked.
+		if (Use &&
+			!Use->HasSubgraphProperties(MIR::EGraphProperties::PreshaderOut) &&
+			!Use->HasSubgraphProperties(MIR::EGraphProperties::PreshaderTemp) &&
+			Use->HasSubgraphProperties(MIR::EGraphProperties::HasParameter))
+		{
+			Use->SetSubgraphProperties(MIR::EGraphProperties::PreshaderTemp);
+
+			// TODO: the builder already performs analysis graph crawling. We should avoid doing further ad-hoc crawls 
+			// and instead integrate this logic inside the existing behavior.
+			RecurseMarkPreshaderTemp(Use);
+		}
+	}
 }
 
-void FMaterialIRValueAnalyzer::Analyze(MIR::FValue* Value)
+// Marks the given value as PreshaderOut, if the value's properties allow it, and it's not already marked.
+static void MarkPreshaderOut(FMaterialIRValueAnalyzer& Analyzer, MIR::FValue* Value)
 {
-	AnalyzeUses(Value);
+	// Check if this is the first time we're marking the value as a preshader output.
+	if (!Value->HasSubgraphProperties(MIR::EGraphProperties::Uniform | MIR::EGraphProperties::HasParameter) || Value->HasSubgraphProperties(MIR::EGraphProperties::PreshaderOut))
+	{
+		return;
+	}
 
-	#define EXPAND_CASE(ValueType) \
-		case MIR::VK_##ValueType: \
-		{ \
-			ValueType##_Analyze(*this, static_cast<MIR::F##ValueType*>(Value)); \
-			break; \
-		}
+	// Both Uniform and HasParameter are set, mark it as PreshaderOut, and its children as PreshaderTemp if not already marked.
+	Value->SetSubgraphProperties(MIR::EGraphProperties::PreshaderOut);
+	Value->GraphProperties &= ~MIR::EGraphProperties::PreshaderTemp;
 
-	// IMPORTANT: Before adding a case here, read the FMaterialIRValueAnalyzer documentation.
-	switch (Value->Kind)
+	// If this is an alias, mark the underlying value as PreshaderOut as well.  This may waste some space in the output buffer
+	// if there is a mix of used and unused components, but this is assumed to be rare, and we want to err on the side of avoiding
+	// the need to generate extra move opcodes, to save CPU perf.
+	if (Value->Analysis_PreshaderOffset)
 	{
-		EXPAND_CASE(AnalysisError);
-		EXPAND_CASE(MaterialParameterCollection);
-		EXPAND_CASE(PrimitiveUniform);
-		EXPAND_CASE(TextureUniform);
-		EXPAND_CASE(VirtualTextureUniform);
-		EXPAND_CASE(Branch);
-		EXPAND_CASE(PreshaderParameter);
-		EXPAND_CASE(SetMaterialOutput);
-		EXPAND_CASE(Extern);
-		EXPAND_CASE(GetVertexInterpolator);
-		EXPAND_CASE(SetVertexInterpolator);
-		default: break;
+		MIR::FValue* Alias = Analyzer.Preshader->ValueInfos[Value->Analysis_PreshaderOffset].Alias;
+		if (Alias)
+		{
+			Alias->SetSubgraphProperties(MIR::EGraphProperties::PreshaderOut);
+			Alias->GraphProperties &= ~MIR::EGraphProperties::PreshaderTemp;
+		}
 	}
 
-	#undef EXPAND_CASE
+	RecurseMarkPreshaderTemp(Value);
 
-	// After analyzing the value, check if this is a preshader value (Uniform with a Parameter in its subtree),
-	// and if so, emit preshader code for it.  Note that the emitter only sets Uniform for operators that are
-	// supported, so we don't have to worry about hitting this code for unsupported operators.
-	if (Value->HasSubgraphProperties(MIR::EGraphProperties::Uniform | MIR::EGraphProperties::HasParameter))
+	// Composite is emitted here, as their emission is deferred, assuming composites will often be folded into operators or
+	// dead stripped by the trivial tail swizzle optimization, unless there is a case where the final result actually needs
+	// to be available on the GPU (the case we are encountering here).
+	if (!Value->Analysis_PreshaderOffset)
 	{
-		Preshader.EmitValue(Value);
+		if (MIR::FComposite* Composite = Value->As<MIR::FComposite>())
+		{
+			Analyzer.Preshader->EmitComposite(Composite);
+		}
+		else
+		{
+			// All other preshader compatible child expressions should already have been processed.
+			UE_MIR_UNREACHABLE();
+		}
 	}
 }
 
-void FMaterialIRValueAnalyzer::AnalyzeInStage(MIR::FValue* Value, MIR::EStage Stage)
+// Trivial Tail Swizzle optimization -- if this use is a swizzle (Composite of Subscripts all pointing to the same
+// vector input value, or all identical values), we want to mark that underlying value as "PreshaderOut" rather than
+// this one.  This prevents trivial tail swizzles from generating extra CPU preshader opcodes, given that swizzle on
+// load is free.  For example, this is a common pattern that occurs when 3 components of a vector param are fetched:
+//
+//     _0 = Composite( Subscript(Param,0), Subscript(Param,1), Subscript(Param,2) );
+//
+// We're trying to avoid this generating an ".xyz" swizzle on the CPU.  For this example, the swizzle has the same
+// order of components, but this optimization also applies if the components are reordered.  We want to apply this
+// optimization here, to simplify downstream logic, which may involve multiple passes over the tree.  We can also
+// accept constants as any of the parameters, as swizzling in a constant is also free.
+static bool ApplyTrivialTailSwizzle(FMaterialIRValueAnalyzer& Analyzer, MIR::FValue* Use)
 {
-	#define EXPAND_CASE(ValueType) \
-	case MIR::VK_##ValueType: \
-	{ \
-		ValueType##_AnalyzeInStage(*this, static_cast<MIR::F##ValueType*>(Value), Stage); \
-		break; \
-	}
-
-	// IMPORTANT: Before adding a case here, read the FMaterialIRValueAnalyzer documentation.
-	switch (Value->Kind)
+	if (!(Use->IsA(MIR::VK_Composite) && Use->Type.IsVector() && Use->GetUses().Num() > 0))
 	{
-		EXPAND_CASE(Extern);
-		EXPAND_CASE(VTPageTableRead);
-		EXPAND_CASE(HardwarePartialDerivative);
-		default: break;
+		return false;
 	}
 
-	#undef EXPAND_CASE
-}
+	TConstArrayView<MIR::FValue*> Components = Use->As<MIR::FComposite>()->GetComponents();
+	MIR::FSubscript* FirstSubscript = nullptr;
+	MIR::FValue* FirstValue = nullptr;
 
-static void RecurseMarkPreshaderTemp(MIR::FValue* Value)
-{
-	for (MIR::FValue* Use : Value->GetUses())
+	// Assume it's trivial until we find a case where it's not.
+	bool bTrivialTailSwizzle = true;
+
+	for (int32 ComponentIndex = 0; ComponentIndex < Components.Num(); ComponentIndex++)
 	{
-		// Check if this value has already been marked as preshader (either Out or Temp).
-		// If so, we can also assume its children have already been marked.
-		if (Use &&
-			!Use->HasSubgraphProperties(MIR::EGraphProperties::PreshaderOut) &&
-			!Use->HasSubgraphProperties(MIR::EGraphProperties::PreshaderTemp) &&
-			Use->HasSubgraphProperties(MIR::EGraphProperties::HasParameter))
+		// Constants can just be ignored
+		if (Components[ComponentIndex]->IsA(MIR::VK_Constant))
 		{
-			Use->SetSubgraphProperties(MIR::EGraphProperties::PreshaderTemp);
-
-			RecurseMarkPreshaderTemp(Use);
+			continue;
 		}
-	}
-}
-
-void FMaterialIRValueAnalyzer::MarkPreshaderOut(MIR::FValue* Value)
-{
-	// Check if this is the first time we're marking the value as a preshader output.
-	if (Value->HasSubgraphProperties(MIR::EGraphProperties::Uniform | MIR::EGraphProperties::HasParameter) && !Value->HasSubgraphProperties(MIR::EGraphProperties::PreshaderOut))
-	{
-		// Both Uniform and HasParameter are set, mark it as PreshaderOut, and its children as PreshaderTemp if not already marked.
-		Value->SetSubgraphProperties(MIR::EGraphProperties::PreshaderOut);
-		Value->GraphProperties &= ~MIR::EGraphProperties::PreshaderTemp;
 
-		// If this is an alias, mark the underlying value as PreshaderOut as well.  This may waste some space in the output buffer
-		// if there is a mix of used and unused components, but this is assumed to be rare, and we want to err on the side of avoiding
-		// the need to generate extra move opcodes, to save CPU perf.
-		if (Value->Analysis_PreshaderOffset)
+		if (Components[ComponentIndex]->IsA(MIR::VK_Subscript))
 		{
-			MIR::FValue* Alias = Preshader.ValueInfos[Value->Analysis_PreshaderOffset].Alias;
-			if (Alias)
+			// If there's a mix of (non-constant) values and subscripts, it's not compatible with trivial tail swizzle!
+			if (FirstValue)
 			{
-				Alias->SetSubgraphProperties(MIR::EGraphProperties::PreshaderOut);
-				Alias->GraphProperties &= ~MIR::EGraphProperties::PreshaderTemp;
+				bTrivialTailSwizzle = false;
+				break;
 			}
-		}
-
-		RecurseMarkPreshaderTemp(Value);
 
-		// Composite is emitted here, as their emission is deferred, assuming composites will often be folded into operators or
-		// dead stripped by the trivial tail swizzle optimization, unless there is a case where the final result actually needs
-		// to be available on the GPU (the case we are encountering here).
-		if (!Value->Analysis_PreshaderOffset)
+			MIR::FSubscript* CurrentSubscript = Components[ComponentIndex]->As<MIR::FSubscript>();
+			if (!FirstSubscript)
+			{
+				FirstSubscript = CurrentSubscript;
+			}
+			else if (FirstSubscript->Arg != CurrentSubscript->Arg)
+			{
+				bTrivialTailSwizzle = false;
+				break;
+			}
+		}
+		else
 		{
-			if (MIR::FComposite* Composite = Value->As<MIR::FComposite>())
+			// If there's a mix of (non-constant) values and subscripts, it's not compatible with trivial tail swizzle!
+			if (FirstSubscript)
 			{
-				Preshader.EmitComposite(Composite);
+				bTrivialTailSwizzle = false;
+				break;
 			}
-			else
+
+			if (!FirstValue)
 			{
-				// All other preshader compatible child expressions should already have been processed.
-				UE_MIR_UNREACHABLE();
+				FirstValue = Components[ComponentIndex];
 			}
+			else if (FirstValue != Components[ComponentIndex])
+			{
+				bTrivialTailSwizzle = false;
+				break;
+			}
+		}
+	}
+
+	if (bTrivialTailSwizzle)
+	{
+		if (FirstSubscript)
+		{
+			// Subscripts are always aliases of the underlying arg, so we need to mark that as used by non-uniform.
+			MarkPreshaderOut(Analyzer, FirstSubscript->Arg);
+		}
+		else if (FirstValue)
+		{
+			MarkPreshaderOut(Analyzer, FirstValue);
 		}
+
+		Use->SetSubgraphProperties(MIR::EGraphProperties::TrivialTailSwizzle);
 	}
+
+	return bTrivialTailSwizzle;
 }
 
-void FMaterialIRValueAnalyzer::AnalyzeUses(MIR::FValue* Value)
+// Handles logic that involves analyzing Uses of the given value.  Updates used component masks and sets PreshaderOut where needed.
+static void AnalyzeUses(FMaterialIRValueAnalyzer& Analyzer, MIR::FValue* Value)
 {
 	for (MIR::FValue* Use : Value->GetUses())
 	{
@@ -536,6 +559,8 @@ void FMaterialIRValueAnalyzer::AnalyzeUses(MIR::FValue* Value)
 
 		// Updated used component mask of uniform values.  Can't be done in UniformParameter_Analyze, because it needs
 		// access to the parent Value to fetch its subscript index.
+		// TODO: this code fights a little with the way the analyzer is designed. We already have a facility to analyze individual
+		//       nodes in the graph (i.e. PrimitiveUniform_Analyze()).
 		if (MIR::FPrimitiveUniform* UniformParameter = Use->As<MIR::FPrimitiveUniform>())
 		{
 			if (UniformParameter->Type.IsScalar())
@@ -554,104 +579,110 @@ void FMaterialIRValueAnalyzer::AnalyzeUses(MIR::FValue* Value)
 				}
 			}
 		}
-
+		
 		// If this graph node is not Uniform, attempt to mark any of its children as PreshaderOut where necessary.
 		if (!(Value->GraphProperties & MIR::EGraphProperties::Uniform))
-		{
-			// Trivial Tail Swizzle optimization -- if this use is a swizzle (Composite of Subscripts all pointing to the same
-			// vector input value, or all identical values), we want to mark that underlying value as "PreshaderOut" rather than
-			// this one.  This prevents trivial tail swizzles from generating extra CPU preshader opcodes, given that swizzle on
-			// load is free.  For example, this is a common pattern that occurs when 3 components of a vector param are fetched:
-			//
-			//     _0 = Composite( Subscript(Param,0), Subscript(Param,1), Subscript(Param,2) );
-			//
-			// We're trying to avoid this generating an ".xyz" swizzle on the CPU.  For this example, the swizzle has the same
-			// order of components, but this optimization also applies if the components are reordered.  We want to apply this
-			// optimization here, to simplify downstream logic, which may involve multiple passes over the tree.  We can also
-			// accept constants as any of the parameters, as swizzling in a constant is also free.
-			bool bTrivialTailSwizzle = false;
-
-			if (Use->IsA(MIR::VK_Composite) && Use->Type.IsVector() && Use->GetUses().Num() > 0)
+		{	
+			if (!ApplyTrivialTailSwizzle(Analyzer, Use))
 			{
-				TConstArrayView<MIR::FValue*> Components = Use->As<MIR::FComposite>()->GetComponents();
-				MIR::FSubscript* FirstSubscript = nullptr;
-				MIR::FValue* FirstValue = nullptr;
+				MarkPreshaderOut(Analyzer, Use);
+			}
+		}
+	}
+}
 
-				// Assume it's trivial until we find a case where it's not.
-				bTrivialTailSwizzle = true;
 
-				for (int32 ComponentIndex = 0; ComponentIndex < Components.Num(); ComponentIndex++)
-				{
-					// Constants can just be ignored
-					if (Components[ComponentIndex]->IsA(MIR::VK_Constant))
-					{
-						continue;
-					}
-
-					if (Components[ComponentIndex]->IsA(MIR::VK_Subscript))
-					{
-						// If there's a mix of (non-constant) values and subscripts, it's not compatible with trivial tail swizzle!
-						if (FirstValue)
-						{
-							bTrivialTailSwizzle = false;
-							break;
-						}
-
-						MIR::FSubscript* CurrentSubscript = Components[ComponentIndex]->As<MIR::FSubscript>();
-						if (!FirstSubscript)
-						{
-							FirstSubscript = CurrentSubscript;
-						}
-						else if (FirstSubscript->Arg != CurrentSubscript->Arg)
-						{
-							bTrivialTailSwizzle = false;
-							break;
-						}
-					}
-					else
-					{
-						// If there's a mix of (non-constant) values and subscripts, it's not compatible with trivial tail swizzle!
-						if (FirstSubscript)
-						{
-							bTrivialTailSwizzle = false;
-							break;
-						}
-
-						if (!FirstValue)
-						{
-							FirstValue = Components[ComponentIndex];
-						}
-						else if (FirstValue != Components[ComponentIndex])
-						{
-							bTrivialTailSwizzle = false;
-							break;
-						}
-					}
-				}
+//////////////////////////////////////////////////////////////////////////////////////////
+// FMaterialIRValueAnalyzer
+//////////////////////////////////////////////////////////////////////////////////////////
 
-				if (bTrivialTailSwizzle)
-				{
-					if (FirstSubscript)
-					{
-						// Subscripts are always aliases of the underlying arg, so we need to mark that as used by non-uniform.
-						MarkPreshaderOut(FirstSubscript->Arg);
-					}
-					else if (FirstValue)
-					{
-						MarkPreshaderOut(FirstValue);
-					}
-
-					Use->SetSubgraphProperties(MIR::EGraphProperties::TrivialTailSwizzle);
-				}
-			}
+void FMaterialIRValueAnalyzer::Setup(UMaterialInterface* InMaterialInterface, FMaterialIRModule* InModule, FMaterialIRPreshader* InPreshader, FMaterialInsights* InInsights)
+{
+	MaterialInterface = InMaterialInterface;
+	Module = InModule;
+	Insights = InInsights;
+	Preshader = InPreshader;
+	CompilationOutput = &InModule->GetCompilationOutput();
+	UniformPrimitiveDefaultValueOffsets = {};
+	EnvironmentDefines = {};
+	VTStacks = {};
+}
 
-			// Trivial tail swizzles are handled above.
-			if (!bTrivialTailSwizzle)
-			{
-				MarkPreshaderOut(Use);
-			}
+void FMaterialIRValueAnalyzer::Analyze(MIR::FValue* Value)
+{
+	// Flow the graph properties downstream from the value's uses into this value.
+	for (MIR::FValue* Use : Value->GetUses())
+	{
+		if (Use)
+		{
+			// Note that we don't put this UnionMask of flags in the enum class, because it breaks Natvis flag display.
+			const MIR::EGraphProperties UnionMask = MIR::EGraphProperties::PixelStageOnly | MIR::EGraphProperties::ReadsPixelNormal | MIR::EGraphProperties::HasParameter;
+
+			Value->GraphProperties |= Use->GraphProperties & UnionMask;
+
+			// Uniform property is a logical intersection (AND) rather than union (OR), but we also need to make sure we don't
+			// clear any of the unioned properties, so always include those in the AND mask.
+			Value->GraphProperties &= (Use->GraphProperties & MIR::EGraphProperties::Uniform) | UnionMask;
+		}
+	}
+
+	// Analyze this value uses
+	AnalyzeUses(*this, Value);
+
+	#define EXPAND_CASE(ValueType) \
+		case MIR::VK_##ValueType: \
+		{ \
+			ValueType##_Analyze(*this, static_cast<MIR::F##ValueType*>(Value)); \
+			break; \
 		}
+
+	// IMPORTANT: Before adding a case here, read the FMaterialIRValueAnalyzer documentation.
+	switch (Value->Kind)
+	{
+		EXPAND_CASE(AnalysisError);
+		EXPAND_CASE(MaterialParameterCollection);
+		EXPAND_CASE(PrimitiveUniform);
+		EXPAND_CASE(TextureUniform);
+		EXPAND_CASE(VirtualTextureUniform);
+		EXPAND_CASE(Branch);
+		EXPAND_CASE(PreshaderParameter);
+		EXPAND_CASE(SetMaterialOutput);
+		EXPAND_CASE(Extern);
+		EXPAND_CASE(GetVertexInterpolator);
+		EXPAND_CASE(SetVertexInterpolator);
+		default: break;
 	}
+
+	#undef EXPAND_CASE
+
+	// After analyzing the value, check if this is a preshader value (Uniform with a Parameter in its subtree),
+	// and if so, emit preshader code for it.  Note that the emitter only sets Uniform for operators that are
+	// supported, so we don't have to worry about hitting this code for unsupported operators.
+	if (Value->HasSubgraphProperties(MIR::EGraphProperties::Uniform | MIR::EGraphProperties::HasParameter))
+	{
+		Preshader->EmitValue(Value);
+	}
+}
+
+void FMaterialIRValueAnalyzer::AnalyzeInStage(MIR::FValue* Value, MIR::EStage Stage)
+{
+	#define EXPAND_CASE(ValueType) \
+	case MIR::VK_##ValueType: \
+	{ \
+		ValueType##_AnalyzeInStage(*this, static_cast<MIR::F##ValueType*>(Value), Stage); \
+		break; \
+	}
+
+	// IMPORTANT: Before adding a case here, read the FMaterialIRValueAnalyzer documentation.
+	switch (Value->Kind)
+	{
+		EXPAND_CASE(Extern);
+		EXPAND_CASE(VTPageTableRead);
+		EXPAND_CASE(HardwarePartialDerivative);
+		default: break;
+	}
+
+	#undef EXPAND_CASE
 }
 
 #endif // #if WITH_EDITOR
diff --git a/Engine/Source/Runtime/Engine/Private/Materials/MaterialIRValueAnalyzer.h b/Engine/Source/Runtime/Engine/Private/Materials/MaterialIRValueAnalyzer.h
--- a/Engine/Source/Runtime/Engine/Private/Materials/MaterialIRValueAnalyzer.h
+++ b/Engine/Source/Runtime/Engine/Private/Materials/MaterialIRValueAnalyzer.h
@@ -44,7 +44,7 @@
 struct FMaterialIRValueAnalyzer
 {
 	// Resets the analyzer to process a new translation run.
-	void Setup(UMaterialInterface* InMaterialInterface, FMaterialIRModule* InModule, FMaterialCompilationOutput* InCompilationOutput, FMaterialInsights* InInsights = nullptr);
+	void Setup(UMaterialInterface* InMaterialInterface, FMaterialIRModule* InModule, FMaterialIRPreshader* InPreshader, FMaterialInsights* InInsights = nullptr);
 		
 	// Performs stage-agnostic analysis on the given value.
 	void Analyze(MIR::FValue* Value);
@@ -52,17 +52,6 @@ struct FMaterialIRValueAnalyzer
 	// Performs stage-specific analysis on the given value, if the value kind requires some.
 	void AnalyzeInStage(MIR::FValue* Value, MIR::EStage Stage);
 
-	// Return writable preshader data.
-	UE::Shader::FPreshaderData& GetPreshaderData() { return CompilationOutput->UniformExpressionSet.UniformPreshaderData; }
-
-private:
-	// Marks the given value as PreshaderOut, if the value's properties allow it, and it's not already marked.
-	void MarkPreshaderOut(MIR::FValue* Value);
-
-	// Handles logic that involves analyzing Uses of the given value.  Updates used component masks and sets PreshaderOut where needed.
-	void AnalyzeUses(MIR::FValue* Value);
-
-public:
 	// The material being built.
 	UMaterialInterface* MaterialInterface{};
 	
@@ -72,6 +61,9 @@ public:
 	// Optional. Specifies the target MaterialInsights to populate, if provided.
 	FMaterialInsights* Insights{};
 
+	// Preshader intermediate representation
+	FMaterialIRPreshader* Preshader{};
+
 	// The destination compilation output whose state to populate, based on the values being analyzed.
 	FMaterialCompilationOutput* CompilationOutput{};
 		
@@ -81,9 +73,6 @@ public:
 	// List of enabled shader environment defines.
 	TSet<FName> EnvironmentDefines;
 
-	// Preshader intermediate representation
-	FMaterialIRPreshader Preshader;
-
 	// Translator specific structure during generation of VT stacks. This is later converted to FMaterialVirtualTextureStack entries in FUniformExpressionSet.
 	// The equivalent in the old translator is FHLSLMaterialTranslator::FMaterialVTStackEntry.
 	struct FVTStackEntry
diff --git a/Engine/Source/Runtime/Engine/Public/MaterialShared.h b/Engine/Source/Runtime/Engine/Public/MaterialShared.h
--- a/Engine/Source/Runtime/Engine/Public/MaterialShared.h
+++ b/Engine/Source/Runtime/Engine/Public/MaterialShared.h
@@ -774,6 +774,10 @@ public:
 	// Returns size of uniform preshader buffer in float4s.
 	uint32 GetUniformBufferSize() const { return UniformPreshaderBufferSize; }
 
+	// Returns the uniform preshader data.
+	UE::Shader::FPreshaderData& GetPreshaderData() { return UniformPreshaderData; }
+	const UE::Shader::FPreshaderData& GetPreshaderData() const { return UniformPreshaderData; }
+
 	// Returns size of temporary preshader buffer in float4s.  Temporary space should be added to the CPU side uniform buffer.
 	uint32 GetTemporaryBufferSize() const { return UniformPreshaderData.Preshader2TemporarySize; }
 	
@@ -833,7 +837,6 @@ protected:
 	friend class FMaterialRenderProxy;
 	friend class FMaterialVirtualTextureStack;
 	friend class FDebugUniformExpressionSet;
-	friend struct FMaterialIRValueAnalyzer;
 	friend struct MIR::FDebugDumpIRState;
 	LAYOUT_FIELD(TMemoryImageArray<FMaterialUniformParameterEvaluation>, UniformParameterEvaluations);
 	LAYOUT_FIELD(TMemoryImageArray<FMaterialUniformPreshaderHeader>, UniformPreshaders);
