[shaders] wait on job submission tasks during shutdown, fixes crashes that can occur when these are still in flight when shutdown is requested. this manifested more frequently on low-core-count systems (and was easy to reproduce with -corelimit=2 and a 100ms sleep added to the submit job task via r.ShaderCompiler.DebugStallSubmitJob=100)

#Synced-CL 51417161

diff --git a/Engine/Source/Runtime/Engine/Private/ShaderCompiler/ShaderCompilerJobCache.cpp b/Engine/Source/Runtime/Engine/Private/ShaderCompiler/ShaderCompilerJobCache.cpp
--- a/Engine/Source/Runtime/Engine/Private/ShaderCompiler/ShaderCompilerJobCache.cpp
+++ b/Engine/Source/Runtime/Engine/Private/ShaderCompiler/ShaderCompilerJobCache.cpp
@@ -126,7 +126,7 @@ TRACE_DECLARE_ATOMIC_INT_COUNTER(Shaders_JobCacheHits, TEXT("Shaders/JobCache/Hi
 
 TRACE_DECLARE_ATOMIC_INT_COUNTER(Shaders_JobCacheDDCRequests, TEXT("Shaders/JobCache/DDCRequests"));
 TRACE_DECLARE_ATOMIC_INT_COUNTER(Shaders_JobCacheDDCHits, TEXT("Shaders/JobCache/DDCHits"));
-TRACE_DECLARE_ATOMIC_MEMORY_COUNTER(Shaders_JobCacheDDCBytesReceived, TEXT("Shaders/JobCache/DDCBytesRecieved"));
+TRACE_DECLARE_ATOMIC_MEMORY_COUNTER(Shaders_JobCacheDDCBytesReceived, TEXT("Shaders/JobCache/DDCBytesReceived"));
 TRACE_DECLARE_ATOMIC_MEMORY_COUNTER(Shaders_JobCacheDDCBytesSent, TEXT("Shaders/JobCache/DDCBytesSent"));
 
 #if WITH_EDITOR
@@ -630,6 +630,7 @@ public:
 	void SubmitJob(FShaderCommonCompileJob* Job);
 	void SubmitJobs(const TArray<FShaderCommonCompileJobPtr>& InJobs);
 	void TransferJobsToLocalWorker(const TArray<FShaderCommonCompileJobPtr>& InJobs);
+	void OnSubmitJobComplete();
 
 	/** This is an entry point for all jobs that have finished the compilation (whether real or cached). Can be called from multiple threads. */
 	void ProcessFinishedJob(FShaderCommonCompileJob* FinishedJob, EShaderCompileJobStatus Status);
@@ -867,6 +868,9 @@ private:
 	/** List of duplicate jobs */
 	TArray<FShaderCommonCompileJob*> DuplicateJobs;
 
+	std::atomic<uint32> NumPendingSubmitJobs{ 0 };
+	FEventRef AllSubmitJobsCompleteEvent;
+
 	/** Statistics - total number of times we tried to Find() some input hash */
 	uint64 TotalSearchAttempts = 0;
 
@@ -1326,6 +1330,14 @@ void FShaderJobCache::TransferJobsToLocalWorker(const TArray<FShaderCommonCompil
 	}
 }
 
+void FShaderJobCache::OnSubmitJobComplete()
+{
+	if (NumPendingSubmitJobs.fetch_sub(1, std::memory_order_acq_rel) == 1)
+	{
+		AllSubmitJobsCompleteEvent->Trigger();
+	}
+}
+
 void FShaderJobCache::SubmitJob(FShaderCommonCompileJob* Job)
 {
 	// Set thread local so DDC query callback can detect if it's in the same thread, and we need to run through the non-async code path.
@@ -1493,19 +1505,28 @@ void FShaderJobCache::SubmitJobs(const TArray<FShaderCommonCompileJobPtr>& InJob
 			}
 		}
 
+		NumPendingSubmitJobs.fetch_add(InJobs.Num(), std::memory_order_relaxed);
 		for (const FShaderCommonCompileJobPtr& Job : InJobs)
 		{
 			UE::Tasks::ETaskPriority Prio = IsRunningCookCommandlet() ? UE::Tasks::ETaskPriority::Normal : UE::Tasks::ETaskPriority::BackgroundNormal;
 			UE::Tasks::Launch(UE_SOURCE_LOCATION, [Job, this]()
 				{
 					TRACE_CPUPROFILER_EVENT_SCOPE(ShaderJobTask);
-					double TimeStart = FPlatformTime::Seconds();
 
 					if (GShaderCompilerDebugStallSubmitJob > 0)
 					{
 						FPlatformProcess::Sleep(GShaderCompilerDebugStallSubmitJob * 0.001f);
 					}
 
+					if (IsEngineExitRequested())
+					{
+						// we do not need to unlink the job here as normal shutdown flow of anything referencing pending shader jobs will call the cancellation path;
+						// we just need to ensure that any tasks that still execute after this cancellation occurs decrement the atomic counter and possibly signal the
+						// event that we wait on in the destructor (to ensure none of these tasks are outstanding before we start unloading modules)
+						OnSubmitJobComplete();
+						return;
+					}
+
 					const bool bSubmitJob = PreprocessShader(Job);
 					Job->UpdateInputHash();
 					Job->UpdateStatus(EShaderCompileJobStatus::Ready);
@@ -1521,6 +1542,8 @@ void FShaderJobCache::SubmitJobs(const TArray<FShaderCommonCompileJobPtr>& InJob
 						JobLock.WriteUnlock();
 						ProcessFinishedJob(Job, EShaderCompileJobStatus::Skipped);
 					}
+
+					OnSubmitJobComplete();
 				}, Prio);
 		}
 	}
@@ -2029,6 +2052,7 @@ uint64 FShaderJobCache::GetCurrentMemoryBudget() const
 
 FShaderJobCache::FShaderJobCache(FCriticalSection& InCompileQueueSection)
 	: CompileQueueSection(InCompileQueueSection)
+	, AllSubmitJobsCompleteEvent(EEventMode::ManualReset)
 {
 	FMemory::Memzero(PendingJobsHead);
 	FMemory::Memzero(NumPendingJobs);
@@ -2044,6 +2068,12 @@ FShaderJobCache::FShaderJobCache(FCriticalSection& InCompileQueueSection)
 
 FShaderJobCache::~FShaderJobCache()
 {
+	AllSubmitJobsCompleteEvent->Reset();
+	if (NumPendingSubmitJobs.load(std::memory_order_acquire) > 0)
+	{
+		AllSubmitJobsCompleteEvent->Wait();
+	}
+
 	for (TMap<FJobOutputHash, FStoredOutput*>::TIterator Iter(Outputs); Iter; ++Iter)
 	{
 		delete Iter.Value();
