Pulling LP's change 51235544 into 40.00 and correctly tying UMeshComponent->verse::material replication links to the mesh components begin/end play lifetime. This is required to correctly clean up material replication handles on the server at the point when replicated mesh components stop replicating

diff --git a/Engine/Source/Runtime/Engine/Classes/Components/MeshComponent.h b/Engine/Source/Runtime/Engine/Classes/Components/MeshComponent.h
--- a/Engine/Source/Runtime/Engine/Classes/Components/MeshComponent.h
+++ b/Engine/Source/Runtime/Engine/Classes/Components/MeshComponent.h
@@ -109,6 +109,11 @@ public:
 	ENGINE_API virtual void PostEditChangeChainProperty(FPropertyChangedChainEvent& PropertyChangedEvent) override;
 #endif
 
+	//~ Begin UActorComponent Interface
+	ENGINE_API virtual void BeginPlay() override;
+	ENGINE_API virtual void EndPlay(const EEndPlayReason::Type EndPlayReason) override;
+	//~ End UActorComponent Interface
+
 	//~ Begin UPrimitiveComponent Interface
 	ENGINE_API virtual int32 GetNumMaterials() const override;
 	ENGINE_API virtual UMaterialInterface* GetMaterial(int32 ElementIndex) const override;
diff --git a/Engine/Source/Runtime/Engine/Private/Components/MeshComponent.cpp b/Engine/Source/Runtime/Engine/Private/Components/MeshComponent.cpp
--- a/Engine/Source/Runtime/Engine/Private/Components/MeshComponent.cpp
+++ b/Engine/Source/Runtime/Engine/Private/Components/MeshComponent.cpp
@@ -87,17 +87,23 @@ void UMeshComponent::SetMaterial(int32 ElementIndex, UMaterialInterface* Materia
 				}
 			}	
 
-			if (UMaterialInterface* PreviousMaterial = OverrideMaterials[ElementIndex].Get())
+			if (HasBegunPlay())
 			{
-				PreviousMaterial->OnRemovedAsOverride(GetOwner());
+				if (UMaterialInterface* PreviousMaterial = OverrideMaterials[ElementIndex].Get())
+				{
+					PreviousMaterial->OnRemovedAsOverride(GetOwner());
+				}
 			}
 
 			// Set the material and invalidate things
 			OverrideMaterials[ElementIndex] = Material;
 
-			if (Material)
+			if (HasBegunPlay())
 			{
-				Material->OnAssignedAsOverride(GetOwner());
+				if (Material)
+				{
+					Material->OnAssignedAsOverride(GetOwner());
+				}
 			}
 
 			// Precache PSOs again
@@ -190,11 +196,14 @@ void UMeshComponent::CleanUpOverrideMaterials()
 		//Remove the override material id that are superior to the static mesh materials number
 		int32 RemoveCount = NumOverrideMaterials - NumMaterials;
 
-		for (int32 MatIndex = NumMaterials; MatIndex < NumOverrideMaterials; MatIndex++)
+		if (HasBegunPlay())
 		{
-			if (UMaterialInterface* MatInterface = OverrideMaterials[MatIndex].Get())
+			for (int32 MatIndex = NumMaterials; MatIndex < NumOverrideMaterials; MatIndex++)
 			{
-				MatInterface->OnRemovedAsOverride(GetOwner());
+				if (UMaterialInterface* MatInterface = OverrideMaterials[MatIndex].Get())
+				{
+					MatInterface->OnRemovedAsOverride(GetOwner());
+				}
 			}
 		}
 
@@ -239,11 +248,14 @@ void UMeshComponent::EmptyOverrideMaterials()
 	bool bRefresh = false;
 	if (OverrideMaterials.Num())
 	{
-		for (int32 MatIndex = 0; MatIndex < OverrideMaterials.Num(); MatIndex++)
+		if (HasBegunPlay())
 		{
-			if (UMaterialInterface* MatInterface = OverrideMaterials[MatIndex].Get())
+			for (int32 MatIndex = 0; MatIndex < OverrideMaterials.Num(); MatIndex++)
 			{
-				MatInterface->OnRemovedAsOverride(GetOwner());
+				if (UMaterialInterface* MatInterface = OverrideMaterials[MatIndex].Get())
+				{
+					MatInterface->OnRemovedAsOverride(GetOwner());
+				}
 			}
 		}
 
@@ -273,6 +285,32 @@ bool UMeshComponent::HasOverrideMaterials()
 	return OverrideMaterials.Num() > 0;
 }
 
+void UMeshComponent::BeginPlay()
+{
+	Super::BeginPlay();
+
+	for (int32 MatIndex = 0; MatIndex < OverrideMaterials.Num(); MatIndex++)
+	{
+		if (UMaterialInterface* MatInterface = OverrideMaterials[MatIndex].Get())
+		{
+			MatInterface->OnAssignedAsOverride(GetOwner());
+		}
+	}
+}
+
+void UMeshComponent::EndPlay(const EEndPlayReason::Type EndPlayReason)
+{
+	for (int32 MatIndex = 0; MatIndex < OverrideMaterials.Num(); MatIndex++)
+	{
+		if (UMaterialInterface* MatInterface = OverrideMaterials[MatIndex].Get())
+		{
+			MatInterface->OnRemovedAsOverride(GetOwner());
+		}
+	}
+
+	Super::EndPlay(EndPlayReason);
+}
+
 int32 UMeshComponent::GetNumMaterials() const
 {
 	return 0;
@@ -499,8 +537,6 @@ void UMeshComponent::SetParameterValueOnMaterials(const FName ParameterName
 			bDynamicSlotOverlayMaterialCreated = true;
 
 			UMaterialInstanceDynamic* DynamicMaterial = UMaterialInstanceDynamic::Create(MaterialInterface, this);
-			DynamicMaterial->OnAssignedAsOverride(GetOwner());
-
 			MaterialSlotsOverlayMaterial[MaterialIndex] = DynamicMaterial;
 
 			return DynamicMaterial;
@@ -611,11 +647,6 @@ void UMeshComponent::BeginDestroy()
 {
 	for (int32 MatIndex = 0; MatIndex < OverrideMaterials.Num(); MatIndex++)
 	{
-		if (UMaterialInterface* MatInterface = OverrideMaterials[MatIndex].Get())
-		{
-			MatInterface->OnRemovedAsOverride(GetOwner());
-		}
-
 		OverrideMaterials[MatIndex] = nullptr;
 	}
 
