[StaticMesh]
- Fix potential leak with static mesh render data when creating and saving static meshes without opening them.

#Synced-CL 51263423

diff --git a/Engine/Source/Runtime/Engine/Private/StaticMesh.cpp b/Engine/Source/Runtime/Engine/Private/StaticMesh.cpp
--- a/Engine/Source/Runtime/Engine/Private/StaticMesh.cpp
+++ b/Engine/Source/Runtime/Engine/Private/StaticMesh.cpp
@@ -7190,6 +7190,7 @@ void UStaticMesh::CacheDerivedData()
 	check(RunningPlatform);
 	const FStaticMeshLODSettings& LODSettings = RunningPlatform->GetStaticMeshLODSettings();
 
+	TUniquePtr<class FStaticMeshRenderData> OldMeshRenderData;
 	if (GetRenderData())
 	{
 		// This is the responsibility of the caller to ensure this has been called
@@ -7208,10 +7209,43 @@ void UStaticMesh::CacheDerivedData()
 				GCardRepresentationAsyncQueue->CancelBuild(this);
 			}
 		}
+
+		if (AreRenderingResourcesInitialized())
+		{
+			PRAGMA_DISABLE_DEPRECATION_WARNINGS
+			OldMeshRenderData = MoveTemp(RenderData);
+			PRAGMA_ENABLE_DEPRECATION_WARNINGS
+		}
 	}
 
 	SetRenderData(MakeUnique<FStaticMeshRenderData>());
 	GetRenderData()->Cache(RunningPlatform, this, LODSettings);
+
+	// If the render data was initialized on the render thread, we need to clean it up
+	if (OldMeshRenderData)
+	{
+		// on the game thread, this might use batching but on any other thread
+		// it will send render commands to release the resource
+		OldMeshRenderData->ReleaseResources();
+
+		if (IsInGameThread())
+		{
+			TRACE_CPUPROFILER_EVENT_SCOPE(OldMeshRenderDataFence);
+			// On the game thread, it's currently hard to know if batching is being used
+			// so we need to use a fence and wait.
+			ReleaseResourcesFence.BeginFence();
+			ReleaseResourcesFence.Wait();
+		}
+		else
+		{
+			// when batching isn't used, we can avoid the wait by just queueing the deletion
+			// on the render thread.
+			ENQUEUE_RENDER_COMMAND(DestroyOldMeshRenderData)(
+				[OldMeshRenderData = MoveTemp(OldMeshRenderData)](FRHICommandListBase&) mutable
+				{
+				});
+		}
+	}
 }
 
 void UStaticMesh::PrepareDerivedDataForActiveTargetPlatforms()
