[HWRT] Small improvements to FRayTracingGeometryManager.
- replaced a couple of code paths setting bEvicted and bAlwaysResident with asserts since they should already be in the correct state.
- expanded CheckIntegrity() to also verify that EvictableGeometries and AlwaysResidentGeometries don't have stale entries.
- made some functions const since they don't modify the state of the system.
- changed the type of CurrentFirstLODIdx to int32 so it properly stores INDEX_NONE.
- also fix only one on-screen message showing at a time.

#Synced-CL 51545719

diff --git a/Engine/Source/Runtime/Engine/Private/Rendering/RayTracingGeometryManager.cpp b/Engine/Source/Runtime/Engine/Private/Rendering/RayTracingGeometryManager.cpp
--- a/Engine/Source/Runtime/Engine/Private/Rendering/RayTracingGeometryManager.cpp
+++ b/Engine/Source/Runtime/Engine/Private/Rendering/RayTracingGeometryManager.cpp
@@ -221,7 +221,7 @@ FRayTracingGeometryManager::FBuildRequest& FRayTracingGeometryManager::GetBuildR
 	return GeometryBuildRequests[RequestIndex];
 }
 
-FRayTracingGeometryManager::FRegisteredGeometry& FRayTracingGeometryManager::GetRegisteredGeometry(RayTracing::FGeometryHandle Handle)
+uint32 FRayTracingGeometryManager::GetRegisteredGeometryIndex(RayTracing::FGeometryHandle Handle) const
 {
 	checkf(!Handle.IsNull(), TEXT("Must provide valid RayTracing::FGeometryHandle."));
 
@@ -231,10 +231,10 @@ FRayTracingGeometryManager::FRegisteredGeometry& FRayTracingGeometryManager::Get
 	checkf(RegisteredGeometries.IsValidIndex(GeometryIndex), TEXT("Stale RayTracing::FGeometryHandle detected - invalid index %u."), GeometryIndex);
 	checkf(GeometryVersion == RegisteredGeometriesVersions[GeometryIndex], TEXT("Stale RayTracing::FGeometryHandle detected - invalid version %u (expected %u)."), GeometryVersion, RegisteredGeometriesVersions[GeometryIndex]);
 
-	return RegisteredGeometries[GeometryIndex];
+	return GeometryIndex;
 }
 
-FRayTracingGeometryManager::FRayTracingGeometryGroup& FRayTracingGeometryManager::GetRegisteredGroup(RayTracing::FGeometryGroupHandle Handle)
+uint32 FRayTracingGeometryManager::GetRegisteredGroupIndex(RayTracing::FGeometryGroupHandle Handle) const
 {
 	checkf(!Handle.IsNull(), TEXT("Must provide valid RayTracing::FGeometryGroupHandle."));
 
@@ -244,7 +244,27 @@ FRayTracingGeometryManager::FRayTracingGeometryGroup& FRayTracingGeometryManager
 	checkf(RegisteredGroups.IsValidIndex(GroupIndex), TEXT("Stale RayTracing::FGeometryGroupHandle detected - invalid index %u."), GroupIndex);
 	checkf(GroupVersion == RegisteredGroupsVersions[GroupIndex], TEXT("Stale RayTracing::FGeometryGroupHandle detected - invalid version %u (expected %u)."), GroupVersion, RegisteredGroupsVersions[GroupIndex]);
 
-	return RegisteredGroups[GroupIndex];
+	return GroupIndex;
+}
+
+FRayTracingGeometryManager::FRegisteredGeometry& FRayTracingGeometryManager::GetRegisteredGeometry(RayTracing::FGeometryHandle Handle)
+{
+	return RegisteredGeometries[GetRegisteredGeometryIndex(Handle)];
+}
+
+const FRayTracingGeometryManager::FRegisteredGeometry& FRayTracingGeometryManager::GetRegisteredGeometry(RayTracing::FGeometryHandle Handle) const
+{
+	return RegisteredGeometries[GetRegisteredGeometryIndex(Handle)];
+}
+
+FRayTracingGeometryManager::FRayTracingGeometryGroup& FRayTracingGeometryManager::GetRegisteredGroup(RayTracing::FGeometryGroupHandle Handle)
+{
+	return RegisteredGroups[GetRegisteredGroupIndex(Handle)];
+}
+
+const FRayTracingGeometryManager::FRayTracingGeometryGroup& FRayTracingGeometryManager::GetRegisteredGroup(RayTracing::FGeometryGroupHandle Handle) const
+{
+	return RegisteredGroups[GetRegisteredGroupIndex(Handle)];
 }
 
 static float GetInitialBuildPriority(ERTAccelerationStructureBuildPriority InBuildPriority)
@@ -835,7 +855,7 @@ void FRayTracingGeometryManager::UpdateReferenceBasedResidency(FRHICommandList&
 		const FRayTracingGeometryGroup& Group = GetRegisteredGroup(GroupHandle);
 		const int32 NumGeometriesInGroup = Group.GeometryHandles.Num();
 
-		for (uint8 LODIndex = Group.CurrentFirstLODIdx; LODIndex < NumGeometriesInGroup; ++LODIndex)
+		for (int32 LODIndex = Group.CurrentFirstLODIdx; LODIndex < NumGeometriesInGroup; ++LODIndex)
 		{
 			RayTracing::FGeometryHandle GeometryHandle = Group.GeometryHandles[LODIndex];
 
@@ -965,10 +985,12 @@ void FRayTracingGeometryManager::UpdateReferenceBasedResidency(FRHICommandList&
 #if !UE_BUILD_SHIPPING
 	if (GRayTracingShowOnScreenWarnings)
 	{
+		uint64 MsgId = Murmur64({ (uint64)this, 0x51bc352fe2ed753a }); // To prevent collisions with other on-screen messages.
+
 		const int ThresholdPercentage = FMath::Clamp(CVarRayTracingAlwaysResidentWarningPercentage.GetValueOnRenderThread(), 0, 100);
 		if (TotalAlwaysResidentSize > ResidentGeometryMemoryPoolSize * ThresholdPercentage / 100.0f)
 		{
-			GEngine->AddOnScreenDebugMessage(uint64(this), 1.f, FColor::Red,
+			GEngine->AddOnScreenDebugMessage(MsgId++, 1.f, FColor::Red,
 				*FString::Printf(TEXT("RAY TRACING GEOMETRY - ALWAYS RESIDENT MEMORY EXCEEDS %d%% OF THE BUDGET (%s / %s)"),
 					ThresholdPercentage,
 					*FText::AsMemory(TotalAlwaysResidentSize).ToString(),
@@ -977,7 +999,7 @@ void FRayTracingGeometryManager::UpdateReferenceBasedResidency(FRHICommandList&
 
 		if (RequestedSize > ResidentGeometryMemoryPoolSize)
 		{
-			GEngine->AddOnScreenDebugMessage(uint64(this), 1.f, FColor::Yellow,
+			GEngine->AddOnScreenDebugMessage(MsgId++, 1.f, FColor::Yellow,
 				*FString::Printf(TEXT("RAY TRACING GEOMETRY - REQUESTED MEMORY OVER BUDGET %s / %s"),
 					*FText::AsMemory(RequestedSize).ToString(),
 					*FText::AsMemory(ResidentGeometryMemoryPoolSize).ToString()));
@@ -1039,7 +1061,7 @@ void FRayTracingGeometryManager::RefreshAlwaysResidentRayTracingGeometries(FRHIC
 
 			if (GroupHandle.IsNull())
 			{
-				RegisteredGeometry.bAlwaysResident = false;
+				checkf(!RegisteredGeometry.bAlwaysResident, TEXT("Ray tracing geometry unexpectely marked as always resident even though it is not assigned to a group."));
 				continue;
 			}
 
@@ -1075,14 +1097,17 @@ void FRayTracingGeometryManager::RefreshAlwaysResidentRayTracingGeometries(FRHIC
 }
 
 #if DO_CHECK
-void FRayTracingGeometryManager::CheckIntegrity()
+void FRayTracingGeometryManager::CheckIntegrity() const
 {
 	const bool bIsRayTracingEnabled = IsRayTracingEnabled();
 	const bool bUsingReferenceBasedResidency = IsRayTracingUsingReferenceBasedResidency();
 
 	if (GRayTracingTestCheckIntegrity)
 	{
-		for (FRegisteredGeometry& RegisteredGeometry : RegisteredGeometries)
+		uint32 NumAlwaysResidentGeometries = 0;
+		uint32 NumEvictableGeometries = 0;
+
+		for (const FRegisteredGeometry& RegisteredGeometry : RegisteredGeometries)
 		{
 			const RayTracing::FGeometryHandle GeometryHandle = RegisteredGeometry.Geometry->RayTracingGeometryHandle;
 			const RayTracing::FGeometryGroupHandle GroupHandle = RegisteredGeometry.Geometry->GroupHandle;
@@ -1100,14 +1125,32 @@ void FRayTracingGeometryManager::CheckIntegrity()
 			}
 
 			checkf(RegisteredGeometry.bEvicted == RegisteredGeometry.Geometry->IsEvicted(), TEXT("Cached bEvicted flag in FRegisteredGeometry is stale"));
-			checkf(RegisteredGeometry.bAlwaysResident == bAlwaysResident, TEXT("Cached bAlwaysResident flag in FRegisteredGeometry is stale"));
-			checkf(RegisteredGeometry.bAlwaysResident == AlwaysResidentGeometries.Contains(GeometryHandle), TEXT("Geometry with bAlwaysResident flag set should be in the AlwaysResidentGeometries set."));
+			checkf(bAlwaysResident == RegisteredGeometry.bAlwaysResident, TEXT("Cached bAlwaysResident flag in FRegisteredGeometry is stale"));
+			checkf(bAlwaysResident == AlwaysResidentGeometries.Contains(GeometryHandle), TEXT("Geometry with bAlwaysResident flag set should be in the AlwaysResidentGeometries set."));
+
+			if (bAlwaysResident)
+			{
+				++NumAlwaysResidentGeometries;
+			}
+
+			if (RegisteredGeometry.Geometry->IsValid() && !RegisteredGeometry.bEvicted && !RegisteredGeometry.bAlwaysResident)
+			{
+				check(EvictableGeometries.Contains(GeometryHandle));
+				++NumEvictableGeometries;
+			}
+			else
+			{
+				check(!EvictableGeometries.Contains(GeometryHandle));
+			}
 		}
 
+		check(NumAlwaysResidentGeometries == AlwaysResidentGeometries.Num());
+		check(NumEvictableGeometries == EvictableGeometries.Num());
+
 		if (!bIsRayTracingEnabled)
 		{
 			// check that everything is evicted
-			for (FRegisteredGeometry& RegisteredGeometry : RegisteredGeometries)
+			for (const FRegisteredGeometry& RegisteredGeometry : RegisteredGeometries)
 			{
 				checkf(RegisteredGeometry.Geometry->IsEvicted() || RegisteredGeometry.Geometry->GetRHI() == nullptr, TEXT("Ray tracing geometry should be evicted when ray tracing is disabled."));
 			}
@@ -1115,7 +1158,7 @@ void FRayTracingGeometryManager::CheckIntegrity()
 		else if (!bUsingReferenceBasedResidency)
 		{
 			// check that all geometries are resident
-			for (FRegisteredGeometry& RegisteredGeometry : RegisteredGeometries)
+			for (const FRegisteredGeometry& RegisteredGeometry : RegisteredGeometries)
 			{
 				checkf(!RegisteredGeometry.Geometry->IsEvicted(), TEXT("Ray tracing geometry should not be evicted when ray tracing is enabled."));
 			}
@@ -1376,7 +1419,7 @@ void FRayTracingGeometryManager::EvictGeometry(FRHICommandListBase& RHICmdList,
 	StreamOutGeometry(Batcher, RegisteredGeometry);
 
 	RegisteredGeometry.Geometry->Evict();
-	RegisteredGeometry.bEvicted = true;
+	check(RegisteredGeometry.bEvicted);
 }
 
 void FRayTracingGeometryManager::BoostPriority(RayTracing::FBuildRequestHandle InRequestHandle, float InBoostValue)
@@ -1703,12 +1746,12 @@ bool FRayTracingGeometryManager::IsGeometryGroupReferenced(RayTracing::FGeometry
 #endif // DO_CHECK
 
 #if !UE_BUILD_SHIPPING
-bool FRayTracingGeometryManager::DumpRayTracingGeometries(TArrayView<RayTracing::FGeometryHandle> InGeometries, const FString& Filename)
+bool FRayTracingGeometryManager::DumpRayTracingGeometries(TArrayView<RayTracing::FGeometryHandle> InGeometries, const FString& Filename) const
 {
 	Algo::Sort(InGeometries, [this](RayTracing::FGeometryHandle& LHSHandle, RayTracing::FGeometryHandle& RHSHandle)
 		{
-			FRegisteredGeometry& LHS = GetRegisteredGeometry(LHSHandle);
-			FRegisteredGeometry& RHS = GetRegisteredGeometry(RHSHandle);
+			const FRegisteredGeometry& LHS = GetRegisteredGeometry(LHSHandle);
+			const FRegisteredGeometry& RHS = GetRegisteredGeometry(RHSHandle);
 
 			return LHS.Size > RHS.Size;
 		});
@@ -1724,8 +1767,8 @@ bool FRayTracingGeometryManager::DumpRayTracingGeometries(TArrayView<RayTracing:
 
 	for (RayTracing::FGeometryHandle GeometryHandle : InGeometries)
 	{
-		FRegisteredGeometry& RegisteredGeometry = GetRegisteredGeometry(GeometryHandle);
-		FRayTracingGeometryInitializer Initializer = RegisteredGeometry.Geometry->GetInitializer();
+		const FRegisteredGeometry& RegisteredGeometry = GetRegisteredGeometry(GeometryHandle);
+		const FRayTracingGeometryInitializer Initializer = RegisteredGeometry.Geometry->GetInitializer();
 
 		const FString Row = FString::Printf(TEXT("%s,%.3f,%d,%d,%d\n"),
 			!Initializer.DebugName.IsNone() ? *Initializer.DebugName.ToString() : TEXT("*UNKNOWN*"),
@@ -1739,7 +1782,7 @@ bool FRayTracingGeometryManager::DumpRayTracingGeometries(TArrayView<RayTracing:
 	return true;
 }
 
-void FRayTracingGeometryManager::DumpUnreferencedAlwaysResidentRayTracingGeometries()
+void FRayTracingGeometryManager::DumpUnreferencedAlwaysResidentRayTracingGeometries() const
 {
 	TArray<RayTracing::FGeometryHandle> UnreferencedAlwaysResidentGeometries;
 
diff --git a/Engine/Source/Runtime/Engine/Public/Rendering/RayTracingGeometryManager.h b/Engine/Source/Runtime/Engine/Public/Rendering/RayTracingGeometryManager.h
--- a/Engine/Source/Runtime/Engine/Public/Rendering/RayTracingGeometryManager.h
+++ b/Engine/Source/Runtime/Engine/Public/Rendering/RayTracingGeometryManager.h
@@ -101,7 +101,7 @@ private:
 
 		TSet<FPrimitiveSceneProxy*> ProxiesWithCachedRayTracingState;
 
-		uint8 CurrentFirstLODIdx = INDEX_NONE;
+		int32 CurrentFirstLODIdx = INDEX_NONE;
 
 		// Due to the way we batch release FRenderResource and SceneProxies, 
 		// ReleaseRayTracingGeometryGroup(...) can end up being called before all FRayTracingGeometry and SceneProxies are actually released.
@@ -183,17 +183,23 @@ private:
 	void RefreshAlwaysResidentRayTracingGeometries(FRHICommandList& RHICmdList);
 
 #if DO_CHECK
-	void CheckIntegrity();
+	void CheckIntegrity() const;
 #endif // DO_CHECK
 
 #if !UE_BUILD_SHIPPING
-	bool DumpRayTracingGeometries(TArrayView<RayTracing::FGeometryHandle> Geometries, const FString& Filename);
-	void DumpUnreferencedAlwaysResidentRayTracingGeometries();
+	bool DumpRayTracingGeometries(TArrayView<RayTracing::FGeometryHandle> Geometries, const FString& Filename) const;
+	void DumpUnreferencedAlwaysResidentRayTracingGeometries() const;
 #endif // !UE_BUILD_SHIPPING
 
 	FBuildRequest& GetBuildRequest(RayTracing::FBuildRequestHandle Handle);
+
+	uint32 GetRegisteredGeometryIndex(RayTracing::FGeometryHandle Handle) const;
+	uint32 GetRegisteredGroupIndex(RayTracing::FGeometryGroupHandle Handle) const;
+
 	FRegisteredGeometry& GetRegisteredGeometry(RayTracing::FGeometryHandle Handle);
+	const FRegisteredGeometry& GetRegisteredGeometry(RayTracing::FGeometryHandle Handle) const;
 	FRayTracingGeometryGroup& GetRegisteredGroup(RayTracing::FGeometryGroupHandle Handle);
+	const FRayTracingGeometryGroup& GetRegisteredGroup(RayTracing::FGeometryGroupHandle Handle) const;
 
 	static bool IsAlwaysResidentGeometry(const FRayTracingGeometry* InGeometry, const FRayTracingGeometryGroup& Group);
 
