Removed legacy streaming virtual texture upload path.
#rnx
#Synced-CL 51537754

diff --git a/Engine/Source/Runtime/Engine/Private/VT/VirtualTextureUploadCache.cpp b/Engine/Source/Runtime/Engine/Private/VT/VirtualTextureUploadCache.cpp
--- a/Engine/Source/Runtime/Engine/Private/VT/VirtualTextureUploadCache.cpp
+++ b/Engine/Source/Runtime/Engine/Private/VT/VirtualTextureUploadCache.cpp
@@ -10,13 +10,6 @@
 DECLARE_MEMORY_STAT_POOL(TEXT("Total GPU Upload Memory"), STAT_TotalGPUUploadSize, STATGROUP_VirtualTextureMemory, FPlatformMemory::MCR_GPU);
 DECLARE_MEMORY_STAT(TEXT("Total CPU Upload Memory"), STAT_TotalCPUUploadSize, STATGROUP_VirtualTextureMemory);
 
-static TAutoConsoleVariable<int32> CVarVTUploadUseLegacyPath(
-	TEXT("r.VT.UploadUseLegacyPath"),
-	0,
-	TEXT("Use the legacy virtual texture upload path which locks staging textures."),
-	ECVF_ReadOnly
-);
-
 static TAutoConsoleVariable<int32> CVarVTUploadMemoryPageSize(
 	TEXT("r.VT.UploadMemoryPageSize"),
 	4,
@@ -248,18 +241,8 @@ void FVTUploadTileAllocator::FStagingBuffer::Release(FRHICommandListBase* RHICmd
 
 FVirtualTextureUploadCache::FVirtualTextureUploadCache()
 {
-	if (CVarVTUploadUseLegacyPath.GetValueOnGameThread())
-	{
-		UploadType = EVTUploadType::StagingTexture;
-	}
-	else if (GRHISupportsDirectGPUMemoryLock && GRHISupportsUpdateFromBufferTexture)
-	{
-		UploadType = EVTUploadType::PersistentBuffer;
-	}
-	else
-	{
-		UploadType = EVTUploadType::StagingCopy;
-	}
+	const bool bSupportPersistentBuffer = (GRHISupportsDirectGPUMemoryLock && GRHISupportsUpdateFromBufferTexture);
+	UploadType = bSupportPersistentBuffer ? EVTUploadType::PersistentBuffer : UploadType = EVTUploadType::StagingCopy;
 }
 
 int32 FVirtualTextureUploadCache::GetOrCreatePoolIndex(EPixelFormat InFormat, uint32 InTileSize)
@@ -281,176 +264,13 @@ int32 FVirtualTextureUploadCache::GetOrCreatePoolIndex(EPixelFormat InFormat, ui
 	return PoolIndex;
 }
 
-void FVirtualTextureUploadCache::Finalize(FRDGBuilder& GraphBuilder)
-{
-	switch (UploadType)
-	{
-	case EVTUploadType::StagingTexture:
-		FinalizeWithLegacyCopyTexture(GraphBuilder);
-		break;
-	case EVTUploadType::StagingCopy:
-	case EVTUploadType::PersistentBuffer:
-		FinalizeWithUpdateTexture(GraphBuilder);
-		break;
-	}
-}
-
 /** Parameter struct for tracking target texture resource states. */
 BEGIN_SHADER_PARAMETER_STRUCT(FVirtualTextureUploadCacheParameters, )
 	RDG_TEXTURE_ACCESS_ARRAY(TextureArray)
 END_SHADER_PARAMETER_STRUCT()
 
-void FVirtualTextureUploadCache::FinalizeWithLegacyCopyTexture(FRDGBuilder& GraphBuilder)
-{
-	for (int PoolIndex = 0; PoolIndex < Pools.Num(); ++PoolIndex)
-	{
-		FPoolEntry& PoolEntry = Pools[PoolIndex];
-		const uint32 BatchCount = PoolEntry.PendingSubmit.Num();
-		if (BatchCount == 0)
-		{
-			continue;
-		}
-
-		// Create/Resize the pool staging buffer texture.
-		const EPixelFormat Format = PoolEntry.Format;
-		const uint32 TileSize = PoolEntry.TileSize;
-		const uint32 TextureIndex = PoolEntry.BatchTextureIndex;
-		PoolEntry.BatchTextureIndex = (PoolEntry.BatchTextureIndex + 1u) % FPoolEntry::NUM_STAGING_TEXTURES;
-		FStagingTexture& StagingTexture = PoolEntry.StagingTexture[TextureIndex];
-
-		// On some platforms the staging texture create/lock behavior will depend on whether we are running with RHI threading
-		const bool bIsCpuWritable = !IsRunningRHIInSeparateThread();
-
-		if (BatchCount > StagingTexture.BatchCapacity || BatchCount * 2 <= StagingTexture.BatchCapacity || bIsCpuWritable != StagingTexture.bIsCPUWritable)
-		{
-			// Staging texture is vertical stacked in widths of multiples of 4
-			// Smaller widths mean smaller stride which is more efficient for copying
-			// Round up to 4 to reduce likely wasted memory from width not aligning to whatever GPU prefers
-			const uint32 MaxTextureDimension = GetMax2DTextureDimension();
-			const uint32 MaxSizeInTiles = FMath::DivideAndRoundDown(MaxTextureDimension, TileSize);
-			const uint32 MaxCapacity = MaxSizeInTiles * MaxSizeInTiles;
-			check(BatchCount <= MaxCapacity);
-			const uint32 WidthInTiles = FMath::DivideAndRoundUp(FMath::DivideAndRoundUp(BatchCount, MaxSizeInTiles), 4u) * 4;
-			check(WidthInTiles > 0u);
-			const uint32 HeightInTiles = FMath::DivideAndRoundUp(BatchCount, WidthInTiles);
-			check(HeightInTiles > 0u);
-
-			if (StagingTexture.RHITexture)
-			{
-				DEC_MEMORY_STAT_BY(STAT_TotalGPUUploadSize, CalcTextureSize(StagingTexture.RHITexture->GetSizeX(), StagingTexture.RHITexture->GetSizeY(), PoolEntry.Format, 1u));
-			}
-
-			FRHITextureCreateDesc Desc =
-				FRHITextureCreateDesc::Create2D(TEXT("VirtualTexture_UploadCacheStagingTexture"), TileSize * WidthInTiles, TileSize * HeightInTiles, PoolEntry.Format);
-
-			if (bIsCpuWritable)
-			{
-				Desc.AddFlags(ETextureCreateFlags::CPUWritable);
-			}
-
-			StagingTexture.RHITexture = GraphBuilder.RHICmdList.CreateTexture(Desc);
-
-			StagingTexture.WidthInTiles = WidthInTiles;
-			StagingTexture.BatchCapacity = WidthInTiles * HeightInTiles;
-			StagingTexture.bIsCPUWritable = bIsCpuWritable;
-			INC_MEMORY_STAT_BY(STAT_TotalGPUUploadSize, CalcTextureSize(TileSize * WidthInTiles, TileSize * HeightInTiles, PoolEntry.Format, 1u));
-		}
-
-		// Make a copy of required data for RDG pass.
-		FTileEntry* TileEntries = GraphBuilder.AllocPODArray<FTileEntry>(BatchCount);
-		FVTUploadTileBufferExt* UploadBuffers = GraphBuilder.AllocPODArray<FVTUploadTileBufferExt>(BatchCount);
-
-		FVirtualTextureUploadCacheParameters* UploadParameters = GraphBuilder.AllocParameters<FVirtualTextureUploadCacheParameters>();
-
-		for (uint32 Index = 0u; Index < BatchCount; ++Index)
-		{
-			FTileEntry& TileEntry = PoolEntry.PendingSubmit[Index];
-
-			// Common case is that all target textures are the same.
-			// So compare with the previous entry as a quick test for a faster path.
-			if (Index > 1 && TileEntry.PooledRenderTarget == TileEntries[Index - 1].PooledRenderTarget)
-			{
-				TileEntry.Texture = TileEntries[Index - 1].Texture;
-			}
-			else
-			{
-				TileEntry.Texture = GraphBuilder.RegisterExternalTexture(TileEntry.PooledRenderTarget, ERDGTextureFlags::None);
-				UploadParameters->TextureArray.AddUnique(FRDGTextureAccess(TileEntry.Texture, ERHIAccess::CopyDest));
-			}
-
-			TileEntries[Index] = TileEntry;
-			UploadBuffers[Index] = TileAllocator.GetBufferFromHandleExt(TileEntry.TileHandle);
-		}
-
-		// Can clear the pending queue now.
-		PoolEntry.PendingSubmit.Reset();
-
-		GraphBuilder.AddPass(
-			RDG_EVENT_NAME("VirtualTextureUploadLegacy"),
-			UploadParameters,
-			ERDGPassFlags::Copy,
-			[this, TileSize, BatchCount, TileEntries, UploadBuffers, StagingTexture, Format](FRHICommandListImmediate& RHICmdList)
-		{
-			const FPixelFormatInfo& FormatInfo = GPixelFormats[Format];
-			const uint32 BlockBytes = FormatInfo.BlockBytes;
-			const uint32 TileWidthInBlocks = FMath::DivideAndRoundUp(TileSize, (uint32)FormatInfo.BlockSizeX);
-			const uint32 TileHeightInBlocks = FMath::DivideAndRoundUp(TileSize, (uint32)FormatInfo.BlockSizeY);
-
-			FRHILockTextureArgs LockArgs = FRHILockTextureArgs::Lock2D(StagingTexture.RHITexture, 0, EResourceLockMode::RLM_WriteOnly, false, false);
-			FRHILockTextureResult LockResult = RHICmdList.LockTexture(LockArgs);
-			void* BatchMemory = LockResult.Data;
-
-			// Copy all tiles to the staging texture
-			for (uint32 Index = 0u; Index < BatchCount; ++Index)
-			{
-				const FTileEntry& Entry = TileEntries[Index];
-				const FVTUploadTileBufferExt& UploadBuffer = UploadBuffers[Index];
-
-				const uint32_t SrcTileX = Index % StagingTexture.WidthInTiles;
-				const uint32_t SrcTileY = Index / StagingTexture.WidthInTiles;
-		
-				uint8* BatchDst = (uint8*)BatchMemory + TileHeightInBlocks * SrcTileY * LockResult.Stride + TileWidthInBlocks * SrcTileX * BlockBytes;
-				for (uint32 y = 0u; y < TileHeightInBlocks; ++y)
-				{
-					FMemory::Memcpy(
-						BatchDst + y * LockResult.Stride,
-						(uint8*)UploadBuffer.BufferMemory + UploadBuffer.BufferOffset + y * UploadBuffer.Stride,
-						TileWidthInBlocks * BlockBytes);
-				}
-			}
-		
-			RHICmdList.UnlockTexture(LockArgs);
-			RHICmdList.Transition(FRHITransitionInfo(StagingTexture.RHITexture, ERHIAccess::SRVMask, ERHIAccess::CopySrc), ERHITransitionCreateFlags::AllowDecayPipelines);
-		
-			// Upload each tile from staging texture to physical texture
-			for (uint32 Index = 0u; Index < BatchCount; ++Index)
-			{
-				const FTileEntry& Entry = TileEntries[Index];
-				
-				const uint32_t SrcTileX = Index % StagingTexture.WidthInTiles;
-				const uint32_t SrcTileY = Index / StagingTexture.WidthInTiles;
-		
-				const uint32 SkipBorderSize = Entry.SubmitSkipBorderSize;
-				const uint32 SubmitTileSize = TileSize - SkipBorderSize * 2;
-				const FIntVector SourceBoxStart(SrcTileX * TileSize + SkipBorderSize, SrcTileY * TileSize + SkipBorderSize, 0);
-				const FIntVector DestinationBoxStart(Entry.SubmitDestX * SubmitTileSize, Entry.SubmitDestY * SubmitTileSize, 0);
-		
-				FRHICopyTextureInfo CopyInfo;
-				CopyInfo.Size = FIntVector(SubmitTileSize, SubmitTileSize, 1);
-				CopyInfo.SourcePosition = SourceBoxStart;
-				CopyInfo.DestPosition = DestinationBoxStart;
-				RHICmdList.CopyTexture(StagingTexture.RHITexture, Entry.Texture->GetRHI(), CopyInfo);
-			}
-		
-			RHICmdList.Transition(FRHITransitionInfo(StagingTexture.RHITexture, ERHIAccess::CopySrc, ERHIAccess::SRVMask));
-		});
-	}
-}
-
-void FVirtualTextureUploadCache::FinalizeWithUpdateTexture(FRDGBuilder& GraphBuilder)
+void FVirtualTextureUploadCache::Finalize(FRDGBuilder& GraphBuilder)
 {
-	check (UploadType == EVTUploadType::StagingCopy || UploadType == EVTUploadType::PersistentBuffer);
-
 	for (int PoolIndex = 0; PoolIndex < Pools.Num(); ++PoolIndex)
 	{
 		FPoolEntry& PoolEntry = Pools[PoolIndex];
diff --git a/Engine/Source/Runtime/Engine/Private/VT/VirtualTextureUploadCache.h b/Engine/Source/Runtime/Engine/Private/VT/VirtualTextureUploadCache.h
--- a/Engine/Source/Runtime/Engine/Private/VT/VirtualTextureUploadCache.h
+++ b/Engine/Source/Runtime/Engine/Private/VT/VirtualTextureUploadCache.h
@@ -50,8 +50,6 @@ struct FVTUploadTileBufferExt
 /** Enumeration of staging buffer types used for upload to the VT physical texture. */
 enum class EVTUploadType
 {
-	/** Copy image data to a (batched) staging texture for upload. Requires immediate context Lock() but was the best path for D3D11 at some point. */
-	StagingTexture,
 	/** "Directly" upload image data. Internally the RHI will copy to staging memory. */
 	StagingCopy,
 	/** Stream image data directly into a persistent staging buffer for upload. This is the optimal path but is only available on some platforms. */
@@ -183,11 +181,6 @@ private:
 	/** Get a pool for a format and tile size. */
 	int32 GetOrCreatePoolIndex(EPixelFormat InFormat, uint32 InTileSize);
 
-	/** Legacy Finalize path that creates and locks as staging texture. */
-	void FinalizeWithLegacyCopyTexture(FRDGBuilder& GraphBuilder);
-	/** Finalize path that uses UpdateTexture2D. */
-	void FinalizeWithUpdateTexture(FRDGBuilder& GraphBuilder);
-
 	/** Description of a single allocated tile. Carries mutable state as tile moves from uploading to submitting to pending delete. */
 	struct FTileEntry
 	{
