Mark GetObjectName and GetPathName as BlueprintThreadSafe, as they already are

#Synced-CL 51539364

diff --git a/Engine/Source/Runtime/Engine/Classes/Kismet/KismetSystemLibrary.h b/Engine/Source/Runtime/Engine/Classes/Kismet/KismetSystemLibrary.h
--- a/Engine/Source/Runtime/Engine/Classes/Kismet/KismetSystemLibrary.h
+++ b/Engine/Source/Runtime/Engine/Classes/Kismet/KismetSystemLibrary.h
@@ -132,11 +132,11 @@ class UKismetSystemLibrary : public UBlueprintFunctionLibrary
 	static bool IsValidClass(UClass* Class);
 
 	// Returns the actual object name.
-	UFUNCTION(BlueprintPure, Category = "Utilities")
+	UFUNCTION(BlueprintPure, Category = "Utilities", meta = (BlueprintThreadSafe))
 	static ENGINE_API FString GetObjectName(const UObject* Object);
 
 	// Returns the full path to the specified object as a string
-	UFUNCTION(BlueprintPure, Category="Utilities", meta = (DisplayName = "Get Object Path String"))
+	UFUNCTION(BlueprintPure, Category="Utilities", meta = (DisplayName = "Get Object Path String", BlueprintThreadSafe))
 	static ENGINE_API FString GetPathName(const UObject* Object);
 
 	// Returns the full path to the specified object as a Soft Object Path
