Fix cut&paste support for overridable maps/arrays inside structs; improve sub-object shadow handling in overridable containers

Extended the overridable serialization system to properly handle cut&paste (export/import text) for overridable arrays and maps that are nested inside structs. The previous code skipped struct properties during text import, causing the operation to silently fail.

Key changes:

- Property.cpp: Added struct-property handling in ImportSingleProperty ? when the overridable container element is a FStructProperty, recursively processes its sub-properties instead of falling through.
- UnrealExporter.cpp: Generalized the properties export function to operate on UStruct* instead of UClass*, and added a PropertyPrefix parameter to support dot-notation paths for struct-nested properties (e.g., "Struct.PropertyName"). Also adds ShadowIndices tracking for instanced sub-object arrays.
- OverridableManager.cpp: Fixed HasOverrides detection to look up nodes via the full FArchiveSerializedPropertyChain (from context path) rather than a bare property ID. Extended the check to cover Add operations in addition to Remove when determining if an overridable container has pending overrides.
- OverriddenPropertySet.cpp: Replaced hard checkf crashes with ensureMsgf + graceful return in the array/map add and remove implementations to prevent data loss when an index is unexpectedly out of range.
- OverriddenPropertySet.h: Promoted GetOverriddenPropertyOperation(PortFlags, Chain, Property, DataPtr, DefaultValue) to public API (COREUOBJECT_API).
- OverridableSerializationTestUtils.h: Added FOverridableSubStruct and extended FOverridableStruct with nested overridable arrays, maps, and instanced sub-object containers to enable struct-level test coverage.
- OverridableSerializationTest.cpp: Fixed failing GetOverridableState() assertions and restructured the SubObjectOverridableMap test to correctly validate cut&paste and base-class-change propagation.

[RN] Copy& paste support of overridable containers inside struct
#Synced-CL 51456614

diff --git a/Engine/Source/Runtime/CoreUObject/Private/UObject/OverriddenPropertySet.cpp b/Engine/Source/Runtime/CoreUObject/Private/UObject/OverriddenPropertySet.cpp
--- a/Engine/Source/Runtime/CoreUObject/Private/UObject/OverriddenPropertySet.cpp
+++ b/Engine/Source/Runtime/CoreUObject/Private/UObject/OverriddenPropertySet.cpp
@@ -1173,7 +1173,12 @@ void FOverriddenPropertySet::NotifyPropertyChange(FOverriddenPropertyNode* Paren
 
 				auto ArrayAddImpl = [&]()
 				{
-					checkf(ArrayHelper.IsValidIndex(ArrayIndex), TEXT("ArrayAdd change type expected to have an valid index"));
+					if (!ensureMsgf(ArrayHelper.IsValidIndex(ArrayIndex), TEXT("ArrayAdd(%s) operation is expecting a valid index(%i), data might be lost"), *ArrayProperty->GetName(), ArrayIndex))
+					{
+						UE_LOGF(LogOverridableObject, Error, "ArrayAdd(%ls) operation is expecting a valid index(%i), data might be lost", *ArrayProperty->GetName(), ArrayIndex);
+						return;
+					}
+
 					if (UObject* AddedSubObject = TryGetInstancedSubObjectValue(InnerObjectProperty, ArrayHelper.GetElementPtr(ArrayIndex)))
 					{
 						if(SubPropertyNode)
@@ -1196,7 +1201,12 @@ void FOverriddenPropertySet::NotifyPropertyChange(FOverriddenPropertyNode* Paren
 
 				auto ArrayRemoveImpl = [&]()
 				{
-					checkf(PreEditArrayHelper.IsValidIndex(ArrayIndex), TEXT("ArrayRemove change type expected to have an valid index"));
+					if (!ensureMsgf(PreEditArrayHelper.IsValidIndex(ArrayIndex), TEXT("ArrayRemove(%s) operation is expecting a valid index(%i), data might be lost"), *ArrayProperty->GetName(), ArrayIndex))
+					{
+						UE_LOGF(LogOverridableObject, Error, "ArrayRemove(%ls) operation is expecting a valid index(%i), data might be lost", *ArrayProperty->GetName(), ArrayIndex);
+						return;
+					}
+
 					if (UObject* RemovedSubObject = TryGetInstancedSubObjectValue(InnerObjectProperty, PreEditArrayHelper.GetElementPtr(ArrayIndex)))
 					{
 						if(SubPropertyNode)
@@ -1398,7 +1408,7 @@ void FOverriddenPropertySet::NotifyPropertyChange(FOverriddenPropertyNode* Paren
 			}
 			checkf(SavedPreEditMap, TEXT("Expecting the same property as the pre edit flow"));
 			FScriptMapHelper PreEditMapHelper(MapProperty, SavedPreEditMap);
-			// The logical should map directly to the pre edit map internal index as we skipped all of the invalid entries
+			// The logical should map directly to the pre edit map internal index as we skipped all the invalid entries
 			int32 InternalPreEditMapIndex = LogicalMapIndex;
 
 			ON_SCOPE_EXIT
@@ -1443,7 +1453,11 @@ void FOverriddenPropertySet::NotifyPropertyChange(FOverriddenPropertyNode* Paren
 
 			auto MapAddImpl = [&]()
 			{
-				checkf(MapHelper.IsValidIndex(InternalMapIndex), TEXT("ArrayAdd change type expected to have an valid index"));
+				if (!ensureMsgf(MapHelper.IsValidIndex(InternalMapIndex), TEXT("MapAdd(%s) operation is expecting a valid index(%i), data might be lost"), *MapProperty->GetName(), InternalMapIndex))
+				{
+					UE_LOGF(LogOverridableObject, Error, "MapAdd(%ls) operation is expecting a valid index(%i), data might be lost", *MapProperty->GetName(), InternalMapIndex);
+					return;
+				}
 
 				if(SubPropertyNode)
 				{
@@ -1455,7 +1469,11 @@ void FOverriddenPropertySet::NotifyPropertyChange(FOverriddenPropertyNode* Paren
 
 			auto MapRemoveImpl = [&]()
 			{
-				checkf(PreEditMapHelper.IsValidIndex(InternalPreEditMapIndex), TEXT("ArrayRemove change type expected to have an valid index"));
+				if (!ensureMsgf(PreEditMapHelper.IsValidIndex(InternalPreEditMapIndex), TEXT("MapRemove(%s) operation is expecting a valid index(%i), data might be lost"), *MapProperty->GetName(), InternalPreEditMapIndex))
+				{
+					UE_LOGF(LogOverridableObject, Error, "MapRemove(%ls) operation is expecting a valid index(%i), data might be lost", *MapProperty->GetName(), InternalPreEditMapIndex);
+					return;
+				}
 
 				if(SubPropertyNode)
 				{
diff --git a/Engine/Source/Runtime/CoreUObject/Private/UObject/Property.cpp b/Engine/Source/Runtime/CoreUObject/Private/UObject/Property.cpp
--- a/Engine/Source/Runtime/CoreUObject/Private/UObject/Property.cpp
+++ b/Engine/Source/Runtime/CoreUObject/Private/UObject/Property.cpp
@@ -1671,6 +1671,12 @@ const TCHAR* FProperty::ImportSingleProperty( const TCHAR* Str, void* DestData,
 			{
 				ArrayOp = ADO_RemoveIndex;
 			}
+			// Special case for struct properties
+			else if (FStructProperty* const StructProperty = ExactCastField<FStructProperty>(Property))
+			{
+				TArray<FDefinedProperty> StructDefinedProperties;
+				return ImportSingleProperty(Str, Property->ContainerPtrToValuePtr<void>(DestData), StructProperty->Struct, SubobjectOuter, PortFlags, Warn, StructDefinedProperties);
+			}
 		}
 
 		FArrayProperty* const ArrayProperty = ExactCastField<FArrayProperty>(Property);
diff --git a/Engine/Source/Runtime/CoreUObject/Public/UObject/OverriddenPropertySet.h b/Engine/Source/Runtime/CoreUObject/Public/UObject/OverriddenPropertySet.h
--- a/Engine/Source/Runtime/CoreUObject/Public/UObject/OverriddenPropertySet.h
+++ b/Engine/Source/Runtime/CoreUObject/Public/UObject/OverriddenPropertySet.h
@@ -124,6 +124,16 @@ public:
 	 * @return the overridden property operation */
 	COREUOBJECT_API static EOverriddenPropertyOperation GetOverriddenPropertyOperationForPortText(const void* DataPtr, const void* DefaultValue, int32 PortFlags);
 
+	/**
+	 * Retrieve the current overridden property operation to know if it has to be serialized or not from a FArchiveSerializedPropertyChain
+	 * @param PortFlags for the import/export text operation
+	 * @param CurrentPropertyChain the current chain of property leading to the property if any
+	 * @param Property the property about to be serialized, can be null
+	 * @param DataPtr to the memory of that property
+	 * @param DefaultValue memory pointer of that property
+	 * @return the overridden property operation */
+	COREUOBJECT_API static EOverriddenPropertyOperation GetOverriddenPropertyOperation(const int32 PortFlags, const FArchiveSerializedPropertyChain* CurrentPropertyChain, FProperty* Property, const void* DataPtr, const void* DefaultValue);
+
 	/**
 	 * Call during the import text
 	 * @return the current text import property path */
@@ -142,8 +152,6 @@ public:
 
 private:
 
-	static EOverriddenPropertyOperation GetOverriddenPropertyOperation(const int32 PortFlags, const FArchiveSerializedPropertyChain* CurrentPropertyChain, FProperty* Property, const void* DataPtr, const void* DefaultValue);
-
 	friend struct FOverridableSerializationLogicInternalAdapter;
 	
 	static ECapabilities Capabilities;
diff --git a/Engine/Source/Runtime/Engine/Private/UnrealExporter.cpp b/Engine/Source/Runtime/Engine/Private/UnrealExporter.cpp
--- a/Engine/Source/Runtime/Engine/Private/UnrealExporter.cpp
+++ b/Engine/Source/Runtime/Engine/Private/UnrealExporter.cpp
@@ -659,22 +659,21 @@ void ExportProperties_Overrides
 (
 	const FExportObjectInnerContext* Context,
 	FOutputDevice&	Out,
-	UClass*			ObjectClass,
-	uint8*			Object,
+	UStruct*		Struct,
+	uint8*			ValuePtr,
 	int32			Indent,
-	UClass*			DiffClass,
+	UStruct*		DiffStruct,
 	uint8*			Diff,
 	UObject*		Parent,
 	uint32			PortFlags,
-	UObject*		ExportRootScope
+	UObject*		ExportRootScope,
+	FOverriddenPropertySet* OverriddenProperties,
+	const FString&	PropertyPrefix = FString()
 )
 {
-	check(ObjectClass != NULL);
-
-	FOverriddenPropertySet* OverriddenProperties = Object ? UnrealExporterPrivate::GetOverriddenProperties((UObject*)Object) : nullptr;
-	FEnableOverridableSerializationScope Scope(OverriddenProperties != nullptr, OverriddenProperties);
+	check(Struct != NULL);
 
-	for (FProperty* Property = ObjectClass->PropertyLink; Property; Property = Property->PropertyLinkNext)
+	for (FProperty* Property = Struct->PropertyLink; Property; Property = Property->PropertyLinkNext)
 	{
 		if (!Property->ShouldPort(PortFlags))
 		{
@@ -689,9 +688,14 @@ void ExportProperties_Overrides
 			// to increase frequency of forward compatibility, only sanitize property names that absolutely need it
 			SanitizedPropertyName = FString::Format(TEXT("\"{0}\""), {Property->GetName().ReplaceCharWithEscapedChar()});
 		}
+		if (!PropertyPrefix.IsEmpty())
+		{
+			SanitizedPropertyName = PropertyPrefix + "." + SanitizedPropertyName;
+		}
 
 		FArrayProperty* ArrayProperty = CastField<FArrayProperty>(Property);
 		FMapProperty* MapProperty = CastField<FMapProperty>(Property);
+		FStructProperty* StructProperty = CastField<FStructProperty>(Property);
 
 		FObjectPropertyBase *const ExportObjectProp = (Property->PropertyFlags & CPF_ExportObject) != 0 ? CastField<FObjectPropertyBase>(Property) : NULL;
 		const uint32 ExportFlags = PortFlags | PPF_Delimited;
@@ -743,29 +747,29 @@ void ExportProperties_Overrides
 				// This is used as the default value in the case of an array property that has
 				// fewer elements than the exported object.
 				uint8* StructDefaults = nullptr;
-				FStructProperty* StructProperty = CastField<FStructProperty>(InnerProp);
-				if (StructProperty != nullptr)
+				FStructProperty* InnerStructProperty = CastField<FStructProperty>(InnerProp);
+				if (InnerStructProperty != nullptr)
 				{
-					checkSlow(StructProperty->Struct);
-					StructDefaults = (uint8*)FMemory::Malloc(StructProperty->Struct->GetStructureSize());
-					StructProperty->InitializeValue(StructDefaults);
+					checkSlow(InnerStructProperty->Struct);
+					StructDefaults = (uint8*)FMemory::Malloc(InnerStructProperty->Struct->GetStructureSize());
+					InnerStructProperty->InitializeValue(StructDefaults);
 				}
 				ON_SCOPE_EXIT
 				{
 					if (StructDefaults)
 					{
-						StructProperty->DestroyValue(StructDefaults);
+						InnerStructProperty->DestroyValue(StructDefaults);
 						FMemory::Free(StructDefaults);
 					}
 				};
 
-				void* Arr = Property->ContainerPtrToValuePtr<void>(Object, PropertyArrayIndex);
+				void* Arr = Property->ContainerPtrToValuePtr<void>(ValuePtr, PropertyArrayIndex);
 				FScriptArrayHelper ArrayHelper(ArrayProperty, Arr);
 
 				void* DiffArr = nullptr;
-				if (DiffClass)
+				if (DiffStruct)
 				{
-					DiffArr = Property->ContainerPtrToValuePtrForDefaults<void>(DiffClass, Diff, PropertyArrayIndex);
+					DiffArr = Property->ContainerPtrToValuePtrForDefaults<void>(DiffStruct, Diff, PropertyArrayIndex);
 				}
 				// we won't use this if DiffArr is NULL, but we have to set it up to something
 				FScriptArrayHelper DiffArrayHelper(ArrayProperty, DiffArr);
@@ -775,10 +779,10 @@ void ExportProperties_Overrides
 				if (OverriddenProperties)
 				{
 					FOverridableTextPortPropertyPathScope ScopePath(Property);
-					Operation = FOverridableSerializationLogic::GetOverriddenPropertyOperationForPortText(Arr, DiffArr, PortFlags);
 					FPropertyVisitorPath* Path = FOverridableSerializationLogic::GetOverriddenPortTextPropertyPath();
 					checkf(Path, TEXT("Expecting a path"));
 					Chain = Path->ToSerializedPropertyChain();
+					Operation = FOverridableSerializationLogic::GetOverriddenPropertyOperation(PortFlags, &Chain, nullptr, Arr, DiffArr);
 				}
 
 				if(!OverriddenProperties && !Property->HasAnyPropertyFlags(CPF_ExperimentalOverridableLogic))
@@ -830,10 +834,10 @@ void ExportProperties_Overrides
 
 						const FObjectProperty* InnerObjectProperty = CastField<FObjectProperty>(InnerProp);
 
-						// Only array of instanced subobjects are handled here as sort of a set where the matching key is done using the archetype
+						// Only array of instanced subobjects are handled here as like a set where the matching key is done using the archetype
 						checkf(InnerObjectProperty&& InnerObjectProperty->HasAnyPropertyFlags(CPF_PersistentInstance), TEXT("Expecting only arrays of instanced subobjects"));
 
-						auto FindObject = [InnerObjectProperty](const FOverriddenPropertyNodeID ObjectToFind, FScriptArrayHelper& ArrayHelper) -> int32
+						auto FindObject = [InnerObjectProperty](const FOverriddenPropertyNodeID& ObjectToFind, FScriptArrayHelper& ArrayHelper) -> int32
 						{
 							const int32 ArrayNum = ArrayHelper.Num();
 							for (int i = 0; i < ArrayNum; ++i)
@@ -849,19 +853,27 @@ void ExportProperties_Overrides
 							return INDEX_NONE;
 						};
 
-						FScriptArrayHelper DefaultsArrayHelper(ArrayProperty, DiffArr);
-
 						TArray<int32> RemovedIndices;
 						TArray<int32> AddedIndices;
+						TArray<int32> ShadowIndices;
 						bool bExternalized = false;
 
+						if (FOverridableSerializationLogic::ShouldPropertyShadowSerializeSubObject(ArrayProperty))
+						{
+							const int32 ArrayNum = ArrayHelper.Num();
+							for (int i = 0; i < ArrayNum; i++)
+							{
+								ShadowIndices.Add(i);
+							}
+						}
+
 						for (const FOverriddenPropertyNode& SubNode: ArrayOverriddenPropertyNode->GetSubPropertyNodes())
 						{
 							switch (SubNode.GetOperation())
 							{
 							case EOverriddenPropertyOperation::Remove:
 								{
-									const int32 DefaultIndex = FindObject(SubNode.GetNodeID(), DefaultsArrayHelper);
+									const int32 DefaultIndex = FindObject(SubNode.GetNodeID(), DiffArrayHelper);
 									if (DefaultIndex != INDEX_NONE)
 									{
 										RemovedIndices.Add(DefaultIndex);
@@ -874,6 +886,7 @@ void ExportProperties_Overrides
 									if (Index != INDEX_NONE)
 									{
 										AddedIndices.Add(Index);
+										ShadowIndices.Remove(Index);
 									}
 									break;
 								}
@@ -900,7 +913,9 @@ void ExportProperties_Overrides
 							Out.Logf(TEXT("%s%s<%s>=%s\r\n"), FCString::Spc(Indent), *SanitizedPropertyName, *GetOverriddenOperationString(EOverriddenPropertyOperation::Remove), *Value);
 						}
 
-						for (int32 i : AddedIndices)
+						auto ExportArray = [&](TArray<int32> ExportIndices, const EOverriddenPropertyOperation Operation)
+						{
+							for (int32 i : ExportIndices)
 						{
 							FOverridableTextPortPropertyPathScope ScopePath(Property, i, EPropertyVisitorInfoType::ContainerIndex);
 
@@ -908,8 +923,12 @@ void ExportProperties_Overrides
 							InnerProp->ExportTextItem_Direct(Value, ArrayHelper.GetRawPtr(i), nullptr, Parent, ExportFlags, ExportRootScope);
 							HandleExportObject(InnerProp, ArrayHelper.GetRawPtr(i), ArrayExportObjectProp);
 
-							Out.Logf(TEXT("%s%s<%s>=%s\r\n"), FCString::Spc(Indent), *SanitizedPropertyName, *GetOverriddenOperationString(EOverriddenPropertyOperation::Add), *Value);
+								Out.Logf(TEXT("%s%s<%s>=%s\r\n"), FCString::Spc(Indent), *SanitizedPropertyName, *GetOverriddenOperationString(Operation), *Value);
 						}
+						};
+
+						ExportArray(AddedIndices, EOverriddenPropertyOperation::Add);
+						ExportArray(ShadowIndices, EOverriddenPropertyOperation::SubObjectsShadowing);
 					}
 					continue;
 				}
@@ -919,11 +938,11 @@ void ExportProperties_Overrides
 			if (MapProperty != nullptr)
 			{
 
-				void* Map = Property->ContainerPtrToValuePtr<void>(Object, PropertyArrayIndex);
+				void* Map = Property->ContainerPtrToValuePtr<void>(ValuePtr, PropertyArrayIndex);
 				void* DiffMap = nullptr;
-				if (DiffClass)
+				if (DiffStruct)
 				{
-					DiffMap = Property->ContainerPtrToValuePtrForDefaults<void>(DiffClass, Diff, PropertyArrayIndex);
+					DiffMap = Property->ContainerPtrToValuePtrForDefaults<void>(DiffStruct, Diff, PropertyArrayIndex);
 				}
 
 				EOverriddenPropertyOperation Operation = EOverriddenPropertyOperation::None;
@@ -950,7 +969,15 @@ void ExportProperties_Overrides
 						TArray<int32> RemovedIndices;
 						TArray<int32> ModifiedIndices;
 						TArray<int32> AddedIndices;
+						TArray<int32> ShadowedIndices;
 
+						if (FOverridableSerializationLogic::ShouldPropertyShadowSerializeSubObject(MapProperty))
+						{
+							for (FScriptMapHelper::FIterator It(MapHelper); It; ++It)
+							{
+								ShadowedIndices.Add(It.GetInternalIndex());
+							}
+						}
 
 						// Figure out the modifications of the map
 						for (const FOverriddenPropertyNode& SubNode : MapOverriddenPropertyNode->GetSubPropertyNodes())
@@ -972,6 +999,7 @@ void ExportProperties_Overrides
 									if (InternalIndex != INDEX_NONE)
 									{
 										AddedIndices.Add(InternalIndex);
+										ShadowedIndices.Remove(InternalIndex);
 									}
 									break;
 								}
@@ -981,6 +1009,7 @@ void ExportProperties_Overrides
 									if (InternalIndex != INDEX_NONE)
 									{
 										ModifiedIndices.Add(InternalIndex);
+										ShadowedIndices.Remove(InternalIndex);
 									}
 									break;
 								}
@@ -990,15 +1019,20 @@ void ExportProperties_Overrides
 							}
 						}
 
-						auto ExportItem = [&MapHelper, &DiffMapHelper, Parent, ExportRootScope, Indent, SanitizedPropertyName, ExportFlags, &Out](int32 i, EOverriddenPropertyOperation Operation)
+						auto ExportArray = [&](TArray<int32> ExportIndices, const EOverriddenPropertyOperation Operation)
+						{
+							for (int32 i : ExportIndices)
 						{
+								FOverridableTextPortPropertyPathScope ScopePath(Property, i, EPropertyVisitorInfoType::ContainerIndex);
+
 							FString Key;
-							MapHelper.KeyProp->ExportTextItem_Direct(Key, MapHelper.GetKeyPtr(i), Operation == EOverriddenPropertyOperation::Modified ? DiffMapHelper.GetKeyPtr(i) : nullptr, Parent, ExportFlags, ExportRootScope);
+								MapHelper.KeyProp->ExportTextItem_Direct(Key, MapHelper.GetKeyPtr(i), /*DefaultValue*/nullptr, Parent, ExportFlags, ExportRootScope);
 
 							FString Value;
-							MapHelper.ValueProp->ExportTextItem_Direct(Value, MapHelper.GetValuePtr(i), Operation == EOverriddenPropertyOperation::Modified ? DiffMapHelper.GetValuePtr(i) : nullptr, Parent, ExportFlags, ExportRootScope);
+								MapHelper.ValueProp->ExportTextItem_Direct(Value, MapHelper.GetValuePtr(i), /*DefaultValue*/nullptr, Parent, ExportFlags, ExportRootScope);
 
 							Out.Logf(TEXT("%s%s<%s>=(%s,%s)\r\n"), FCString::Spc(Indent), *SanitizedPropertyName, *GetOverriddenOperationString(Operation), *Key, *Value);
+							}
 						};
 
 						for (int32 i : RemovedIndices)
@@ -1009,40 +1043,50 @@ void ExportProperties_Overrides
 							Out.Logf(TEXT("%s%s<%s>=%s\r\n"), FCString::Spc(Indent), *SanitizedPropertyName, *GetOverriddenOperationString(EOverriddenPropertyOperation::Remove), *Value);
 						}
 
-						for (int32 i : ModifiedIndices)
-						{
-							FOverridableTextPortPropertyPathScope ScopePath(Property, i, EPropertyVisitorInfoType::ContainerIndex);
+						ExportArray(ModifiedIndices, EOverriddenPropertyOperation::Modified);
+						ExportArray(AddedIndices, EOverriddenPropertyOperation::Add);
+						ExportArray(ShadowedIndices, EOverriddenPropertyOperation::SubObjectsShadowing);
+					}
 
-							ExportItem(i, EOverriddenPropertyOperation::Modified);
+					continue;
+				}
 						}
 
-						for (int32 i : AddedIndices)
+			// Recursively call this method for struct in the overridable serialization case
+			if (StructProperty != nullptr && Struct == DiffStruct && OverriddenProperties)
 						{
-							FOverridableTextPortPropertyPathScope ScopePath(Property, i, EPropertyVisitorInfoType::ContainerIndex);
-
-							ExportItem(i, EOverriddenPropertyOperation::Add);
-						}
-					}
-
+				FOverridableTextPortPropertyPathScope ScopePath(StructProperty);
+
+				ExportProperties_Overrides(
+					Context, 
+					Out, 
+					StructProperty->Struct, 
+					Property->ContainerPtrToValuePtr<uint8>(ValuePtr, PropertyArrayIndex), 
+					Indent, 
+					StructProperty->Struct,
+					Diff ? Property->ContainerPtrToValuePtr<uint8>(Diff, PropertyArrayIndex) : nullptr,
+					Parent,
+					PortFlags,
+					ExportRootScope,
+					OverriddenProperties,
+					SanitizedPropertyName);
 					continue;
 				}
-			}
-
 
 			// Generic property export code
 			FString	Value;
 			const bool bStaticArray = Property->ArrayDim > 1;
 			FOverridableTextPortPropertyPathScope ScopePath(Property, bStaticArray ? PropertyArrayIndex : INDEX_NONE, bStaticArray ? EPropertyVisitorInfoType::StaticArrayIndex : EPropertyVisitorInfoType::None);
 
-			uint8* DiffData = (DiffClass && Property->IsInContainer(DiffClass->GetPropertiesSize())) ? Diff : NULL;
-			if (Property->ExportText_InContainer(PropertyArrayIndex, Value, Object, DiffData, Parent, ExportFlags, ExportRootScope))
+			uint8* DiffData = (DiffStruct && Property->IsInContainer(DiffStruct->GetPropertiesSize())) ? Diff : NULL;
+			if (Property->ExportText_InContainer(PropertyArrayIndex, Value, ValuePtr, DiffData, Parent, ExportFlags, ExportRootScope))
 			{
-				HandleExportObject(Property, Property->ContainerPtrToValuePtr<uint8>(Object, PropertyArrayIndex), ExportObjectProp);
+				HandleExportObject(Property, Property->ContainerPtrToValuePtr<uint8>(ValuePtr, PropertyArrayIndex), ExportObjectProp);
 
 				FString OverridableOperation;
 				if (OverriddenProperties)
 				{
-					const EOverriddenPropertyOperation Operation = FOverridableSerializationLogic::GetOverriddenPropertyOperationForPortText(Object, DiffData, PortFlags);
+					const EOverriddenPropertyOperation Operation = FOverridableSerializationLogic::GetOverriddenPropertyOperationForPortText(ValuePtr, DiffData, PortFlags);
 					if (Operation!=EOverriddenPropertyOperation::None)
 					{
 						OverridableOperation = FString::Printf(TEXT("<%s>"), *GetOverriddenOperationString(Operation));
@@ -1060,9 +1104,6 @@ void ExportProperties_Overrides
 			}
 		}
 	}
-
-	// Allows to import/export C++ properties in case the automatic unreal script mesh wouldn't work.
-	Parent->ExportCustomProperties(Out, Indent);
 }
 
 /** Shim ExportProperties when FOverridableSerializationLogic::ECapabilities::T3DSerialization capabilities disabled */
@@ -1290,7 +1331,13 @@ void ExportProperties
 {
 	if (FOverridableSerializationLogic::HasCapabilities(FOverridableSerializationLogic::ECapabilities::T3DSerialization))
 	{
-		ExportProperties_Overrides(Context, Out, ObjectClass, Object, Indent, DiffClass, Diff, Parent, PortFlags, ExportRootScope);
+		FOverriddenPropertySet* OverriddenProperties = Object ? UnrealExporterPrivate::GetOverriddenProperties((UObject*)Object) : nullptr;
+		FEnableOverridableSerializationScope Scope(OverriddenProperties != nullptr, OverriddenProperties);
+
+		ExportProperties_Overrides(Context, Out, ObjectClass, Object, Indent, DiffClass, Diff, Parent, PortFlags, ExportRootScope, OverriddenProperties);
+
+		// Allows to import/export C++ properties in case the automatic unreal script mesh wouldn't work.
+		Parent->ExportCustomProperties(Out, Indent);
 	}
 	else
 	{
