Const correctness fixes in CurveTable related methods.

#Synced-CL 51567995

diff --git a/Engine/Source/Runtime/Engine/Classes/Curves/SimpleCurve.h b/Engine/Source/Runtime/Engine/Classes/Curves/SimpleCurve.h
--- a/Engine/Source/Runtime/Engine/Classes/Curves/SimpleCurve.h
+++ b/Engine/Source/Runtime/Engine/Classes/Curves/SimpleCurve.h
@@ -7,7 +7,7 @@
 #include "SimpleCurve.generated.h"
 
 /** One key in a rich, editable float curve */
-USTRUCT()
+USTRUCT(BlueprintType)
 struct FSimpleCurveKey
 {
 	GENERATED_USTRUCT_BODY()
diff --git a/Engine/Source/Runtime/Engine/Classes/Engine/CurveTable.h b/Engine/Source/Runtime/Engine/Classes/Engine/CurveTable.h
--- a/Engine/Source/Runtime/Engine/Classes/Engine/CurveTable.h
+++ b/Engine/Source/Runtime/Engine/Classes/Engine/CurveTable.h
@@ -62,10 +62,10 @@ class UCurveTable
 	ENGINE_API FSimpleCurve& AddSimpleCurve(FName RowName);
 
 	/** Move the curve to another FName in the table */
-	ENGINE_API void RenameRow(FName& CurveName, FName& NewCurveName);
+	ENGINE_API void RenameRow(const FName& CurveName, const FName& NewCurveName);
 
 	/** Remove a curve row from the table.  Note the associated curve will be deleted. */
-	ENGINE_API void DeleteRow(FName& CurveName);
+	ENGINE_API void DeleteRow(const FName& CurveName);
 
 protected:
 	/** 
diff --git a/Engine/Source/Runtime/Engine/Classes/Kismet/DataTableFunctionLibrary.h b/Engine/Source/Runtime/Engine/Classes/Kismet/DataTableFunctionLibrary.h
--- a/Engine/Source/Runtime/Engine/Classes/Kismet/DataTableFunctionLibrary.h
+++ b/Engine/Source/Runtime/Engine/Classes/Kismet/DataTableFunctionLibrary.h
@@ -97,7 +97,7 @@ class UDataTableFunctionLibrary : public UBlueprintFunctionLibrary
 	 * @return True if the operation was successful.
 	 */
 	UFUNCTION(BlueprintCallable, Category = "DataTable")
-	static ENGINE_API bool RenameCurveTableRow(UCurveTable* CurveTable, FName& RowName, FName& NewRowName);
+	static ENGINE_API bool RenameCurveTableRow(UCurveTable* CurveTable, const FName& RowName, const FName& NewRowName);
 
     
 	/** Get the row struct used by the given Data Table, if any */
diff --git a/Engine/Source/Runtime/Engine/Private/CurveTable.cpp b/Engine/Source/Runtime/Engine/Private/CurveTable.cpp
--- a/Engine/Source/Runtime/Engine/Private/CurveTable.cpp
+++ b/Engine/Source/Runtime/Engine/Private/CurveTable.cpp
@@ -605,7 +605,7 @@ FSimpleCurve& UCurveTable::AddSimpleCurve(FName RowName)
 	return *Result;
 }
 
-void UCurveTable::RenameRow(FName& CurveName, FName& NewCurveName)
+void UCurveTable::RenameRow(const FName& CurveName, const FName& NewCurveName)
 {
 	if (CurveName != NewCurveName && !RowMap.Contains(NewCurveName))
 	{
@@ -617,7 +617,7 @@ void UCurveTable::RenameRow(FName& CurveName, FName& NewCurveName)
 	}
 }
 
-void UCurveTable::DeleteRow(FName& CurveName)
+void UCurveTable::DeleteRow(const FName& CurveName)
 {
 	FRealCurve* CurveToDelete = nullptr;
 	if (RowMap.RemoveAndCopyValue(CurveName, CurveToDelete))
diff --git a/Engine/Source/Runtime/Engine/Private/DataTableFunctionLibrary.cpp b/Engine/Source/Runtime/Engine/Private/DataTableFunctionLibrary.cpp
--- a/Engine/Source/Runtime/Engine/Private/DataTableFunctionLibrary.cpp
+++ b/Engine/Source/Runtime/Engine/Private/DataTableFunctionLibrary.cpp
@@ -107,7 +107,8 @@ bool UDataTableFunctionLibrary::RemoveCurveTableRow(UCurveTable* CurveTable, FNa
 
 }
 
-bool UDataTableFunctionLibrary::RenameCurveTableRow(UCurveTable* CurveTable, FName& RowName, FName& NewRowName)
+bool UDataTableFunctionLibrary::RenameCurveTableRow(
+	UCurveTable* CurveTable, const FName& RowName, const FName& NewRowName)
 {
 	if (!CurveTable)
 	{
