Substrate - NMT - Adds support for unlit BSDF expression node.
Also fixes DUAL_SOURCE_COLOR_BLENDING_ENABLED to match legacy translator and allow colored transmittance blending.

#tests PC Editor, substrate unlit node with NMT.
#rnx
#Synced-CL 51581332

diff --git a/Engine/Source/Runtime/Engine/Private/Materials/MaterialExpressionsToMIR.cpp b/Engine/Source/Runtime/Engine/Private/Materials/MaterialExpressionsToMIR.cpp
--- a/Engine/Source/Runtime/Engine/Private/Materials/MaterialExpressionsToMIR.cpp
+++ b/Engine/Source/Runtime/Engine/Private/Materials/MaterialExpressionsToMIR.cpp
@@ -6550,6 +6550,29 @@ void UMaterialExpressionSubstrateSingleLayerWaterBSDF::Build(MIR::FEmitter& Em)
 	FinalizeBSDFAndPromoteToOperatorIfNeeded(Em, SubstrateOperator, SingleLayerWaterBSDF);
 }
 
+void UMaterialExpressionSubstrateUnlitBSDF::Build(MIR::FEmitter& Em)
+{
+	FSubstrateTranslatorData* SubstrateTranslatorData = Em.GetSubstrateTranslatorData();
+
+	FGuid PathUniqueId = SubstrateTranslatorData->SubstrateTreeStackGetPathUniqueId();
+	FSubstrateOperator& SubstrateOperator = SubstrateTranslatorData->SubstrateCompilationGetOperator(PathUniqueId);
+
+	FValueRef EmissiveColorArg = CompileWithDefaultFloat3(Em, EmissiveColor, 0.0f, 0.0f, 0.0f);
+	FValueRef TransmittanceColorArg = CompileWithDefaultFloat3(Em, TransmittanceColor, 1.0f, 1.0f, 1.0f);
+	FValueRef NormalArg = CompileWithDefaultNormalWS(Em, Material, Normal);
+
+	FValueRef UnlitBSDF = Em.Extern<MIR::FExternSimpleHLSL>(
+		{
+			MIR::EExternSimpleType::SubstrateData,
+			TEXT("GetSubstrateUnlitBSDF($0, $1, $2)"),
+		},
+		EmissiveColorArg,
+		TransmittanceColorArg,
+		NormalArg);
+
+	FinalizeBSDFAndPromoteToOperatorIfNeeded(Em, SubstrateOperator, UnlitBSDF);
+}
+
 static FValueRef SubstrateBlendNormal(MIR::FEmitter& Em, FValueRef NormalCodeChunk0, FValueRef NormalCodeChunk1, FValueRef MixCodeChunk)
 {
 	FValueRef SafeMixCodeChunk = Em.Saturate(MixCodeChunk);
diff --git a/Engine/Source/Runtime/Engine/Private/Materials/MaterialIRToHLSLTranslator.cpp b/Engine/Source/Runtime/Engine/Private/Materials/MaterialIRToHLSLTranslator.cpp
--- a/Engine/Source/Runtime/Engine/Private/Materials/MaterialIRToHLSLTranslator.cpp
+++ b/Engine/Source/Runtime/Engine/Private/Materials/MaterialIRToHLSLTranslator.cpp
@@ -2600,7 +2600,7 @@ static void GetShaderCompilerEnvironment(const FMaterial* Material, const FMater
 		}
 	}
 
-	OutEnvironment.SetDefine(TEXT("DUAL_SOURCE_COLOR_BLENDING_ENABLED"), false);
+	OutEnvironment.SetDefine(TEXT("DUAL_SOURCE_COLOR_BLENDING_ENABLED"), Material->IsDualBlendingEnabled() ? 1 : 0);
 	OutEnvironment.SetDefine(TEXT("TEXTURE_SAMPLE_DEBUG"), false);
 
 	for (int32 VTStackIndex = 0; VTStackIndex < CompilationOutput.UniformExpressionSet.GetVTStacks().Num(); ++VTStackIndex)
diff --git a/Engine/Source/Runtime/Engine/Public/Materials/MaterialExpressionSubstrate.h b/Engine/Source/Runtime/Engine/Public/Materials/MaterialExpressionSubstrate.h
--- a/Engine/Source/Runtime/Engine/Public/Materials/MaterialExpressionSubstrate.h
+++ b/Engine/Source/Runtime/Engine/Public/Materials/MaterialExpressionSubstrate.h
@@ -593,7 +593,7 @@ class UMaterialExpressionSubstrateUnlitBSDF : public UMaterialExpressionSubstrat
 	//~ Begin UMaterialExpression Interface
 #if WITH_EDITOR
 	virtual int32 Compile(class FMaterialCompiler* Compiler, int32 OutputIndex) override;
-	// SUBSTRATE_TODO virtual void Build(MIR::FEmitter& Emitter) override;
+	virtual void Build(MIR::FEmitter& Emitter) override;
 	virtual void GetCaption(TArray<FString>& OutCaptions) const override;
 	virtual EMaterialValueType GetOutputValueType(int32 OutputIndex) override;
 	virtual EMaterialValueType GetInputValueType(int32 InputIndex) override;
