[shaders] modify some hashes stored in shadermap data (mostly but not exclusively editor only) to use 64-bit xxhash instead of sha1 to improve both speed of generation and memory footprint.

#Synced-CL 51584990

diff --git a/Engine/Plugins/Compositing/OpenColorIO/Source/OpenColorIO/Private/OpenColorIODerivedDataVersion.h b/Engine/Plugins/Compositing/OpenColorIO/Source/OpenColorIO/Private/OpenColorIODerivedDataVersion.h
--- a/Engine/Plugins/Compositing/OpenColorIO/Source/OpenColorIO/Private/OpenColorIODerivedDataVersion.h
+++ b/Engine/Plugins/Compositing/OpenColorIO/Source/OpenColorIO/Private/OpenColorIODerivedDataVersion.h
@@ -10,4 +10,4 @@ OpenColorIOShaderDerivedDataVersion.h: Derived data version for OpenColorIO shad
 
 // In case of merge conflicts with DDC versions, you MUST generate a new GUID and set this new
 // guid as version
-#define OPENCOLORIO_DERIVEDDATA_VER		TEXT("A03A807F-923B-48AE-AA2A-77A1E51A6394")
+#define OPENCOLORIO_DERIVEDDATA_VER		TEXT("EF2734CC-A92C-4025-A75B-96AB950E5198")
diff --git a/Engine/Plugins/Compositing/OpenColorIO/Source/OpenColorIO/Private/OpenColorIOShaderMap.cpp b/Engine/Plugins/Compositing/OpenColorIO/Source/OpenColorIO/Private/OpenColorIOShaderMap.cpp
--- a/Engine/Plugins/Compositing/OpenColorIO/Source/OpenColorIO/Private/OpenColorIOShaderMap.cpp
+++ b/Engine/Plugins/Compositing/OpenColorIO/Source/OpenColorIO/Private/OpenColorIOShaderMap.cpp
@@ -85,15 +85,14 @@ void UpdateOpenColorIOShaderCompilingStats(const FOpenColorIOTransformResource*
 }*/
 
 /** Hashes the color transform specific part of this shader map Id. */
-void FOpenColorIOShaderMapId::GetOpenColorIOHash(FSHAHash& OutHash) const
+uint64 FOpenColorIOShaderMapId::GetOpenColorIOHash() const
 {
-	FSHA1 HashState;
+	FXxHash64Builder HashState;
 
-	HashState.UpdateWithString(*ShaderCodeAndConfigHash, ShaderCodeAndConfigHash.Len());
+	HashState.Update(*ShaderCodeAndConfigHash, ShaderCodeAndConfigHash.Len() * sizeof(FMemoryImageString::ElementType));
 	HashState.Update((const uint8*)&FeatureLevel, sizeof(FeatureLevel));
 
-	HashState.Final();
-	HashState.GetHash(&OutHash.Hash[0]);
+	return HashState.Finalize().Hash;
 }
 
 /** 
@@ -199,7 +198,7 @@ void FOpenColorIOShaderType::BeginCompileShader(
 }
 
 FShader* FOpenColorIOShaderType::FinishCompileShader(
-	const FSHAHash& InShaderMapHash,
+	uint64 InShaderMapHash,
 	const FShaderCompileJob& InCurrentJob,
 	const FString& InDebugDescription
 	) const
@@ -451,7 +450,7 @@ void FOpenColorIOShaderMap::Compile(FOpenColorIOTransformResource* InColorTransf
 	}
 }
 
-FShader* FOpenColorIOShaderMap::ProcessCompilationResultsForSingleJob(const TRefCountPtr<class FShaderCommonCompileJob>& SingleJob, const FSHAHash& InShaderMapHash)
+FShader* FOpenColorIOShaderMap::ProcessCompilationResultsForSingleJob(const TRefCountPtr<class FShaderCommonCompileJob>& SingleJob, uint64 InShaderMapHash)
 {
 	FShaderCompileJob* CurrentJob = SingleJob->GetSingleShaderJob();
 	check(CurrentJob->Id == CompilingId);
@@ -476,8 +475,7 @@ bool FOpenColorIOShaderMap::ProcessCompilationResults(const TArray<FShaderCommon
 
 	double StartTime = FPlatformTime::Seconds();
 
-	FSHAHash ShaderMapHash;
-	GetContent()->ShaderMapId.GetOpenColorIOHash(ShaderMapHash);
+	uint64 ShaderMapHash = GetContent()->ShaderMapId.GetOpenColorIOHash();
 
 	do
 	{
@@ -568,7 +566,7 @@ bool FOpenColorIOShaderMap::IsComplete(const FOpenColorIOTransformResource* InCo
 
 void FOpenColorIOShaderMap::GetShaderList(TMap<FShaderId, TShaderRef<FShader>>& OutShaders) const
 {
-	GetContent()->GetShaderList(*this, FSHAHash(), OutShaders);
+	GetContent()->GetShaderList(*this, 0u, OutShaders);
 }
 
 void FOpenColorIOShaderMap::GetShaderList(TMap<FHashedName, TShaderRef<FShader>>& OutShaders) const
diff --git a/Engine/Plugins/Compositing/OpenColorIO/Source/OpenColorIO/Public/OpenColorIOShaderType.h b/Engine/Plugins/Compositing/OpenColorIO/Source/OpenColorIO/Public/OpenColorIOShaderType.h
--- a/Engine/Plugins/Compositing/OpenColorIO/Source/OpenColorIO/Public/OpenColorIOShaderType.h
+++ b/Engine/Plugins/Compositing/OpenColorIO/Source/OpenColorIO/Public/OpenColorIOShaderType.h
@@ -65,7 +65,7 @@ public:
 			const FParameters* InParameters,
 			int32 InPermutationId,
 			const FShaderCompilerOutput& CompilerOutput,
-			const FSHAHash& InOCIOShaderMapHash,
+			uint64 InOCIOShaderMapHash,
 			const FString& InDebugDescription
 			)
 		: FShaderCompiledShaderInitializerType(InType, InParameters, InPermutationId, CompilerOutput, InOCIOShaderMapHash, nullptr, nullptr)
@@ -141,10 +141,13 @@ public:
 
 	/**
 	 * Either creates a new instance of this type or returns an equivalent existing shader.
+	 * @param InOCIOShaderMapHash - The hash value of the shadermap which will reference this shader
 	 * @param CurrentJob - Compile job that was enqueued by BeginCompileShader.
+	 * @param InDebugDescription - Description of the shader job (type, permutation ID, etc.)
+	 * @return The constructed FShader instance for the completed job
 	 */
 	FShader* FinishCompileShader(
-		const FSHAHash& InOCIOShaderMapHash,
+		uint64 InOCIOShaderMapHash,
 		const FShaderCompileJob& CurrentJob,
 		const FString& InDebugDescription
 		) const;
diff --git a/Engine/Plugins/Compositing/OpenColorIO/Source/OpenColorIO/Public/OpenColorIOShared.h b/Engine/Plugins/Compositing/OpenColorIO/Source/OpenColorIO/Public/OpenColorIOShared.h
--- a/Engine/Plugins/Compositing/OpenColorIO/Source/OpenColorIO/Public/OpenColorIOShared.h
+++ b/Engine/Plugins/Compositing/OpenColorIO/Source/OpenColorIO/Public/OpenColorIOShared.h
@@ -85,7 +85,7 @@ public:
 	}
 
 	/** Hashes the color transform-specific part of this shader map Id. */
-	void GetOpenColorIOHash(FSHAHash& OutHash) const;
+	uint64 GetOpenColorIOHash() const;
 
 	/**
 	 * Tests this set against another for equality, disregarding override settings.
@@ -292,7 +292,7 @@ private:
 	uint32 bHasFrozenContent : 1;
 
 #if WITH_EDITOR
-	FShader* ProcessCompilationResultsForSingleJob(const TRefCountPtr<class FShaderCommonCompileJob>& SingleJob, const FSHAHash& InShaderMapHash);
+	FShader* ProcessCompilationResultsForSingleJob(const TRefCountPtr<class FShaderCommonCompileJob>& SingleJob, uint64 InShaderMapHash);
 #endif
 
 	bool IsOpenColorIOShaderComplete(const FOpenColorIOTransformResource* InColorTransform, const FOpenColorIOShaderType* InShaderType, bool bSilent);
diff --git a/Engine/Plugins/Experimental/NNERuntimeIREE/Source/NNERuntimeIREEShader/Internal/NNERuntimeIREEShaderShared.h b/Engine/Plugins/Experimental/NNERuntimeIREE/Source/NNERuntimeIREEShader/Internal/NNERuntimeIREEShaderShared.h
--- a/Engine/Plugins/Experimental/NNERuntimeIREE/Source/NNERuntimeIREEShader/Internal/NNERuntimeIREEShaderShared.h
+++ b/Engine/Plugins/Experimental/NNERuntimeIREE/Source/NNERuntimeIREEShader/Internal/NNERuntimeIREEShaderShared.h
@@ -74,7 +74,7 @@ public:
 	}
 
 	/** Hashes the kernel specific part of this shader map Id. */
-	void GetKernelHash(FSHAHash& OutHash) const;
+	uint64 GetKernelHash() const;
 
 	/**
 	 * Tests this set against another for equality, disregarding override settings.
@@ -201,7 +201,7 @@ public:
 
 private:
 #if WITH_EDITOR
-	FShader* ProcessCompilationResultsForSingleJob(FShaderCompileJob& InSingleJob, const FSHAHash& InShaderMapHash);
+	FShader* ProcessCompilationResultsForSingleJob(FShaderCompileJob& InSingleJob, uint64 InShaderMapHash);
 #endif
 
 	bool IsIREEShaderComplete(const FNNERuntimeIREEResource* InKernel, const FNNERuntimeIREEShaderType* InShaderType, bool bSilent);
diff --git a/Engine/Plugins/Experimental/NNERuntimeIREE/Source/NNERuntimeIREEShader/Internal/NNERuntimeIREEShaderType.h b/Engine/Plugins/Experimental/NNERuntimeIREE/Source/NNERuntimeIREEShader/Internal/NNERuntimeIREEShaderType.h
--- a/Engine/Plugins/Experimental/NNERuntimeIREE/Source/NNERuntimeIREEShader/Internal/NNERuntimeIREEShaderType.h
+++ b/Engine/Plugins/Experimental/NNERuntimeIREE/Source/NNERuntimeIREEShader/Internal/NNERuntimeIREEShaderType.h
@@ -51,7 +51,7 @@ public:
 			const FParameters* InParameters,
 			int32 InPermutationId,
 			const FShaderCompilerOutput& CompilerOutput,
-			const FSHAHash& InKernelShaderMapHash,
+			uint64 InKernelShaderMapHash,
 			const FString& InDebugDescription
 			)
 			: FGlobalShaderType::CompiledShaderInitializerType(
@@ -141,7 +141,7 @@ public:
 	 * Either creates a new instance of this type or returns an equivalent existing shader.
 	 */
 	FShader* FinishCompileShader(
-		const FSHAHash& InKernelShaderMapHash,
+		uint64 InKernelShaderMapHash,
 		const FShaderCompileJob& CurrentJob,
 		const FString& InDebugDescription
 		) const;
diff --git a/Engine/Plugins/Experimental/NNERuntimeIREE/Source/NNERuntimeIREEShader/Private/NNERuntimeIREEShaderMap.cpp b/Engine/Plugins/Experimental/NNERuntimeIREE/Source/NNERuntimeIREEShader/Private/NNERuntimeIREEShaderMap.cpp
--- a/Engine/Plugins/Experimental/NNERuntimeIREE/Source/NNERuntimeIREEShader/Private/NNERuntimeIREEShaderMap.cpp
+++ b/Engine/Plugins/Experimental/NNERuntimeIREE/Source/NNERuntimeIREEShader/Private/NNERuntimeIREEShaderMap.cpp
@@ -54,15 +54,14 @@ static inline bool ShouldCacheNNERuntimeIREEShader(const FNNERuntimeIREEShaderTy
 }
 
 /** Hashes the kernel specific part of this shader map Id. */
-void FNNERuntimeIREEShaderMapId::GetKernelHash(FSHAHash& OutHash) const
+uint64 FNNERuntimeIREEShaderMapId::GetKernelHash() const
 {
-	FSHA1 HashState;
+	FXxHash64Builder HashState;
 
 	HashState.Update((const uint8*)&ShaderCodeHash, sizeof(ShaderCodeHash));
 	HashState.Update((const uint8*)&FeatureLevel, sizeof(FeatureLevel));
 
-	HashState.Final();
-	HashState.GetHash(&OutHash.Hash[0]);
+	return HashState.Finalize().Hash;
 }
 
 /** 
@@ -169,7 +168,7 @@ void FNNERuntimeIREEShaderType::BeginCompileShader(
  * @param InCurrentJob - Compile job that was enqueued by BeginCompileShader.
  */
 FShader* FNNERuntimeIREEShaderType::FinishCompileShader(
-	const FSHAHash& InShaderMapHash,
+	uint64 InShaderMapHash,
 	const FShaderCompileJob& InCurrentJob,
 	const FString& InDebugDescription
 	) const
@@ -335,7 +334,7 @@ void FNNERuntimeIREEShaderMap::Compile(
 	}
 }
 
-FShader* FNNERuntimeIREEShaderMap::ProcessCompilationResultsForSingleJob(FShaderCompileJob& CurrentJob, const FSHAHash& InShaderMapHash)
+FShader* FNNERuntimeIREEShaderMap::ProcessCompilationResultsForSingleJob(FShaderCompileJob& CurrentJob, uint64 InShaderMapHash)
 {
 	check(CurrentJob.Id == CompilingId);
 
@@ -363,8 +362,7 @@ bool FNNERuntimeIREEShaderMap::ProcessCompilationResults(const TArray<FShaderCom
 
 	double StartTime = FPlatformTime::Seconds();
 
-	FSHAHash ShaderMapHash;
-	GetContent()->ShaderMapId.GetKernelHash(ShaderMapHash);
+	uint64 ShaderMapHash = GetContent()->ShaderMapId.GetKernelHash();
 
 	do
 	{
@@ -445,7 +443,7 @@ bool FNNERuntimeIREEShaderMap::IsComplete(const FNNERuntimeIREEResource* InKerne
 
 void FNNERuntimeIREEShaderMap::GetShaderList(TMap<FShaderId, TShaderRef<FShader>>& OutShaders) const
 {
-	GetContent()->GetShaderList(*this, FSHAHash(), OutShaders);
+	GetContent()->GetShaderList(*this, 0u, OutShaders);
 }
 
 void FNNERuntimeIREEShaderMap::GetShaderList(TMap<FHashedName, TShaderRef<FShader>>& OutShaders) const
diff --git a/Engine/Plugins/FX/Niagara/Source/NiagaraShader/Private/NiagaraShader.cpp b/Engine/Plugins/FX/Niagara/Source/NiagaraShader/Private/NiagaraShader.cpp
--- a/Engine/Plugins/FX/Niagara/Source/NiagaraShader/Private/NiagaraShader.cpp
+++ b/Engine/Plugins/FX/Niagara/Source/NiagaraShader/Private/NiagaraShader.cpp
@@ -151,21 +151,21 @@ bool FNiagaraShaderMapPointerTable::LoadFromArchive(FArchive& Ar, const FPlatfor
 
 
 /** Hashes the script-specific part of this shader map Id. */
-void FNiagaraShaderMapId::GetScriptHash(FSHAHash& OutHash) const
+uint64 FNiagaraShaderMapId::GetScriptHash() const
 {
-	FSHA1 HashState;
+	FXxHash64Builder HashState;
 	HashState.Update((const uint8*)&CompilerVersionID, sizeof(CompilerVersionID));
 	HashState.Update(BaseCompileHash.Hash, FNiagaraCompileHash::HashSize);
 	HashState.Update((const uint8*)&FeatureLevel, sizeof(FeatureLevel));
 
 	for (int32 Index = 0; Index < AdditionalDefines.Num(); Index++)
 	{
-		HashState.UpdateWithString(*AdditionalDefines[Index], AdditionalDefines[Index].Len());
+		HashState.Update(*AdditionalDefines[Index], AdditionalDefines[Index].Len() * sizeof(FMemoryImageString::ElementType));
 	}
 
 	for (int32 Index = 0; Index < AdditionalVariables.Num(); Index++)
 	{
-		HashState.UpdateWithString(*AdditionalVariables[Index], AdditionalVariables[Index].Len());
+		HashState.Update(*AdditionalVariables[Index], AdditionalVariables[Index].Len() * sizeof(FMemoryImageString::ElementType));
 	}
 
 	for (int32 Index = 0; Index < ReferencedCompileHashes.Num(); Index++)
@@ -178,8 +178,7 @@ void FNiagaraShaderMapId::GetScriptHash(FSHAHash& OutHash) const
 		HashState.Update(Dependency.SourceHash.Hash, sizeof(Dependency.SourceHash));
 	}
 	
-	HashState.Final();
-	HashState.GetHash(&OutHash.Hash[0]);
+	return HashState.Finalize().Hash;
 }
 
 /** 
@@ -486,12 +485,8 @@ void FNiagaraShaderType::AddUniformBufferIncludesToEnvironment(FShaderCompilerEn
 	UE::ShaderParameters::AddUniformBufferIncludesToEnvironment(OutEnvironment, ReferencedUniformBuffers);
 }
 
-/**
- * Either creates a new instance of this type or returns an equivalent existing shader.
- * @param CurrentJob - Compile job that was enqueued by BeginCompileShader.
- */
 FShader* FNiagaraShaderType::FinishCompileShader(
-	const FSHAHash& ShaderMapHash,
+	uint64 ShaderMapHash,
 	const FShaderCompileJob& CurrentJob
 	) const
 {
@@ -794,7 +789,7 @@ void FNiagaraShaderMap::CreateCompileJobs(
 	bCompiledSuccessfully = false;
 }
 
-FShader* FNiagaraShaderMap::ProcessCompilationResultsForSingleJob(const TRefCountPtr<class FShaderCommonCompileJob>& SingleJob, const FSHAHash& ShaderMapHash)
+FShader* FNiagaraShaderMap::ProcessCompilationResultsForSingleJob(const TRefCountPtr<class FShaderCommonCompileJob>& SingleJob, uint64 ShaderMapHash)
 {
 	auto CurrentJob = SingleJob->GetSingleShaderJob();
 	check(CurrentJob->Id == CompilingId);
@@ -818,8 +813,7 @@ FShader* FNiagaraShaderMap::ProcessCompilationResultsForSingleJob(const TRefCoun
 
 void FNiagaraShaderMap::ProcessAndFinalizeShaderCompileJob(const TRefCountPtr<FShaderCommonCompileJob>& SingleJob)
 {
-	FSHAHash ShaderMapHash;
-	GetContent()->ShaderMapId.GetScriptHash(ShaderMapHash);
+	uint64 ShaderMapHash = GetContent()->ShaderMapId.GetScriptHash();
 
 	ProcessCompilationResultsForSingleJob(SingleJob, ShaderMapHash);
 
@@ -835,8 +829,7 @@ bool FNiagaraShaderMap::ProcessCompilationResults(const TArray<FNiagaraShaderScr
 
 	double StartTime = FPlatformTime::Seconds();
 
-	FSHAHash ShaderMapHash;
-	GetContent()->ShaderMapId.GetScriptHash(ShaderMapHash);
+	uint64 ShaderMapHash = GetContent()->ShaderMapId.GetScriptHash();
 
 	do
 	{
@@ -975,7 +968,7 @@ void FNiagaraShaderMap::LoadMissingShadersFromMemory(const FNiagaraShaderScript*
 
 void FNiagaraShaderMap::GetShaderList(TMap<FShaderId, TShaderRef<FShader>>& OutShaders) const
 {
-	GetContent()->GetShaderList(*this, FSHAHash(), OutShaders);
+	GetContent()->GetShaderList(*this, 0u, OutShaders);
 }
 
 void FNiagaraShaderMap::GetShaderList(TMap<FHashedName, TShaderRef<FShader>>& OutShaders) const
diff --git a/Engine/Plugins/FX/Niagara/Source/NiagaraShader/Public/NiagaraShaderType.h b/Engine/Plugins/FX/Niagara/Source/NiagaraShader/Public/NiagaraShaderType.h
--- a/Engine/Plugins/FX/Niagara/Source/NiagaraShader/Public/NiagaraShaderType.h
+++ b/Engine/Plugins/FX/Niagara/Source/NiagaraShader/Public/NiagaraShaderType.h
@@ -72,7 +72,7 @@ public:
 			const FParameters* InParameters,
 			int32 InPermutationId,
 			const FShaderCompilerOutput& CompilerOutput,
-			const FSHAHash& InNiagaraShaderMapHash
+			uint64 InNiagaraShaderMapHash
 		)
 			: FShaderCompiledShaderInitializerType(InType, InParameters, InPermutationId, CompilerOutput, InNiagaraShaderMapHash, nullptr, nullptr)
 		{
@@ -160,11 +160,11 @@ public:
 
 	/**
 	 * Either creates a new instance of this type or returns an equivalent existing shader.
-	 * @param script - The script to link the shader with.
+	 * @param NiagaraShaderMapHash - The hash computed for the owning shadermap
 	 * @param CurrentJob - Compile job that was enqueued by BeginCompileShader.
 	 */
 	FShader* FinishCompileShader(
-		const FSHAHash& NiagaraShaderMapHash,
+		uint64 NiagaraShaderMapHash,
 		const FShaderCompileJob& CurrentJob
 		) const;
 #endif // WITH_EDITOR
diff --git a/Engine/Plugins/FX/Niagara/Source/NiagaraShader/Public/NiagaraShared.h b/Engine/Plugins/FX/Niagara/Source/NiagaraShader/Public/NiagaraShared.h
--- a/Engine/Plugins/FX/Niagara/Source/NiagaraShader/Public/NiagaraShared.h
+++ b/Engine/Plugins/FX/Niagara/Source/NiagaraShader/Public/NiagaraShared.h
@@ -371,7 +371,7 @@ public:
 	}
 
 	/** Hashes the script-specific part of this shader map Id. */
-	void GetScriptHash(FSHAHash& OutHash) const;
+	uint64 GetScriptHash() const;
 
 	/**
 	* Tests this set against another for equality, disregarding override settings.
@@ -669,7 +669,7 @@ private:
 	uint32 bIsPersistent : 1;
 
 #if WITH_EDITOR
-	FShader* ProcessCompilationResultsForSingleJob(const TRefCountPtr<class FShaderCommonCompileJob>& SingleJob, const FSHAHash& ShaderMapHash);
+	FShader* ProcessCompilationResultsForSingleJob(const TRefCountPtr<class FShaderCommonCompileJob>& SingleJob, uint64 ShaderMapHash);
 
 	bool IsNiagaraShaderComplete(const FNiagaraShaderScript* Script, const FNiagaraShaderType* ShaderType, bool bSilent);
 #endif
diff --git a/Engine/Plugins/Importers/USDImporter/Source/USDExporter/Private/MaterialExporterUSD.cpp b/Engine/Plugins/Importers/USDImporter/Source/USDExporter/Private/MaterialExporterUSD.cpp
--- a/Engine/Plugins/Importers/USDImporter/Source/USDExporter/Private/MaterialExporterUSD.cpp
+++ b/Engine/Plugins/Importers/USDImporter/Source/USDExporter/Private/MaterialExporterUSD.cpp
@@ -89,10 +89,8 @@ namespace UE::MaterialExporterUSD::Private
 				FMaterialShaderMapId ShaderMapID;
 				MaterialResource->BuildShaderMapId(ShaderMapID, nullptr);
 
-				FSHAHash MaterialHash;
-				ShaderMapID.GetMaterialHash(MaterialHash);
-
-				VersionHash.Update(&MaterialHash.Hash[0], sizeof(MaterialHash.Hash));
+				uint64 MaterialHash = ShaderMapID.GetMaterialHash();
+				VersionHash.Update(reinterpret_cast<const uint8*>(&MaterialHash), sizeof(MaterialHash));
 			}
 			else
 			{
diff --git a/Engine/Plugins/Runtime/ComputeFramework/Source/ComputeFramework/Private/ComputeFramework/ComputeKernelDerivedDataVersion.h b/Engine/Plugins/Runtime/ComputeFramework/Source/ComputeFramework/Private/ComputeFramework/ComputeKernelDerivedDataVersion.h
--- a/Engine/Plugins/Runtime/ComputeFramework/Source/ComputeFramework/Private/ComputeFramework/ComputeKernelDerivedDataVersion.h
+++ b/Engine/Plugins/Runtime/ComputeFramework/Source/ComputeFramework/Private/ComputeFramework/ComputeKernelDerivedDataVersion.h
@@ -5,4 +5,4 @@
 #include "CoreMinimal.h"
 
 // In case of merge conflicts with DDC versions, you MUST generate a new GUID and set this new guid as version
-#define COMPUTEKERNEL_DERIVEDDATA_VER	TEXT("A8B8CDC1-862C-4D85-BA8B-CD6C3FFF32DD")
+#define COMPUTEKERNEL_DERIVEDDATA_VER	TEXT("F4E7148D-2146-4CFA-B831-B8DB48382101")
diff --git a/Engine/Plugins/Runtime/ComputeFramework/Source/ComputeFramework/Private/ComputeFramework/ComputeKernelShaderType.h b/Engine/Plugins/Runtime/ComputeFramework/Source/ComputeFramework/Private/ComputeFramework/ComputeKernelShaderType.h
--- a/Engine/Plugins/Runtime/ComputeFramework/Source/ComputeFramework/Private/ComputeFramework/ComputeKernelShaderType.h
+++ b/Engine/Plugins/Runtime/ComputeFramework/Source/ComputeFramework/Private/ComputeFramework/ComputeKernelShaderType.h
@@ -47,7 +47,7 @@ public:
 			const FParameters* InParameters,
 			int32 InPermutationId,
 			const FShaderCompilerOutput& CompilerOutput,
-			const FSHAHash& InComputeKernelShaderMapHash,
+			uint64 InComputeKernelShaderMapHash,
 			const FString& InDebugDescription
 			)
 			: FGlobalShaderType::CompiledShaderInitializerType(
@@ -137,7 +137,7 @@ public:
 	 * Either creates a new instance of this type or returns an equivalent existing shader.
 	 */
 	FShader* FinishCompileShader(
-		const FSHAHash& InComputeKernelShaderMapHash,
+		uint64 InComputeKernelShaderMapHash,
 		const FShaderCompileJob& CurrentJob,
 		const FString& InDebugDescription
 		) const;
diff --git a/Engine/Plugins/Runtime/ComputeFramework/Source/ComputeFramework/Private/ComputeFramework/ComputeKernelShared.h b/Engine/Plugins/Runtime/ComputeFramework/Source/ComputeFramework/Private/ComputeFramework/ComputeKernelShared.h
--- a/Engine/Plugins/Runtime/ComputeFramework/Source/ComputeFramework/Private/ComputeFramework/ComputeKernelShared.h
+++ b/Engine/Plugins/Runtime/ComputeFramework/Source/ComputeFramework/Private/ComputeFramework/ComputeKernelShared.h
@@ -78,7 +78,7 @@ public:
 	}
 
 	/** Hashes the kernel specific part of this shader map Id. */
-	void GetComputeKernelHash(FSHAHash& OutHash) const;
+	uint64 GetComputeKernelHash() const;
 
 	/**
 	 * Tests this set against another for equality, disregarding override settings.
@@ -285,7 +285,7 @@ private:
 	uint32 bHasFrozenContent : 1;
 
 #if WITH_EDITOR
-	FShader* ProcessCompilationResultsForSingleJob(FShaderCompileJob& InSingleJob, const FSHAHash& InShaderMapHash);
+	FShader* ProcessCompilationResultsForSingleJob(FShaderCompileJob& InSingleJob, uint64 InShaderMapHash);
 #endif
 
 	bool IsComputeKernelShaderComplete(const FComputeKernelResource* InKernel, const FComputeKernelShaderType* InShaderType, bool bSilent);
diff --git a/Engine/Plugins/Runtime/ComputeFramework/Source/ComputeFramework/Private/ComputeKernelShaderMap.cpp b/Engine/Plugins/Runtime/ComputeFramework/Source/ComputeFramework/Private/ComputeKernelShaderMap.cpp
--- a/Engine/Plugins/Runtime/ComputeFramework/Source/ComputeFramework/Private/ComputeKernelShaderMap.cpp
+++ b/Engine/Plugins/Runtime/ComputeFramework/Source/ComputeFramework/Private/ComputeKernelShaderMap.cpp
@@ -58,15 +58,14 @@ static inline bool ShouldCacheComputeKernelShader(const FComputeKernelShaderType
 
 
 /** Hashes the kernel specific part of this shader map Id. */
-void FComputeKernelShaderMapId::GetComputeKernelHash(FSHAHash& OutHash) const
+uint64 FComputeKernelShaderMapId::GetComputeKernelHash() const
 {
-	FSHA1 HashState;
+	FXxHash64Builder HashState;
 
 	HashState.Update((const uint8*)&ShaderCodeHash, sizeof(ShaderCodeHash));
 	HashState.Update((const uint8*)&FeatureLevel, sizeof(FeatureLevel));
 
-	HashState.Final();
-	HashState.GetHash(&OutHash.Hash[0]);
+	return HashState.Finalize().Hash;
 }
 
 /** 
@@ -191,7 +190,7 @@ void FComputeKernelShaderType::BeginCompileShader(
  * @param InCurrentJob - Compile job that was enqueued by BeginCompileShader.
  */
 FShader* FComputeKernelShaderType::FinishCompileShader(
-	const FSHAHash& InShaderMapHash,
+	uint64 InShaderMapHash,
 	const FShaderCompileJob& InCurrentJob,
 	const FString& InDebugDescription
 	) const
@@ -429,7 +428,7 @@ void FComputeKernelShaderMap::Compile(
 	}
 }
 
-FShader* FComputeKernelShaderMap::ProcessCompilationResultsForSingleJob(FShaderCompileJob& CurrentJob, const FSHAHash& InShaderMapHash)
+FShader* FComputeKernelShaderMap::ProcessCompilationResultsForSingleJob(FShaderCompileJob& CurrentJob, uint64 InShaderMapHash)
 {
 	check(CurrentJob.Id == CompilingId);
 
@@ -454,8 +453,7 @@ bool FComputeKernelShaderMap::ProcessCompilationResults(const TArray<FShaderComm
 
 	double StartTime = FPlatformTime::Seconds();
 
-	FSHAHash ShaderMapHash;
-	GetContent()->ShaderMapId.GetComputeKernelHash(ShaderMapHash);
+	uint64 ShaderMapHash = GetContent()->ShaderMapId.GetComputeKernelHash();
 
 	do
 	{
@@ -591,7 +589,7 @@ void FComputeKernelShaderMap::LoadMissingShadersFromMemory(const FComputeKernelR
 
 void FComputeKernelShaderMap::GetShaderList(TMap<FShaderId, TShaderRef<FShader>>& OutShaders) const
 {
-	GetContent()->GetShaderList(*this, FSHAHash(), OutShaders);
+	GetContent()->GetShaderList(*this, 0u, OutShaders);
 }
 
 void FComputeKernelShaderMap::GetShaderList(TMap<FHashedName, TShaderRef<FShader>>& OutShaders) const
diff --git a/Engine/Source/Editor/UnrealEd/Private/Commandlets/DumpMaterialShaderTypes.cpp b/Engine/Source/Editor/UnrealEd/Private/Commandlets/DumpMaterialShaderTypes.cpp
--- a/Engine/Source/Editor/UnrealEd/Private/Commandlets/DumpMaterialShaderTypes.cpp
+++ b/Engine/Source/Editor/UnrealEd/Private/Commandlets/DumpMaterialShaderTypes.cpp
@@ -229,7 +229,7 @@ public:
 		DebugWriter->Serialize(const_cast<char*>(MatInstHeader), sizeof(MatInstHeader) - 1);
 
 		// Print each item.
-		for (TPair<FSHAHash, TArray<FString>> UniqueMaterialInstance : UniqueMaterialInstances)
+		for (TPair<uint64, TArray<FString>> UniqueMaterialInstance : UniqueMaterialInstances)
 		{
 			FString DuplicateLine = FString::Printf(TEXT("Duplicates: %d\n"), UniqueMaterialInstance.Value.Num());
 			DebugWriter->Serialize(const_cast<ANSICHAR*>(StringCast<ANSICHAR>(*DuplicateLine).Get()), DuplicateLine.Len());
@@ -248,7 +248,7 @@ public:
 		DebugWriter->Serialize(const_cast<ANSICHAR*>(StringCast<ANSICHAR>(*OuputLine).Get()), OuputLine.Len());
 	}
 
-	void AddMaterialInstance(const FString& MaterialInstanceName, const FSHAHash& StaticParameterHash)
+	void AddMaterialInstance(const FString& MaterialInstanceName, uint64 StaticParameterHash)
 	{
 		TArray<FString>& InstanceList = UniqueMaterialInstances.FindOrAdd(StaticParameterHash);
 		InstanceList.Add(MaterialInstanceName);
@@ -270,7 +270,7 @@ private:
 	TMap<FString, int32>	GlobalShaderTypeHistogram;
 
 	/** Unique material instances. */
-	TMap<FSHAHash, TArray<FString>> UniqueMaterialInstances;
+	TMap<uint64, TArray<FString>> UniqueMaterialInstances;
 
 	/** Store a copy of the the filename. */
 	FString FileName;
@@ -513,9 +513,7 @@ static void ProcessSwitchOptimizer(const ITargetPlatform* TargetPlatform, const
 
 		static inline uint32 GetKeyHash(const FMaterialShaderMapId& Key)
 		{
-			FSHAHash Hash;
-			Key.GetMaterialHash(Hash, false);
-			return CityHash32((char*)Hash.Hash, 20);
+			return (uint32)Key.GetMaterialHash(false);
 		}
 	};
 
@@ -529,13 +527,12 @@ static void ProcessSwitchOptimizer(const ITargetPlatform* TargetPlatform, const
 
 		static inline uint32 GetKeyHash(const StaticSwitchArrayType& Keys)
 		{
-			FSHA1 Hash;
+			FXxHash64Builder Hash;
 			for(const FStaticSwitchParameter& Key : Keys)
 			{
 				Key.UpdateHash(Hash);
 			}
-			Hash.Final();
-			return CityHash32((char*)Hash.m_digest, FSHA1::DigestSize);
+			return (uint32)Hash.Finalize().Hash;
 		}
 	};
 
@@ -904,11 +901,11 @@ static int ProcessMaterialInstances(const ITargetPlatform* TargetPlatform, const
 			Output.Log(FString::Printf(TEXT("Parent: %s"), Top ? *Top->GetFullName() : TEXT("NO PARENT")));
 			Output.Log(FString::Printf(TEXT("Static Parameters: %d"), MaterialInstance->bHasStaticPermutationResource ? MaterialInstance->GetStaticParameters().StaticSwitchParameters.Num() : 0));
 
-			FSHAHash OutHash;
+			uint64 OutHash;
 
 			if (MaterialInstance->bHasStaticPermutationResource)
 			{
-				FSHA1 Hasher;
+				FXxHash64Builder Hasher;
 				const FStaticParameterSet& ParameterSet = MaterialInstance->GetStaticParameters();
 				for (int32 StaticSwitchIndex = 0; StaticSwitchIndex < ParameterSet.StaticSwitchParameters.Num(); ++StaticSwitchIndex)
 				{
@@ -919,11 +916,10 @@ static int ProcessMaterialInstances(const ITargetPlatform* TargetPlatform, const
 					Output.Log(FString::Printf(TEXT("\t%s : %s"), *StaticSwitchParameter.ParameterInfo.ToString(), StaticSwitchParameter.Value ? TEXT("True") : TEXT("False")));
 				}
 
-				Hasher.Final();
-				Hasher.GetHash(&OutHash.Hash[0]);
+				OutHash = Hasher.Finalize().Hash;
 
 				Output.AddMaterialInstance(AssetData.AssetName.ToString(), OutHash);
-				Output.Log(FString::Printf(TEXT("Static Parameter Hash: %s"), *OutHash.ToString()));
+				Output.Log(FString::Printf(TEXT("Static Parameter Hash: %" UINT64_FMT), OutHash));
 			}
 
 			Output.Log(FString::Printf(TEXT("Base Property Overrides: %s"), MaterialInstance->HasOverridenBaseProperties() ? TEXT("True") : TEXT("False")));
diff --git a/Engine/Source/Runtime/Core/Private/UObject/FortniteMainBranchObjectVersion.cpp b/Engine/Source/Runtime/Core/Private/UObject/FortniteMainBranchObjectVersion.cpp
--- a/Engine/Source/Runtime/Core/Private/UObject/FortniteMainBranchObjectVersion.cpp
+++ b/Engine/Source/Runtime/Core/Private/UObject/FortniteMainBranchObjectVersion.cpp
@@ -8,10 +8,10 @@ TMap<FGuid, FGuid> FFortniteMainBranchObjectVersion::GetSystemGuids()
 	TMap<FGuid, FGuid> SystemGuids;
 	const FDevSystemGuids& DevGuids = FDevSystemGuids::Get();
 
-	SystemGuids.Add(DevGuids.GLOBALSHADERMAP_DERIVEDDATA_VER, FGuid("EF39291B950C42C99FFE8D7AA98064BA"));
+	SystemGuids.Add(DevGuids.GLOBALSHADERMAP_DERIVEDDATA_VER, FGuid("C340EEF623494448979F6D24B496E85A"));
 	SystemGuids.Add(DevGuids.LANDSCAPE_MOBILE_COOK_VERSION, FGuid("32D02EF867C74B71A0D4E0FA41392732"));
-	SystemGuids.Add(DevGuids.MATERIALSHADERMAP_DERIVEDDATA_VER, FGuid("1A77CBCDFBDB4C12BF8563EDBB7D8888"));
-	SystemGuids.Add(DevGuids.NIAGARASHADERMAP_DERIVEDDATA_VER, FGuid("CC229D919A214E59B79D01AC2E5BF981"));
+	SystemGuids.Add(DevGuids.MATERIALSHADERMAP_DERIVEDDATA_VER, FGuid("FB0EB87E970C4FB9B217E0E00942BFCB"));
+	SystemGuids.Add(DevGuids.NIAGARASHADERMAP_DERIVEDDATA_VER, FGuid("87F6AE6C269A4592B281752A582E7F13"));
 	SystemGuids.Add(DevGuids.Niagara_LatestScriptCompileVersion, FGuid("FB15C2870A4E4AE8B5F58566F0E3825E"));
 	SystemGuids.Add(DevGuids.MaterialTranslationDDCVersion, FGuid("B79683EE9E684C13BAC2FDE8B180477F"));
 
diff --git a/Engine/Source/Runtime/Engine/Private/Materials/Material.cpp b/Engine/Source/Runtime/Engine/Private/Materials/Material.cpp
--- a/Engine/Source/Runtime/Engine/Private/Materials/Material.cpp
+++ b/Engine/Source/Runtime/Engine/Private/Materials/Material.cpp
@@ -364,13 +364,13 @@ void FMaterialResource::BuildShaderMapIdOverride(const FBuildShaderMapIdArgs& Ar
 	CachedData.AppendReferencedFunctionIdsTo(OutId.ReferencedFunctions);
 	CachedData.AppendReferencedParameterCollectionIdsTo(OutId.ReferencedParameterCollections);
 
-	CachedData.GetExternalCodeReferencesHash(OutId.ExternalCodeReferencesHash);
+	OutId.ExternalCodeReferencesHash = CachedData.GetExternalCodeReferencesHash();
 
-	Material->GetForceRecompileTextureIdsHash(OutId.TextureReferencesHash);
+	OutId.TextureReferencesHash = Material->GetForceRecompileTextureIdsHash();
 
 	if (MaterialInstance)
 	{
-		MaterialInstance->GetBasePropertyOverridesHash(OutId.BasePropertyOverridesHash);
+		OutId.BasePropertyOverridesHash = MaterialInstance->GetBasePropertyOverridesHash();
 
 		FStaticParameterSet CompositedStaticParameters;
 		MaterialInstance->GetStaticParameterValues(CompositedStaticParameters);
@@ -2460,7 +2460,7 @@ static FAutoConsoleVariable GCompileMaterialsForShaderFormatCVar(
 	);
 
 #if WITH_EDITOR
-void UMaterial::GetForceRecompileTextureIdsHash(FSHAHash &TextureReferencesHash)
+uint64 UMaterial::GetForceRecompileTextureIdsHash() const
 {
 	TArray<UTexture*> ForceRecompileTextures;
 	for (const UMaterialExpression *MaterialExpression : GetExpressions())
@@ -2479,21 +2479,17 @@ void UMaterial::GetForceRecompileTextureIdsHash(FSHAHash &TextureReferencesHash)
 	if (ForceRecompileTextures.Num() <= 0)
 	{
 		//There is no Texture that trig a recompile of the material, nothing to add to the hash
-		return;
+		return 0u;
 	}
 
-	FSHA1 TextureCompileDependencies;
-	FString OriginalHash = TextureReferencesHash.ToString();
-	TextureCompileDependencies.UpdateWithString(*OriginalHash, OriginalHash.Len());
+	FXxHash64Builder TextureCompileDependencies;
 
 	for (UTexture *ForceRecompileTexture : ForceRecompileTextures)
 	{
-		FString TextureGuidString = ForceRecompileTexture->GetLightingGuid().ToString();
-		TextureCompileDependencies.UpdateWithString(*TextureGuidString, TextureGuidString.Len());
+		TextureCompileDependencies.Update(&ForceRecompileTexture->GetLightingGuid(), sizeof(FGuid));
 	}
 
-	TextureCompileDependencies.Final();
-	TextureCompileDependencies.GetHash(&TextureReferencesHash.Hash[0]);
+	return TextureCompileDependencies.Finalize().Hash;
 }
 
 bool UMaterial::IsTextureForceRecompileCacheRessource(UTexture *Texture)
diff --git a/Engine/Source/Runtime/Engine/Private/Materials/MaterialCachedData.cpp b/Engine/Source/Runtime/Engine/Private/Materials/MaterialCachedData.cpp
--- a/Engine/Source/Runtime/Engine/Private/Materials/MaterialCachedData.cpp
+++ b/Engine/Source/Runtime/Engine/Private/Materials/MaterialCachedData.cpp
@@ -148,9 +148,9 @@ void FMaterialCachedExpressionData::AppendReferencedParameterCollectionIdsTo(TAr
 	}
 }
 
-void FMaterialCachedExpressionData::GetExternalCodeReferencesHash(FSHAHash& OutHash) const
+uint64 FMaterialCachedExpressionData::GetExternalCodeReferencesHash() const
 {
-	FSHA1 Hasher;
+	FXxHash64Builder Hasher;
 	for (const TObjectPtr<UClass>& ExternalCodeExpressionClass : ReferencedExternalCodeExpressionClasses)
 	{
 		if (const UObject* DefaultExternalCodeExpression = ExternalCodeExpressionClass->GetDefaultObject())
@@ -165,7 +165,7 @@ void FMaterialCachedExpressionData::GetExternalCodeReferencesHash(FSHAHash& OutH
 			}
 		}
 	}
-	OutHash = Hasher.Finalize();
+	return Hasher.Finalize().Hash;
 }
 
 #if WITH_EDITOR
diff --git a/Engine/Source/Runtime/Engine/Private/Materials/MaterialExpressions.cpp b/Engine/Source/Runtime/Engine/Private/Materials/MaterialExpressions.cpp
--- a/Engine/Source/Runtime/Engine/Private/Materials/MaterialExpressions.cpp
+++ b/Engine/Source/Runtime/Engine/Private/Materials/MaterialExpressions.cpp
@@ -14980,7 +14980,7 @@ void FMaterialLayersFunctionsID::SerializeForDDC(FArchive& Ar)
 }
 
 
-void FMaterialLayersFunctionsID::UpdateHash(FSHA1& HashState) const
+void FMaterialLayersFunctionsID::UpdateHash(FXxHash64Builder& HashState) const
 {
 	for (const FGuid &Guid : LayerIDs)
 	{
diff --git a/Engine/Source/Runtime/Engine/Private/Materials/MaterialExternalCodeRegistry.cpp b/Engine/Source/Runtime/Engine/Private/Materials/MaterialExternalCodeRegistry.cpp
--- a/Engine/Source/Runtime/Engine/Private/Materials/MaterialExternalCodeRegistry.cpp
+++ b/Engine/Source/Runtime/Engine/Private/Materials/MaterialExternalCodeRegistry.cpp
@@ -10,11 +10,11 @@
  * FMaterialExternalCodeEnvironmentDefine
  */
 
-void FMaterialExternalCodeEnvironmentDefine::UpdateHash(FSHA1& Hasher) const
+void FMaterialExternalCodeEnvironmentDefine::UpdateHash(FXxHash64Builder& Hasher) const
 {
 	FString NameString = Name.ToString();
-	Hasher.UpdateWithString(*NameString, NameString.Len());
-	Hasher.Update(ShaderFrequency);
+	Hasher.Update(*NameString, NameString.NumBytesWithoutNull());
+	Hasher.Update(&ShaderFrequency, sizeof(ShaderFrequency));
 }
 
 
@@ -22,23 +22,24 @@ void FMaterialExternalCodeEnvironmentDefine::UpdateHash(FSHA1& Hasher) const
  * FMaterialExternalCodeDeclaration 
  */
 
-void FMaterialExternalCodeDeclaration::UpdateHash(FSHA1& Hasher) const
+void FMaterialExternalCodeDeclaration::UpdateHash(FXxHash64Builder& Hasher) const
 {
 	const uint8 bIsInlinedByte = bIsInlined ? 1 : 0;
 	Hasher.Update(&bIsInlinedByte, sizeof(bIsInlinedByte));
 
-	Hasher.Update(ReturnType);
+	Hasher.Update(&ReturnType, sizeof(ReturnType));
 
 	FString NameString = Name.ToString();
-	Hasher.UpdateWithString(*NameString, NameString.Len());
+	Hasher.Update(*NameString, NameString.NumBytesWithoutNull());
 
-	Hasher.UpdateWithString(*Definition, Definition.Len());
-	Hasher.Update(Derivative);
-	Hasher.Update(ShaderFrequency);
+	Hasher.Update(*Definition, Definition.NumBytesWithoutNull());
+	Hasher.Update(&Derivative, sizeof(Derivative));
+	Hasher.Update(&ShaderFrequency, sizeof(ShaderFrequency));
 
 	for (const TEnumAsByte<EMaterialDomain>& InDomain : Domains)
 	{
-		Hasher.Update(InDomain.GetValue());
+		EMaterialDomain Domain = InDomain.GetValue();
+		Hasher.Update(&Domain, sizeof(Domain));
 	}
 
 	for (const FMaterialExternalCodeEnvironmentDefine& InEnvironmentDefine : EnvironmentDefines)
diff --git a/Engine/Source/Runtime/Engine/Private/Materials/MaterialExternalCodeRegistry.h b/Engine/Source/Runtime/Engine/Private/Materials/MaterialExternalCodeRegistry.h
--- a/Engine/Source/Runtime/Engine/Private/Materials/MaterialExternalCodeRegistry.h
+++ b/Engine/Source/Runtime/Engine/Private/Materials/MaterialExternalCodeRegistry.h
@@ -105,7 +105,7 @@ struct FMaterialExternalCodeEnvironmentDefine
 	EMaterialShaderFrequency ShaderFrequency = EMaterialShaderFrequency::Any;
 
 	/** Updates the input hasher state with the content of this environemnt define. */
-	void UpdateHash(FSHA1& Hasher) const;
+	void UpdateHash(FXxHash64Builder& Hasher) const;
 };
 
 /** Declaration of external HLSL code. Such code expressions can be emitted as part of a material translation. */
@@ -169,7 +169,7 @@ public:
 	}
 
 	/** Updates the input hasher state with the content of this external code declaration. */
-	void UpdateHash(FSHA1& Hasher) const;
+	void UpdateHash(FXxHash64Builder& Hasher) const;
 };
 
 UCLASS(Config=MaterialExpressions)
diff --git a/Engine/Source/Runtime/Engine/Private/Materials/MaterialInstance.cpp b/Engine/Source/Runtime/Engine/Private/Materials/MaterialInstance.cpp
--- a/Engine/Source/Runtime/Engine/Private/Materials/MaterialInstance.cpp
+++ b/Engine/Source/Runtime/Engine/Private/Materials/MaterialInstance.cpp
@@ -5021,14 +5021,14 @@ bool UMaterialInstance::IsChildOf(const UMaterialInterface* ParentMaterialInterf
 	Properties of the base material. Can now be overridden by instances.
 */
 
-void UMaterialInstance::GetBasePropertyOverridesHash(FSHAHash& OutHash)const
+uint64 UMaterialInstance::GetBasePropertyOverridesHash() const
 {
 	check(IsInGameThread());
 
 	const UMaterial* Mat = GetMaterial();
 	check(Mat);
 
-	FSHA1 Hash;
+	FXxHash64Builder Hash;
 	bool bHasOverrides = false;
 	
 	auto GetPropertyOverrideHash = [&](auto InstanceValue, auto MatValue, const FString &HashString)
@@ -5045,7 +5045,7 @@ void UMaterialInstance::GetBasePropertyOverridesHash(FSHAHash& OutHash)const
 		
 		if (bOverridden)
 		{
-			Hash.UpdateWithString(*HashString, HashString.Len());
+			Hash.Update(*HashString, HashString.NumBytesWithoutNull());
 			Hash.Update(reinterpret_cast<const uint8*>(&InstanceValue), sizeof(InstanceValue));
 			bHasOverrides = true;
 		}
@@ -5070,8 +5070,11 @@ void UMaterialInstance::GetBasePropertyOverridesHash(FSHAHash& OutHash)const
 
 	if (bHasOverrides)
 	{
-		Hash.Final();
-		Hash.GetHash(&OutHash.Hash[0]);
+		return Hash.Finalize().Hash;
+	}
+	else
+	{
+		return 0u;
 	}
 }
 
diff --git a/Engine/Source/Runtime/Engine/Private/Materials/MaterialShader.cpp b/Engine/Source/Runtime/Engine/Private/Materials/MaterialShader.cpp
--- a/Engine/Source/Runtime/Engine/Private/Materials/MaterialShader.cpp
+++ b/Engine/Source/Runtime/Engine/Private/Materials/MaterialShader.cpp
@@ -1145,7 +1145,7 @@ void FSubstrateCompilationConfig::Append(FShaderKeyGenerator& KeyGen) const
 	}
 }
 
-void FSubstrateCompilationConfig::UpdateHash(FSHA1& Hasher) const
+void FSubstrateCompilationConfig::UpdateHash(FXxHash64Builder& Hasher) const
 {
 	Hasher.Update((const uint8*)(&bFullSimplify), sizeof(bFullSimplify));
 	Hasher.Update((const uint8*)(&BytesPerPixelOverride), sizeof(BytesPerPixelOverride));
@@ -1270,17 +1270,26 @@ void FMaterialShaderMapId::Serialize(FArchive& Ar, bool bLoadingCooked)
 		{
 			FSHAHash LegacyHash;
 			Ar << LegacyHash;
+			TextureReferencesHash = 0u;
 		}
 
 		if (Ar.CustomVer(FRenderingObjectVersion::GUID) >= FRenderingObjectVersion::AddedMaterialExpressionIncludesHash)
 		{
 			Ar << ExpressionIncludesHash;
 		}
+		else
+		{
+			ExpressionIncludesHash = 0u;
+		}
 
 		if (Ar.UEVer() >= VER_UE4_MATERIAL_INSTANCE_BASE_PROPERTY_OVERRIDES)
 		{
 			Ar << BasePropertyOverridesHash;
 		}
+		else
+		{
+			BasePropertyOverridesHash = 0u;
+		}
 
 		if (!bIsLegacyPackage)
 		{
@@ -1293,8 +1302,8 @@ void FMaterialShaderMapId::Serialize(FArchive& Ar, bool bLoadingCooked)
 		{
 			bUsingNewHLSLGenerator = false;
 			PreferredShaderCodeFormat = NAME_None;
-			ExternalCodeReferencesHash = FSHAHash();
-			BaseMaterialPathHash = FSHAHash();
+			ExternalCodeReferencesHash = 0u;
+			BaseMaterialPathHash = 0u;
 		}
 
 		// SUBSTRATE_TODO We do not need to serialize FSubstrateCompilationConfig for now since this is only used when debugging in the editor.
@@ -1307,8 +1316,8 @@ void FMaterialShaderMapId::Serialize(FArchive& Ar, bool bLoadingCooked)
 		if (bIsSavingCooked)
 		{
 			// Saving cooked data, this should be valid
-			GetMaterialHash(CookedShaderMapIdHash);
-			checkf(CookedShaderMapIdHash != FSHAHash(), TEXT("Tried to save an invalid shadermap id hash during cook"));
+			CookedShaderMapIdHash = GetMaterialHash();
+			checkf(CookedShaderMapIdHash != 0u, TEXT("Tried to save an invalid shadermap id hash during cook"));
 		}
 		
 		Ar << CookedShaderMapIdHash;
@@ -1316,7 +1325,7 @@ void FMaterialShaderMapId::Serialize(FArchive& Ar, bool bLoadingCooked)
 #else
 	// Cooked so can assume this is valid
 	Ar << CookedShaderMapIdHash;
-	checkf(CookedShaderMapIdHash != FSHAHash(), TEXT("Loaded an invalid cooked shadermap id hash"));
+	checkf(CookedShaderMapIdHash != 0u, TEXT("Loaded an invalid cooked shadermap id hash"));
 #endif // WITH_EDITOR
 
 	if (!bIsLegacyPackage)
@@ -1334,22 +1343,22 @@ void FMaterialShaderMapId::Serialize(FArchive& Ar, bool bLoadingCooked)
 
 #if WITH_EDITOR
 /** Hashes the material-specific part of this shader map Id. */
-void FMaterialShaderMapId::GetMaterialHash(FSHAHash& OutHash, bool bWithStaticParameters) const
+uint64 FMaterialShaderMapId::GetMaterialHash(bool bWithStaticParameters) const
 {
 	check(IsContentValid());
-	FSHA1 HashState;
+	FXxHash64Builder HashState;
 
 	HashState.Update((const uint8*)&Usage, sizeof(Usage));
 	if (Usage == EMaterialShaderMapUsage::MaterialExportCustomOutput)
 	{
-		HashState.UpdateWithString(*UsageCustomOutput, UsageCustomOutput.Len());
+		HashState.Update(*UsageCustomOutput, UsageCustomOutput.NumBytesWithoutNull());
 	}
 
 	HashState.Update((const uint8*)&BaseMaterialId, sizeof(BaseMaterialId));
 
 	FString QualityLevelString;
 	GetMaterialQualityLevelName(QualityLevel, QualityLevelString);
-	HashState.UpdateWithString(*QualityLevelString, QualityLevelString.Len());
+	HashState.Update(*QualityLevelString, QualityLevelString.NumBytesWithoutNull());
 
 	HashState.Update((const uint8*)&FeatureLevel, sizeof(FeatureLevel));
 
@@ -1408,8 +1417,7 @@ void FMaterialShaderMapId::GetMaterialHash(FSHAHash& OutHash, bool bWithStaticPa
 
 	SubstrateCompilationConfig.UpdateHash(HashState);
 
-	HashState.Final();
-	HashState.GetHash(&OutHash.Hash[0]);
+	return HashState.Finalize().Hash;
 }
 #endif // WITH_EDITOR
 
@@ -1568,6 +1576,11 @@ bool FMaterialShaderMapId::Equals(const FMaterialShaderMapId& ReferenceSet, bool
 			return false;
 		}
 
+		if (ExternalCodeReferencesHash != ReferenceSet.ExternalCodeReferencesHash)
+		{
+			return false;
+		}
+
 		if (BaseMaterialPathHash != ReferenceSet.BaseMaterialPathHash)
 		{
 			return false;
@@ -1967,7 +1980,7 @@ void FMaterialShaderType::BeginCompileShaderPipeline(
  */
 FShader* FMaterialShaderType::FinishCompileShader(
 	const FUniformExpressionSet& UniformExpressionSet,
-	const FSHAHash& MaterialShaderMapHash,
+	uint64 MaterialShaderMapHash,
 	const FShaderCompileJob& CurrentJob,
 	const FShaderPipelineType* ShaderPipelineType,
 	const FString& InDebugDescription
@@ -2044,7 +2057,7 @@ void FMaterialShaderMap::GetAllOutdatedTypes(TArray<const FShaderType*>& Outdate
 
 TRACE_DECLARE_INT_COUNTER(Shaders_FMaterialShaderMapDDCRequests, TEXT("Shaders/FMaterialShaderMap/DDCRequests"));
 TRACE_DECLARE_INT_COUNTER(Shaders_FMaterialShaderMapDDCHits, TEXT("Shaders/FMaterialShaderMap/DDCHits"));
-TRACE_DECLARE_MEMORY_COUNTER(Shaders_FMaterialShaderMapDDCBytesReceived, TEXT("Shaders/FMaterialShaderMap/DDCBytesRecieved"));
+TRACE_DECLARE_MEMORY_COUNTER(Shaders_FMaterialShaderMapDDCBytesReceived, TEXT("Shaders/FMaterialShaderMap/DDCBytesReceived"));
 TRACE_DECLARE_MEMORY_COUNTER(Shaders_FMaterialShaderMapDDCBytesSent, TEXT("Shaders/FMaterialShaderMap/DDCBytesSent"));
 
 void FMaterialShaderMap::LoadFromDerivedDataCache(const FMaterial* Material, const FMaterialShaderMapId& ShaderMapId, EShaderPlatform InPlatform, const ITargetPlatform* TargetPlatform, TRefCountPtr<FMaterialShaderMap>& InOutShaderMap, FString& OutDDCKeyDesc)
@@ -2947,7 +2960,7 @@ static FHashedName GetPreprocessedSourceKey(const FVertexFactoryType* VertexFact
 	}
 }
 
-FShader* FMaterialShaderMap::ProcessCompilationResultsForSingleJob(FShaderCompileJob* SingleJob, const FShaderPipelineType* ShaderPipeline, const FSHAHash& MaterialShaderMapHash)
+FShader* FMaterialShaderMap::ProcessCompilationResultsForSingleJob(FShaderCompileJob* SingleJob, const FShaderPipelineType* ShaderPipeline, uint64 MaterialShaderMapHash)
 {
 	TRACE_CPUPROFILER_EVENT_SCOPE(FMaterialShaderMap::ProcessCompilationResultsForSingleJob);
 
@@ -3034,8 +3047,7 @@ void FMaterialShaderMap::ProcessCompilationResults(const TArray<FShaderCommonCom
 
 	double StartTime = FPlatformTime::Seconds();
 
-	FSHAHash MaterialShaderMapHash;
-	ShaderMapId.GetMaterialHash(MaterialShaderMapHash);
+	uint64 MaterialShaderMapHash = ShaderMapId.GetMaterialHash();
 
 	do
 	{
@@ -3634,10 +3646,11 @@ void FMaterialShaderMap::DumpShaderList(FShaderListReport& Out) const
 
 void FMaterialShaderMap::GetShaderList(TMap<FShaderId, TShaderRef<FShader>>& OutShaders) const
 {
-	FSHAHash MaterialShaderMapHash;
+	uint64 MaterialShaderMapHash;
 #if WITH_EDITOR
-	// TODO
-	ShaderMapId.GetMaterialHash(MaterialShaderMapHash);
+	MaterialShaderMapHash = ShaderMapId.IsCookedId() ? ShaderMapId.CookedShaderMapIdHash : ShaderMapId.GetMaterialHash();
+#else
+	MaterialShaderMapHash = ShaderMapId.CookedShaderMapIdHash;
 #endif
 
 	GetContent()->GetShaderList(*this, MaterialShaderMapHash, OutShaders);
@@ -4114,7 +4127,7 @@ void FMaterialShaderMap::DumpDebugInfo(FOutputDevice& OutputDevice) const
 }
 
 #if WITH_EDITOR
-void FMaterialShaderMap::InitalizeForODSC(EShaderPlatform TargetShaderPlatform, const FMaterialCompilationOutput& NewCompilationOutput)
+void FMaterialShaderMap::InitializeForODSC(EShaderPlatform TargetShaderPlatform, const FMaterialCompilationOutput& NewCompilationOutput)
 {
 	// Empty Content
 	FMaterialShaderMapContent* NewContent = new FMaterialShaderMapContent(TargetShaderPlatform);
diff --git a/Engine/Source/Runtime/Engine/Private/Materials/MaterialShared.cpp b/Engine/Source/Runtime/Engine/Private/Materials/MaterialShared.cpp
--- a/Engine/Source/Runtime/Engine/Private/Materials/MaterialShared.cpp
+++ b/Engine/Source/Runtime/Engine/Private/Materials/MaterialShared.cpp
@@ -775,7 +775,7 @@ void FMaterial::BuildShaderMapIdOverride(const FBuildShaderMapIdArgs& Args) cons
 			GetDependentShaderAndVFTypes(OutId.LayoutParams, ShaderTypes, ShaderPipelineTypes, VFTypes);
 		}
 
-		OutId.BaseMaterialPathHash = FSHA1::HashBuffer(Args.BaseMaterialPath.GetCharArray().GetData(), Args.BaseMaterialPath.NumBytesWithoutNull());
+		OutId.BaseMaterialPathHash = FXxHash64::HashBuffer(Args.BaseMaterialPath.GetCharArray().GetData(), Args.BaseMaterialPath.NumBytesWithoutNull()).Hash;
 		OutId.Usage = GetShaderMapUsage();
 		OutId.bUsingNewHLSLGenerator = IsUsingNewHLSLGenerator();
 		OutId.PreferredShaderCodeFormat = FDataDrivenShaderPlatformInfo::GetPreferredShaderCodeFormat(Args.Platform);
@@ -783,9 +783,9 @@ void FMaterial::BuildShaderMapIdOverride(const FBuildShaderMapIdArgs& Args) cons
 		OutId.QualityLevel = GetQualityLevel();
 		OutId.FeatureLevel = GetFeatureLevel();
 		OutId.SetShaderDependencies(ShaderTypes, ShaderPipelineTypes, VFTypes, Args.Platform);
-		GetReferencedTexturesHash(OutId.TextureReferencesHash);
-		GetExpressionIncludesHash(OutId.ExpressionIncludesHash);
-		GetExternalCodeReferencesHash(OutId.ExternalCodeReferencesHash);
+		OutId.TextureReferencesHash = GetReferencedTexturesHash();
+		OutId.ExpressionIncludesHash = GetExpressionIncludesHash();
+		OutId.ExternalCodeReferencesHash = GetExternalCodeReferencesHash();
 
 		OutId.SubstrateCompilationConfig = GetSubstrateCompilationConfig();
 
@@ -4648,9 +4648,9 @@ void FMaterial::GetDependentShaderAndVFTypes(const FPlatformTypeLayoutParameters
 }
 
 #if WITH_EDITOR
-void FMaterial::GetReferencedTexturesHash(FSHAHash& OutHash) const
+uint64 FMaterial::GetReferencedTexturesHash() const
 {
-	FSHA1 HashState;
+	FXxHash64Builder HashState;
 
 	const TArray<TObjectPtr<UObject>>& ReferencedTextures = GetMatDefaultTextures();
 	// Hash the names of the uniform expression textures to capture changes in their order or values resulting from material compiler code changes
@@ -4663,7 +4663,7 @@ void FMaterial::GetReferencedTexturesHash(FSHAHash& OutHash) const
 			TextureName = ReferencedTextures[TextureIndex]->GetName();
 		}
 
-		HashState.UpdateWithString(*TextureName, TextureName.Len());
+		HashState.Update(*TextureName, TextureName.NumBytesWithoutNull());
 	}
 
 	UMaterialShaderQualitySettings* MaterialShaderQualitySettings = UMaterialShaderQualitySettings::Get();
@@ -4672,13 +4672,12 @@ void FMaterial::GetReferencedTexturesHash(FSHAHash& OutHash) const
 		MaterialShaderQualitySettings->GetShaderPlatformQualitySettings(GetShaderPlatform())->AppendToHashState(QualityLevel, HashState);
 	}
 
-	HashState.Final();
-	HashState.GetHash(&OutHash.Hash[0]);
+	return HashState.Finalize().Hash;
 }
 
-void FMaterial::GetExpressionIncludesHash(FSHAHash& OutHash) const
+uint64 FMaterial::GetExpressionIncludesHash() const
 {
-	FSHA1 HashState;
+	FXxHash64Builder HashState;
 
 	for (const FString& ExpressionIncludeFilePath : GetCachedExpressionData().EditorOnlyData->ExpressionIncludeFilePaths)
 	{
@@ -4688,12 +4687,12 @@ void FMaterial::GetExpressionIncludesHash(FSHAHash& OutHash) const
 		HashState.Update(FileHash.Hash, UE_ARRAY_COUNT(FileHash.Hash));
 	}
 
-	OutHash = HashState.Finalize();
+	return HashState.Finalize().Hash;
 }
 
-void FMaterial::GetExternalCodeReferencesHash(FSHAHash& OutHash) const
+uint64 FMaterial::GetExternalCodeReferencesHash() const
 {
-	GetCachedExpressionData().GetExternalCodeReferencesHash(OutHash);
+	return GetCachedExpressionData().GetExternalCodeReferencesHash();
 }
 #endif // WITH_EDITOR
 
diff --git a/Engine/Source/Runtime/Engine/Private/Materials/MeshMaterialShader.cpp b/Engine/Source/Runtime/Engine/Private/Materials/MeshMaterialShader.cpp
--- a/Engine/Source/Runtime/Engine/Private/Materials/MeshMaterialShader.cpp
+++ b/Engine/Source/Runtime/Engine/Private/Materials/MeshMaterialShader.cpp
@@ -150,14 +150,9 @@ static inline FString GetJobName(const FShaderCompileJob* SingleJob, const FShad
 	return String;
 }
 
-/**
- * Either creates a new instance of this type or returns an equivalent existing shader.
- * @param Material - The material to link the shader with.
- * @param CurrentJob - Compile job that was enqueued by BeginCompileShader.
- */
 FShader* FMeshMaterialShaderType::FinishCompileShader( 
 	const FUniformExpressionSet& UniformExpressionSet, 
-	const FSHAHash& MaterialShaderMapHash,
+	uint64 MaterialShaderMapHash,
 	const FShaderCompileJob& CurrentJob,
 	const FShaderPipelineType* ShaderPipelineType,
 	const FString& InDebugDescription) const
diff --git a/Engine/Source/Runtime/Engine/Private/ShaderCompiler/ShaderCompilerEditor.cpp b/Engine/Source/Runtime/Engine/Private/ShaderCompiler/ShaderCompilerEditor.cpp
--- a/Engine/Source/Runtime/Engine/Private/ShaderCompiler/ShaderCompilerEditor.cpp
+++ b/Engine/Source/Runtime/Engine/Private/ShaderCompiler/ShaderCompilerEditor.cpp
@@ -158,13 +158,10 @@ FShader* FGlobalShaderTypeCompiler::FinishCompileShader(const FGlobalShaderType*
 		}
 
 		// Create the global shader map hash
-		FSHAHash GlobalShaderMapHash;
+		uint64 GlobalShaderMapHash;
 		{
-			FSHA1 HashState;
-			const TCHAR* GlobalShaderString = TEXT("GlobalShaderMap");
-			HashState.UpdateWithString(GlobalShaderString, FCString::Strlen(GlobalShaderString));
-			HashState.Final();
-			HashState.GetHash(&GlobalShaderMapHash.Hash[0]);
+			FStringView GlobalShaderString = TEXTVIEW("GlobalShaderMap");
+			GlobalShaderMapHash = FXxHash64::HashBuffer(GlobalShaderString.GetData(), GlobalShaderString.NumBytes()).Hash;
 		}
 
 		Shader = ShaderType->ConstructCompiled(FGlobalShaderType::CompiledShaderInitializerType(ShaderType, nullptr, CurrentJob.Key.PermutationId, CurrentJob.Output, GlobalShaderMapHash, ShaderPipelineType, nullptr));
diff --git a/Engine/Source/Runtime/Engine/Public/MaterialCachedData.h b/Engine/Source/Runtime/Engine/Public/MaterialCachedData.h
--- a/Engine/Source/Runtime/Engine/Public/MaterialCachedData.h
+++ b/Engine/Source/Runtime/Engine/Public/MaterialCachedData.h
@@ -286,7 +286,7 @@ struct FMaterialCachedExpressionData
 	void AppendReferencedParameterCollectionIdsTo(TArray<FGuid>& OutIds) const;
 
 	/** Returns a hash over all external code references from this cache. */
-	void GetExternalCodeReferencesHash(FSHAHash& OutHash) const;
+	uint64 GetExternalCodeReferencesHash() const;
 
 	bool IsPropertyConnected(EMaterialProperty Property) const
 	{
diff --git a/Engine/Source/Runtime/Engine/Public/MaterialShaderType.h b/Engine/Source/Runtime/Engine/Public/MaterialShaderType.h
--- a/Engine/Source/Runtime/Engine/Public/MaterialShaderType.h
+++ b/Engine/Source/Runtime/Engine/Public/MaterialShaderType.h
@@ -104,7 +104,7 @@ public:
 			int32 InPermutationId,
 			const FShaderCompilerOutput& CompilerOutput,
 			const FUniformExpressionSet& InUniformExpressionSet,
-			const FSHAHash& InMaterialShaderMapHash,
+			uint64 InMaterialShaderMapHash,
 			const FShaderPipelineType* InShaderPipeline,
 			const FVertexFactoryType* InVertexFactoryType,
 			const FString& InDebugDescription
@@ -167,7 +167,7 @@ public:
 	 */
 	FShader* FinishCompileShader(
 		const FUniformExpressionSet& UniformExpressionSet,
-		const FSHAHash& MaterialShaderMapHash,
+		uint64 MaterialShaderMapHash,
 		const FShaderCompileJob& CurrentJob,
 		const FShaderPipelineType* ShaderPipeline,
 		const FString& InDebugDescription
diff --git a/Engine/Source/Runtime/Engine/Public/MaterialShared.h b/Engine/Source/Runtime/Engine/Public/MaterialShared.h
--- a/Engine/Source/Runtime/Engine/Public/MaterialShared.h
+++ b/Engine/Source/Runtime/Engine/Public/MaterialShared.h
@@ -457,7 +457,7 @@ public:
 	friend bool operator==(const FMaterialCacheTagStack& Lhs, const FMaterialCacheTagStack& Rhs)
 	{
 		// Note: ShaderFrequencyMask is intentionally ignored — it only affects the editor-side decision to include the parameter in the VS compact uniform buffer.
-		return Lhs.TagGuid != Rhs.TagGuid;
+		return Lhs.TagGuid == Rhs.TagGuid;
 	}
 
 	friend inline bool operator!=(const FMaterialCacheTagStack& Lhs, const FMaterialCacheTagStack& Rhs)
@@ -1262,7 +1262,7 @@ public:
 class FMaterialShaderMapId
 {
 public:
-	FSHAHash CookedShaderMapIdHash;
+	uint64 CookedShaderMapIdHash = 0u;
 
 #if WITH_EDITOR
 	/** 
@@ -1325,19 +1325,19 @@ public:
 	 * This is stored in the shader map Id to gracefully handle situations where code changes
 	 * that generates the array of textures that the uniform expressions use to link up after being loaded from the DDC.
 	 */
-	FSHAHash TextureReferencesHash;
+	uint64 TextureReferencesHash = 0u;
 
 	/** A hash of the content of files (dynamically) included by material expressions. */
-	FSHAHash ExpressionIncludesHash;
+	uint64 ExpressionIncludesHash = 0u;
 	
 	/** A hash of the base property overrides for this material instance. */
-	FSHAHash BasePropertyOverridesHash;
+	uint64 BasePropertyOverridesHash = 0u;
 
 	/** A hash of all external code declarations referenced in the shader map. */
-	FSHAHash ExternalCodeReferencesHash;
+	uint64 ExternalCodeReferencesHash = 0u;
 
 	/** A hash of the base material's full path */
-	FSHAHash BaseMaterialPathHash;
+	uint64 BaseMaterialPathHash = 0u;
 
 	/** Is the material using the new HLSL generator? */
 	bool bUsingNewHLSLGenerator;
@@ -1389,19 +1389,19 @@ public:
 	bool IsValid() const
 	{
 #if WITH_EDITOR
-		return !IsCookedId() ? BaseMaterialId.IsValid() : (CookedShaderMapIdHash != FSHAHash());
+		return !IsCookedId() ? BaseMaterialId.IsValid() : (CookedShaderMapIdHash != 0u);
 #else
-		return (CookedShaderMapIdHash != FSHAHash());
+		return (CookedShaderMapIdHash != 0u);
 #endif
 	}
 
 	friend uint32 GetTypeHash(const FMaterialShaderMapId& Ref)
 	{
 #if WITH_EDITOR
-		return !Ref.IsCookedId() ? Ref.BaseMaterialId.A : (*(uint32*)&Ref.CookedShaderMapIdHash.Hash[0]);
+		return !Ref.IsCookedId() ? Ref.BaseMaterialId.A : (uint32)(Ref.CookedShaderMapIdHash);
 #else
-		// Using the hash value directly instead of FSHAHash CRC as fairly uniform distribution
-		return *(uint32*)&Ref.CookedShaderMapIdHash.Hash[0];
+		// Using the a truncated form of the cooked shadermap hash value directly, this is fine assuming a good quality hash
+		return (uint32)Ref.CookedShaderMapIdHash;
 #endif
 	}
 
@@ -1420,11 +1420,18 @@ public:
 
 #if WITH_EDITOR
 	/** Hashes the material-specific part of this shader map Id. */
-	ENGINE_API void GetMaterialHash(FSHAHash& OutHash, bool bWithStaticParameters) const;
+	UE_DEPRECATED(5.8, "Use version returning uint64 (SHA1 hashes no longer used)")
+	inline void GetMaterialHash(FSHAHash& OutHash, bool bWithStaticParameters) const 
+	{ 
+		uint64 NewHash = GetMaterialHash(bWithStaticParameters);
+		OutHash = FSHA1::HashBuffer(&NewHash, sizeof(NewHash));
+	}
+
+	ENGINE_API uint64 GetMaterialHash(bool bWithStaticParameters) const;
 
-	void GetMaterialHash(FSHAHash& OutHash) const
+	uint64 GetMaterialHash() const
 	{
-		GetMaterialHash(OutHash, true);
+		return GetMaterialHash(true);
 	}
 #endif
 
@@ -1874,7 +1881,7 @@ public:
 	void DumpDebugInfo(FOutputDevice& OutputDevice) const;
 
 #if WITH_EDITOR
-	void InitalizeForODSC(EShaderPlatform TargetShaderPlatform, const FMaterialCompilationOutput& NewCompilationOutput);
+	void InitializeForODSC(EShaderPlatform TargetShaderPlatform, const FMaterialCompilationOutput& NewCompilationOutput);
 
 	void AppendCompileStateDebugInfo(FStringBuilderBase& OutDebugInfo) const;
 #endif
@@ -1956,7 +1963,7 @@ private:
 	uint32 bIsFromODSC : 1;
 
 #if WITH_EDITOR
-	FShader* ProcessCompilationResultsForSingleJob(class FShaderCompileJob* SingleJob, const FShaderPipelineType* ShaderPipeline, const FSHAHash& MaterialShaderMapHash);
+	FShader* ProcessCompilationResultsForSingleJob(class FShaderCompileJob* SingleJob, const FShaderPipelineType* ShaderPipeline, uint64 MaterialShaderMapHash);
 #endif
 
 #if WITH_EDITOR && ALLOW_SHADERMAP_DEBUG_DATA
@@ -3137,11 +3144,11 @@ private:
 	ENGINE_API TShaderRef<FShader> GetShader(class FMeshMaterialShaderType* ShaderType, FVertexFactoryType* VertexFactoryType, int32 PermutationId, bool bFatalIfMissing = true) const;
 
 #if WITH_EDITOR
-	void GetReferencedTexturesHash(FSHAHash& OutHash) const;
+	uint64 GetReferencedTexturesHash() const;
 
-	void GetExpressionIncludesHash(FSHAHash& OutHash) const;
+	uint64 GetExpressionIncludesHash() const;
 
-	void GetExternalCodeReferencesHash(FSHAHash& OutHash) const;
+	uint64 GetExternalCodeReferencesHash() const;
 #endif // WITH_EDITOR
 
 	friend class FMaterialShaderMap;
diff --git a/Engine/Source/Runtime/Engine/Public/Materials/Material.h b/Engine/Source/Runtime/Engine/Public/Materials/Material.h
--- a/Engine/Source/Runtime/Engine/Public/Materials/Material.h
+++ b/Engine/Source/Runtime/Engine/Public/Materials/Material.h
@@ -1742,11 +1742,9 @@ private:
 	void BeginCacheShadersForResources(EShaderPlatform ShaderPlatform, const TArray<FMaterialResource*>& ResourcesToCache, EMaterialShaderPrecompileMode PrecompileMode = EMaterialShaderPrecompileMode::Default, const ITargetPlatform* TargetPlatform = nullptr);
 
 	/**
-	 * If there is some texture reference used by a TextureProperty node in any expressions, this function
-	 * will extract the current hash of TextureReferencesHash into a string  then append the texture guid used by the node
-	 * and recompute a new hash.
+	 * Retrieves a hash of all texture references which should force a recompile of this material
 	 */
-	void GetForceRecompileTextureIdsHash(FSHAHash &TextureReferencesHash);
+	uint64 GetForceRecompileTextureIdsHash() const;
 #endif // WITH_EDITOR
 
 public:
diff --git a/Engine/Source/Runtime/Engine/Public/Materials/MaterialInstance.h b/Engine/Source/Runtime/Engine/Public/Materials/MaterialInstance.h
--- a/Engine/Source/Runtime/Engine/Public/Materials/MaterialInstance.h
+++ b/Engine/Source/Runtime/Engine/Public/Materials/MaterialInstance.h
@@ -1128,8 +1128,14 @@ public:
 	ENGINE_API virtual void UpdateCachedData();
 #endif
 
-	ENGINE_API void GetBasePropertyOverridesHash(FSHAHash& OutHash)const;
-	ENGINE_API virtual bool HasOverridenBaseProperties()const;
+	UE_DEPRECATED(5.8, "Use parameterless version returning a uint64; SHA1 hashes are no longer used")
+	inline void GetBasePropertyOverridesHash(FSHAHash& OutHash) const
+	{
+		uint64 NewHash = GetBasePropertyOverridesHash();
+		OutHash = FSHA1::HashBuffer(&NewHash, sizeof(NewHash));
+	}
+	ENGINE_API uint64 GetBasePropertyOverridesHash() const;
+	ENGINE_API virtual bool HasOverridenBaseProperties() const;
 
 #if WITH_EDITOR
 	ENGINE_API FString GetBasePropertyOverrideString() const;
diff --git a/Engine/Source/Runtime/Engine/Public/Materials/MaterialLayersFunctions.h b/Engine/Source/Runtime/Engine/Public/Materials/MaterialLayersFunctions.h
--- a/Engine/Source/Runtime/Engine/Public/Materials/MaterialLayersFunctions.h
+++ b/Engine/Source/Runtime/Engine/Public/Materials/MaterialLayersFunctions.h
@@ -13,6 +13,7 @@
 
 class FArchive;
 class FShaderKeyGenerator;
+class FXxHash64Builder;
 struct FMaterialLayersFunctions;
 class UMaterialExpressionMaterialFunctionCall;
 
@@ -44,7 +45,7 @@ struct FMaterialLayersFunctionsID
 		return Ref;
 	}
 
-	void UpdateHash(FSHA1& HashState) const;
+	void UpdateHash(FXxHash64Builder& HashState) const;
 
 	//TODO: Investigate whether this is really required given it is only used by FMaterialShaderMapId AND that one also uses UpdateHash
 	void AppendKeyString(FString& KeyString) const;
diff --git a/Engine/Source/Runtime/Engine/Public/Materials/MaterialParameters.h b/Engine/Source/Runtime/Engine/Public/Materials/MaterialParameters.h
--- a/Engine/Source/Runtime/Engine/Public/Materials/MaterialParameters.h
+++ b/Engine/Source/Runtime/Engine/Public/Materials/MaterialParameters.h
@@ -11,6 +11,7 @@
 #include "MaterialParameters.generated.h"
 
 class FShaderKeyGenerator;
+class FXxHash64Builder;
 class UTexture;
 class UTextureCollection;
 class UCurveLinearColor;
@@ -506,7 +507,7 @@ struct FSubstrateCompilationConfig
 	FString GetShaderMapKeyString() const;
 	void Append(FShaderKeyGenerator& KeyGen) const;
 
-	void UpdateHash(FSHA1& Hasher) const;
+	void UpdateHash(FXxHash64Builder& Hasher) const;
 
 	void Serialize(FArchive& Ar);
 #endif
diff --git a/Engine/Source/Runtime/Engine/Public/MeshMaterialShaderType.h b/Engine/Source/Runtime/Engine/Public/MeshMaterialShaderType.h
--- a/Engine/Source/Runtime/Engine/Public/MeshMaterialShaderType.h
+++ b/Engine/Source/Runtime/Engine/Public/MeshMaterialShaderType.h
@@ -33,7 +33,7 @@ public:
 			int32 PermutationId,
 			const FShaderCompilerOutput& CompilerOutput,
 			const FUniformExpressionSet& InUniformExpressionSet,
-			const FSHAHash& InMaterialShaderMapHash,
+			uint64 InMaterialShaderMapHash,
 			const FString& InDebugDescription,
 			const FShaderPipelineType* InShaderPipeline,
 			const FVertexFactoryType* InVertexFactoryType
@@ -92,7 +92,7 @@ public:
 	 */
 	FShader* FinishCompileShader(
 		const FUniformExpressionSet& UniformExpressionSet, 
-		const FSHAHash& MaterialShaderMapHash,
+		uint64 MaterialShaderMapHash,
 		const FShaderCompileJob& CurrentJob,
 		const FShaderPipelineType* ShaderPipeline,
 		const FString& InDebugDescription
diff --git a/Engine/Source/Runtime/Engine/Public/StaticParameterSet.h b/Engine/Source/Runtime/Engine/Public/StaticParameterSet.h
--- a/Engine/Source/Runtime/Engine/Public/StaticParameterSet.h
+++ b/Engine/Source/Runtime/Engine/Public/StaticParameterSet.h
@@ -3,15 +3,15 @@
 #pragma once
 
 #include "CoreMinimal.h"
-#include "UObject/NameTypes.h"
+#include "Hash/xxhash.h"
+#include "Materials/MaterialParameters.h"
+#include "Materials/MaterialLayersFunctions.h"
 #include "Misc/Guid.h"
 #include "UObject/FortniteMainBranchObjectVersion.h"
+#include "UObject/NameTypes.h"
 #include "UObject/RenderingObjectVersion.h"
 #include "UObject/ReleaseObjectVersion.h"
 #include "UObject/UE5ReleaseStreamObjectVersion.h"
-#include "UObject/FortniteMainBranchObjectVersion.h"
-#include "Materials/MaterialParameters.h"
-#include "Materials/MaterialLayersFunctions.h"
 #include "StaticParameterSet.generated.h"
 
 class FShaderKeyGenerator;
@@ -68,7 +68,8 @@ struct FStaticParameterBase
 		Ar << ExpressionGUID;
 	}
 
-	void UpdateHash(FSHA1& HashState) const
+	template<typename THasher>
+	void UpdateHash(THasher& HashState) const
 	{
 		const FString ParameterName = ParameterInfo.ToString();
 		HashState.Update((const uint8*)*ParameterName, ParameterName.Len() * sizeof(TCHAR));
@@ -123,7 +124,7 @@ struct FStaticSwitchParameter : public FStaticParameterBase
 		return Ar;
 	}
 
-	void UpdateHash(FSHA1& HashState) const
+	void UpdateHash(FXxHash64Builder& HashState) const
 	{
 		FStaticParameterBase::UpdateHash(HashState);
 		uint8 HashValue = Value;
@@ -200,7 +201,10 @@ struct FStaticComponentMaskParameter : public FStaticParameterBase
 		return Ar;
 	}
 
-	void UpdateHash(FSHA1& HashState) const
+	UE_DEPRECATED(5.8, "Use overload accepting an FXxHash64Builder; SHA1 hashes are no longer used")
+	void UpdateHash(FSHA1& HashState) const {}
+
+	void UpdateHash(FXxHash64Builder& HashState) const
 	{
 		FStaticParameterBase::UpdateHash(HashState);
 		uint8 Values[4];
@@ -338,7 +342,10 @@ PRAGMA_DISABLE_DEPRECATION_WARNINGS
 PRAGMA_ENABLE_DEPRECATION_WARNINGS
 	}
 
-	void UpdateHash(FSHA1& HashState) const
+	UE_DEPRECATED(5.8, "Use overload accepting an FXxHash64Builder; SHA1 hashes are no longer used")
+	void UpdateHash(FSHA1& HashState) const {}
+
+	void UpdateHash(FXxHash64Builder& HashState) const
 	{
 		const FString LayerNameString = LayerName.ToString();
 		HashState.Update((const uint8*)*LayerNameString, LayerNameString.Len() * sizeof(TCHAR));
diff --git a/Engine/Source/Runtime/MaterialShaderQualitySettings/Classes/ShaderPlatformQualitySettings.h b/Engine/Source/Runtime/MaterialShaderQualitySettings/Classes/ShaderPlatformQualitySettings.h
--- a/Engine/Source/Runtime/MaterialShaderQualitySettings/Classes/ShaderPlatformQualitySettings.h
+++ b/Engine/Source/Runtime/MaterialShaderQualitySettings/Classes/ShaderPlatformQualitySettings.h
@@ -90,8 +90,23 @@ public:
 	}
 
 	MATERIALSHADERQUALITYSETTINGS_API const FMaterialQualityOverrides& GetQualityOverrides(EMaterialQualityLevel::Type QualityLevel) const;
-	MATERIALSHADERQUALITYSETTINGS_API void BuildHash(EMaterialQualityLevel::Type QualityLevel, class FSHAHash& OutHash) const;
-	MATERIALSHADERQUALITYSETTINGS_API void AppendToHashState(EMaterialQualityLevel::Type QualityLevel, class FSHA1& HashState) const;
+	UE_DEPRECATED(5.8, "UShaderPlatformQualitySettings::BuildHash is unused in engine code and has been deprecated; use AppendToHashState instead if needed.")
+	inline void BuildHash(EMaterialQualityLevel::Type QualityLevel, class FSHAHash& OutHash) const 
+	{
+		FSHA1 Builder;
+		AppendToHashState(QualityLevel, Builder);
+		OutHash = Builder.Finalize();
+	}
+	
+	template <typename THasher>
+	void AppendToHashState(EMaterialQualityLevel::Type QualityLevel, THasher& HashState) const
+	{
+		// since this hashes the full size of the raw struct, we must ensure there is no padding
+		static_assert(std::has_unique_object_representations_v<FMaterialQualityOverrides>);
+
+		const FMaterialQualityOverrides& QualityLevelOverrides = GetQualityOverrides(QualityLevel);
+		HashState.Update((const uint8*)&QualityLevelOverrides, sizeof(QualityLevelOverrides));
+	}
 	
 	virtual const TCHAR* GetConfigOverridePlatform() const override
 	{
diff --git a/Engine/Source/Runtime/MaterialShaderQualitySettings/Private/MaterialShaderQualitySettings.cpp b/Engine/Source/Runtime/MaterialShaderQualitySettings/Private/MaterialShaderQualitySettings.cpp
--- a/Engine/Source/Runtime/MaterialShaderQualitySettings/Private/MaterialShaderQualitySettings.cpp
+++ b/Engine/Source/Runtime/MaterialShaderQualitySettings/Private/MaterialShaderQualitySettings.cpp
@@ -1,13 +1,14 @@
 // Copyright Epic Games, Inc. All Rights Reserved.
 
 #include "MaterialShaderQualitySettings.h"
-#include "UObject/Package.h"
-#include "ShaderPlatformQualitySettings.h"
-#include "Misc/SecureHash.h"
-#include "RHI.h"
+
+#include "DataDrivenShaderPlatformInfo.h"
 #include "Interfaces/ITargetPlatform.h"
 #include "Interfaces/ITargetPlatformManagerModule.h"
-#include "DataDrivenShaderPlatformInfo.h"
+#include "Hash/xxhash.h"
+#include "RHI.h"
+#include "ShaderPlatformQualitySettings.h"
+#include "UObject/Package.h"
 
 #include UE_INLINE_GENERATED_CPP_BY_NAME(MaterialShaderQualitySettings)
 
@@ -129,22 +130,6 @@ const FMaterialQualityOverrides& UShaderPlatformQualitySettings::GetQualityOverr
 	return (QualityLevel < EMaterialQualityLevel::Num) ? QualityOverrides[(int32)QualityLevel] : DefaultQualityOverrides;
 }
 
-void UShaderPlatformQualitySettings::BuildHash(EMaterialQualityLevel::Type QualityLevel, FSHAHash& OutHash) const
-{
-	FSHA1 Hash;
-
-	AppendToHashState(QualityLevel, Hash);
-
-	Hash.Final();
-	Hash.GetHash(&OutHash.Hash[0]);
-}
-
-void UShaderPlatformQualitySettings::AppendToHashState(EMaterialQualityLevel::Type QualityLevel, FSHA1& HashState) const
-{
-	const FMaterialQualityOverrides& QualityLevelOverrides = GetQualityOverrides(QualityLevel);
-	HashState.Update((const uint8*)&QualityLevelOverrides, sizeof(QualityLevelOverrides));
-}
-
 //////////////////////////////////////////////////////////////////////////
 
 bool FMaterialQualityOverrides::CanOverride(EShaderPlatform ShaderPlatform) const
diff --git a/Engine/Source/Runtime/RenderCore/Private/Shader.cpp b/Engine/Source/Runtime/RenderCore/Private/Shader.cpp
--- a/Engine/Source/Runtime/RenderCore/Private/Shader.cpp
+++ b/Engine/Source/Runtime/RenderCore/Private/Shader.cpp
@@ -741,7 +741,7 @@ FShaderCompiledShaderInitializerType::FShaderCompiledShaderInitializerType(
 	const FShaderType::FParameters* InParameters,
 	int32 InPermutationId,
 	const FShaderCompilerOutput& CompilerOutput,
-	const FSHAHash& InMaterialShaderMapHash,
+	uint64 InMaterialShaderMapHash,
 	const FShaderPipelineType* InShaderPipeline,
 	const FVertexFactoryType* InVertexFactoryType
 	) 
diff --git a/Engine/Source/Runtime/RenderCore/Private/ShaderMap.cpp b/Engine/Source/Runtime/RenderCore/Private/ShaderMap.cpp
--- a/Engine/Source/Runtime/RenderCore/Private/ShaderMap.cpp
+++ b/Engine/Source/Runtime/RenderCore/Private/ShaderMap.cpp
@@ -728,7 +728,7 @@ void FShaderMapContent::DumpShaderList(const FShaderMapBase& InShaderMap, FShade
 }
 #endif // !UE_BUILD_SHIPPING
 
-void FShaderMapContent::GetShaderList(const FShaderMapBase& InShaderMap, const FSHAHash& InMaterialShaderMapHash, TMap<FShaderId, TShaderRef<FShader>>& OutShaders) const
+void FShaderMapContent::GetShaderList(const FShaderMapBase& InShaderMap, uint64 InMaterialShaderMapHash, TMap<FShaderId, TShaderRef<FShader>>& OutShaders) const
 {
 	for (int32 ShaderIndex = 0; ShaderIndex < Shaders.Num(); ++ShaderIndex)
 	{
diff --git a/Engine/Source/Runtime/RenderCore/Public/Shader.h b/Engine/Source/Runtime/RenderCore/Public/Shader.h
--- a/Engine/Source/Runtime/RenderCore/Public/Shader.h
+++ b/Engine/Source/Runtime/RenderCore/Public/Shader.h
@@ -670,7 +670,7 @@ class FShaderId
 {
 public:
 	inline FShaderId() {}
-	inline FShaderId(const FShaderType* InType, const FSHAHash& InMaterialShaderMapHash, const FHashedName& InShaderPipeline, const FVertexFactoryType* InVertexFactoryType, int32 InPermutationId, EShaderPlatform InPlatform)
+	inline FShaderId(const FShaderType* InType, uint64 InMaterialShaderMapHash, const FHashedName& InShaderPipeline, const FVertexFactoryType* InVertexFactoryType, int32 InPermutationId, EShaderPlatform InPlatform)
 		: Type(InType)
 		, VFType(InVertexFactoryType)
 		, ShaderPipelineName(InShaderPipeline)
@@ -682,7 +682,7 @@ public:
 	const FShaderType* Type;
 	const FVertexFactoryType* VFType;
 	FHashedName ShaderPipelineName;
-	FSHAHash MaterialShaderMapHash;
+	uint64 MaterialShaderMapHash;
 	int32 PermutationId;
 	uint32 Platform : SP_NumBits;
 
@@ -1608,7 +1608,7 @@ struct FShaderCompiledShaderInitializerType
 	TConstArrayView<uint8> Code;
 	const FShaderParameterMap& ParameterMap;
 	const FSHAHash& OutputHash;
-	FSHAHash MaterialShaderMapHash;
+	uint64 MaterialShaderMapHash;
 	const FShaderPipelineType* ShaderPipeline;
 	const FVertexFactoryType* VertexFactoryType;
 	uint32 NumInstructions;
@@ -1622,7 +1622,7 @@ struct FShaderCompiledShaderInitializerType
 		const FShaderType::FParameters* InParameters,
 		int32 InPermutationId,
 		const FShaderCompilerOutput& CompilerOutput,
-		const FSHAHash& InMaterialShaderMapHash,
+		uint64 InMaterialShaderMapHash,
 		const FShaderPipelineType* InShaderPipeline,
 		const FVertexFactoryType* InVertexFactoryType
 	);
@@ -2423,9 +2423,17 @@ public:
 	RENDERCORE_API void DumpShaderList(const FShaderMapBase& InShaderMap, FShaderListReport& Out) const;
 #endif // !UE_BUILD_SHIPPING
 
-	/** Builds a list of the shaders in a shader map. */
-	RENDERCORE_API void GetShaderList(const FShaderMapBase& InShaderMap, const FSHAHash& InMaterialShaderMapHash, TMap<FShaderId, TShaderRef<FShader>>& OutShaders) const;
 
+	UE_DEPRECATED(5.8, "Use overload accepting uint64 shadermap hash")
+	inline void GetShaderList(const FShaderMapBase& InShaderMap, const FSHAHash& InMaterialShaderMapHash, TMap<FShaderId, TShaderRef<FShader>>& OutShaders) const
+	{
+		uint64 TruncatedShaHash;
+		memcpy(&TruncatedShaHash, InMaterialShaderMapHash.Hash, sizeof(uint64));
+		GetShaderList(InShaderMap, TruncatedShaHash, OutShaders);
+	}
+	
+	/** Builds a list of the shaders in a shader map. */
+	RENDERCORE_API void GetShaderList(const FShaderMapBase& InShaderMap, uint64 InMaterialShaderMapHash, TMap<FShaderId, TShaderRef<FShader>>& OutShaders) const;
 	/** Builds a list of the shaders in a shader map. Key is FShaderType::TypeName */
 	RENDERCORE_API void GetShaderList(const FShaderMapBase& InShaderMap, TMap<FHashedName, TShaderRef<FShader>>& OutShaders) const;
 
