[Subsurface Profile] Bug fix. Factor out the extra surface albedo multiply from the transmission profile to the transmission tint color for old subsurface profile to keep the same look. Bump the subsurface profile version. All newly created profiles match more to pbr if surface albedo matches the base color.

#Synced-CL 51558978

diff --git a/Engine/Source/Runtime/Engine/Private/Rendering/BurleyNormalizedSSS.cpp b/Engine/Source/Runtime/Engine/Private/Rendering/BurleyNormalizedSSS.cpp
--- a/Engine/Source/Runtime/Engine/Private/Rendering/BurleyNormalizedSSS.cpp
+++ b/Engine/Source/Runtime/Engine/Private/Rendering/BurleyNormalizedSSS.cpp
@@ -201,8 +201,8 @@ void ComputeTransmissionProfileBurley(FLinearColor* TargetBuffer, uint32 TargetB
 		const float DistanceInMm = i * InvSize * (MaxTransmissionProfileDistance * CmToMm) * InvUnitScale;
 		const float OffsetInMm = (ProfileRadiusOffset * CmToMm) * InvUnitScale;
 
-
-		FLinearColor TransmissionProfile = Burley_TransmissionProfile(DistanceInMm + OffsetInMm, SurfaceAlbedo, ScalingFactor, DiffuseMeanFreePathInMm);
+		// The profile does not contain surface albedo, the SSS recombine pass will use stored base color as it.
+		FLinearColor TransmissionProfile = Burley_TransmissionProfile(DistanceInMm + OffsetInMm, FLinearColor::White, ScalingFactor, DiffuseMeanFreePathInMm);
 		TargetBuffer[i] = TransmissionProfile * TransmissionTintColor; // Apply tint to the profile
 		//Use Luminance of scattering as SSSS shadow.
 		TargetBuffer[i].A = exp(-DistanceInMm * ExtinctionScale);
diff --git a/Engine/Source/Runtime/Engine/Private/Rendering/SubsurfaceProfile.cpp b/Engine/Source/Runtime/Engine/Private/Rendering/SubsurfaceProfile.cpp
--- a/Engine/Source/Runtime/Engine/Private/Rendering/SubsurfaceProfile.cpp
+++ b/Engine/Source/Runtime/Engine/Private/Rendering/SubsurfaceProfile.cpp
@@ -70,6 +70,9 @@ struct FSubsurfaceProfileCustomVersion
 		// Old profile will change extinction value and transmission unit scale
 		ExtinctionTransmissionMultiplierFix,
 
+		// Old profile will move one extra surface albedo to transmission tint.
+		FactorExtraSurfaceAlbedoToTransmissionTintFix,
+
 		// -----<new versions can be added above this line>-------------------------------------------------
 		VersionPlusOne,
 		LatestVersion = VersionPlusOne - 1
@@ -169,6 +172,11 @@ FSubsurfaceProfileTexture::FSubsurfaceProfileTexture()
 		DefaultSkin.DistanceScale *= 10.0f;
 	}
 
+	// Same to Ar.CustomVer(FSubsurfaceProfileCustomVersion::GUID) < FSubsurfaceProfileCustomVersion::FactorExtraSurfaceAlbedoToTransmissionTintFix
+	{
+		DefaultSkin.TransmissionTintColor *= DefaultSkin.SurfaceAlbedo;
+	}
+
 	// add element 0, it is used as default profile
 	SubsurfaceProfileEntries.Add(FSubsurfaceProfileEntry(DefaultSkin, 0));
 }
@@ -923,6 +931,12 @@ void USubsurfaceProfile::Serialize(FArchive& Ar)
 			// The transmission distance was scaled up by 10.0 in old profiles.
 			Settings.DistanceScale *= 10.0f;
 		}
+
+		if (Ar.CustomVer(FSubsurfaceProfileCustomVersion::GUID) < FSubsurfaceProfileCustomVersion::FactorExtraSurfaceAlbedoToTransmissionTintFix)
+		{
+			// Transmission has SurfaceAlbedo * StoredBaseColor previously.
+			Settings.TransmissionTintColor *= Settings.SurfaceAlbedo;
+		}
 	}
 }
 
