Substrate - NMT - Added support for LightFunction, PostProcess and UI nodes.

#rnx

#Synced-CL 51581332

diff --git a/Engine/Source/Runtime/Engine/Private/Materials/MaterialExpressionsToMIR.cpp b/Engine/Source/Runtime/Engine/Private/Materials/MaterialExpressionsToMIR.cpp
--- a/Engine/Source/Runtime/Engine/Private/Materials/MaterialExpressionsToMIR.cpp
+++ b/Engine/Source/Runtime/Engine/Private/Materials/MaterialExpressionsToMIR.cpp
@@ -6573,6 +6573,75 @@ void UMaterialExpressionSubstrateUnlitBSDF::Build(MIR::FEmitter& Em)
 	FinalizeBSDFAndPromoteToOperatorIfNeeded(Em, SubstrateOperator, UnlitBSDF);
 }
 
+void UMaterialExpressionSubstrateLightFunction::Build(MIR::FEmitter& Em)
+{
+	FSubstrateTranslatorData* SubstrateTranslatorData = Em.GetSubstrateTranslatorData();
+
+	FGuid PathUniqueId = SubstrateTranslatorData->SubstrateTreeStackGetPathUniqueId();
+	FSubstrateOperator& SubstrateOperator = SubstrateTranslatorData->SubstrateCompilationGetOperator(PathUniqueId);
+
+	FValueRef ColorArg = CompileWithDefaultFloat3(Em, Color, 0.0f, 0.0f, 0.0f);				// Light function color
+	FValueRef TransmittanceColorArg = Em.ConstantFloat3(FVector3f(1.0f, 1.0f, 1.0f));		// Opacity / Transmittance is ignored by light functions
+	FValueRef NormalArg = Em.ConstantFloat3(FVector3f(0.0f, 0.0f, 1.0f));					// place holder normal
+
+	FValueRef UnlitBSDF = Em.Extern<MIR::FExternSimpleHLSL>(
+		{
+			MIR::EExternSimpleType::SubstrateData,
+			TEXT("GetSubstrateUnlitBSDF($0, $1, $2)"),
+		},
+		ColorArg,
+		TransmittanceColorArg,
+		NormalArg);
+
+	FinalizeBSDFAndPromoteToOperatorIfNeeded(Em, SubstrateOperator, UnlitBSDF);
+}
+
+void UMaterialExpressionSubstratePostProcess::Build(MIR::FEmitter& Em)
+{
+	FSubstrateTranslatorData* SubstrateTranslatorData = Em.GetSubstrateTranslatorData();
+
+	FGuid PathUniqueId = SubstrateTranslatorData->SubstrateTreeStackGetPathUniqueId();
+	FSubstrateOperator& SubstrateOperator = SubstrateTranslatorData->SubstrateCompilationGetOperator(PathUniqueId);
+
+	FValueRef OpacityArg = CompileWithDefaultFloat1(Em, Opacity, 1.0f);
+	FValueRef TransmittanceColorArg = Em.Saturate(Em.Subtract(Em.ConstantFloat(1.0f), OpacityArg));	// Transmittance from Opacity
+
+	FValueRef ColorArg = CompileWithDefaultFloat3(Em, Color, 0.0f, 0.0f, 0.0f);						// Post process color
+	FValueRef NormalArg = Em.ConstantFloat3(FVector3f(0.0f, 0.0f, 1.0f));							// place holder normal
+
+	FValueRef UnlitBSDF = Em.Extern<MIR::FExternSimpleHLSL>(
+		{
+			MIR::EExternSimpleType::SubstrateData,
+			TEXT("GetSubstrateUnlitBSDF($0, $1, $2)"),
+		},
+		ColorArg,
+		TransmittanceColorArg,
+		NormalArg);
+
+	FinalizeBSDFAndPromoteToOperatorIfNeeded(Em, SubstrateOperator, UnlitBSDF);
+}
+
+void UMaterialExpressionSubstrateUI::Build(MIR::FEmitter& Em)
+{
+	FSubstrateTranslatorData* SubstrateTranslatorData = Em.GetSubstrateTranslatorData();
+	FGuid PathUniqueId = SubstrateTranslatorData->SubstrateTreeStackGetPathUniqueId();
+	FSubstrateOperator& SubstrateOperator = SubstrateTranslatorData->SubstrateCompilationGetOperator(PathUniqueId);
+
+	// Compile inputs with defaults like the legacy code did.
+	FValueRef ColorArg = CompileWithDefaultFloat3(Em, Color, 0.0f, 0.0f, 0.0f);		// Color
+	FValueRef OpacityArg = CompileWithDefaultFloat1(Em, Opacity, 1.0f);				// Opacity
+
+	// Call the Substrate UI BSDF helper and finalize/promote the resulting BSDF as needed.
+	FValueRef UIBSDF = Em.Extern<MIR::FExternSimpleHLSL>(
+		{ MIR::EExternSimpleType::SubstrateData, TEXT("SubstrateCreateUIMaterial($0, $1)") },
+		ColorArg,
+		OpacityArg
+	);
+
+	FinalizeBSDFAndPromoteToOperatorIfNeeded(Em, SubstrateOperator, UIBSDF);
+
+}
+
 static FValueRef SubstrateBlendNormal(MIR::FEmitter& Em, FValueRef NormalCodeChunk0, FValueRef NormalCodeChunk1, FValueRef MixCodeChunk)
 {
 	FValueRef SafeMixCodeChunk = Em.Saturate(MixCodeChunk);
diff --git a/Engine/Source/Runtime/Engine/Public/Materials/MaterialExpressionSubstrate.h b/Engine/Source/Runtime/Engine/Public/Materials/MaterialExpressionSubstrate.h
--- a/Engine/Source/Runtime/Engine/Public/Materials/MaterialExpressionSubstrate.h
+++ b/Engine/Source/Runtime/Engine/Public/Materials/MaterialExpressionSubstrate.h
@@ -835,7 +835,7 @@ class UMaterialExpressionSubstrateLightFunction : public UMaterialExpressionSubs
 	//~ Begin UMaterialExpression Interface
 #if WITH_EDITOR
 	virtual int32 Compile(class FMaterialCompiler* Compiler, int32 OutputIndex) override;
-	// SUBSTRATE_TODO virtual void Build(MIR::FEmitter& Emitter) override;
+	virtual void Build(MIR::FEmitter& Emitter) override;
 	virtual void GetCaption(TArray<FString>& OutCaptions) const override;
 	virtual EMaterialValueType GetOutputValueType(int32 OutputIndex) override;
 	virtual EMaterialValueType GetInputValueType(int32 InputIndex) override;
@@ -942,7 +942,7 @@ class UMaterialExpressionSubstratePostProcess : public UMaterialExpressionSubstr
 	//~ Begin UMaterialExpression Interface
 #if WITH_EDITOR
 	virtual int32 Compile(class FMaterialCompiler* Compiler, int32 OutputIndex) override;
-	// SUBSTRATE_TODO virtual void Build(MIR::FEmitter& Emitter) override;
+	virtual void Build(MIR::FEmitter& Emitter) override;
 	virtual void GetCaption(TArray<FString>& OutCaptions) const override;
 	virtual EMaterialValueType GetOutputValueType(int32 OutputIndex) override;
 	virtual EMaterialValueType GetInputValueType(int32 InputIndex) override;
@@ -984,7 +984,7 @@ class UMaterialExpressionSubstrateUI : public UMaterialExpressionSubstrateBSDF
 	//~ Begin UMaterialExpression Interface
 #if WITH_EDITOR
 	virtual int32 Compile(class FMaterialCompiler* Compiler, int32 OutputIndex) override;
-	// SUBSTRATE_TODO virtual void Build(MIR::FEmitter& Emitter) override;
+	virtual void Build(MIR::FEmitter& Emitter) override;
 	virtual void GetCaption(TArray<FString>& OutCaptions) const override;
 	virtual EMaterialValueType GetOutputValueType(int32 OutputIndex) override;
 	virtual EMaterialValueType GetInputValueType(int32 InputIndex) override;
