Remove AsyncRegisterLevelContext feature from level streaming

The AsyncRegisterLevelContext feature allowed level streaming to add primitives to the scene asynchronously on a task thread during incremental component registration. It was gated behind a CVar (LevelStreaming.AsyncRegisterLevelContext.Enabled) that defaulted to false and was never enabled.

This change removes the feature entirely:
- Deletes AsyncRegisterLevelContext.h and .cpp (internal, non-exported)
- Removes FAsyncAddPrimitiveQueue and FSendRenderDynamicDataPrimitivesQueue structs
- Removes both CVars (Enabled, PrimitiveBatchSize)
- Simplifies FRegisterComponentContext by removing all async code paths
- Simplifies the incremental component registration loop in UWorld::AddToWorld
- Deprecates ULevel::CreateAsyncRegisterLevelContext() and DestroyAsyncRegisterLevelContext() with UE_DEPRECATED(5.8) since they are ENGINE_API and shipped in 5.7

The synchronous batched path in FRegisterComponentContext::Process() was already
the default and only active code path. Runtime behavior is unchanged.

#Synced-CL 51585013

diff --git a/Engine/Source/Runtime/Engine/Classes/Components/ActorComponent.h b/Engine/Source/Runtime/Engine/Classes/Components/ActorComponent.h
--- a/Engine/Source/Runtime/Engine/Classes/Components/ActorComponent.h
+++ b/Engine/Source/Runtime/Engine/Classes/Components/ActorComponent.h
@@ -32,7 +32,6 @@ class ULevel;
 class UWorld;
 class UPrimitiveComponent;
 struct FSimpleMemberReference;
-struct FAsyncRegisterLevelContext;
 class IPrimitiveComponent;
 class UEngineReplicationBridge;
 
@@ -64,7 +63,6 @@ public:
 	FRegisterComponentContext& SetIncrementalUpdateTimeout(const UE::FTimeout& InTimeout);
 	FRegisterComponentContext& SetIncrementalUpdateAddPrimitiveBatchSize(int32 Value);
 	FRegisterComponentContext& SetIncrementalUpdateNumComponentsToUpdate(int32 Value);
-	FRegisterComponentContext& SetAsyncRegisterLevelContext(FAsyncRegisterLevelContext* Value);
 
 	const UE::FTimeout& GetIncrementalUpdateTimeout() const { return Timeout; }
 	int32 GetIncrementalUpdateAddPrimitiveBatchSize() const { return IncrementalUpdateAddPrimitiveBatchSize; }
@@ -78,6 +76,9 @@ public:
 	UE_DEPRECATED(5.5, "This function is not used anymore")
 	int32 Count() const { return AddPrimitiveBatches.Num(); }
 
+	UE_DEPRECATED(5.8, "AsyncRegisterLevelContext has been removed.")
+	FRegisterComponentContext& SetAsyncRegisterLevelContext(struct FAsyncRegisterLevelContext* Value);
+
 	ENGINE_API void Process();
 
 private:
@@ -87,8 +88,6 @@ private:
 	UE::FTimeout Timeout;
 	int32	IncrementalUpdateAddPrimitiveBatchSize;
 	int32	IncrementalUpdateNumComponentsToUpdate;
-	FAsyncRegisterLevelContext* AsyncRegisterLevelContext;
-	
 	FPrimitiveBatch AddPrimitiveBatches;
 	FPrimitiveBatch SendRenderDynamicDataPrimitives;
 
diff --git a/Engine/Source/Runtime/Engine/Classes/Engine/Level.h b/Engine/Source/Runtime/Engine/Classes/Engine/Level.h
--- a/Engine/Source/Runtime/Engine/Classes/Engine/Level.h
+++ b/Engine/Source/Runtime/Engine/Classes/Engine/Level.h
@@ -18,9 +18,6 @@
 #include "Misc/EditorPathObjectInterface.h"
 #include "Components/ActorComponent.h"
 
-// TODO: Internal headers should not be included
-#include "Runtime/Engine/Internal/Streaming/AsyncRegisterLevelContext.h"
-
 #include <atomic>
 
 #include "Level.generated.h"
@@ -866,10 +863,6 @@ private:
 		return PreUnregisteringActorComponents.Num() > 0;
 	}
 
-	ENGINE_API void CreateAsyncRegisterLevelContext();
-	ENGINE_API void DestroyAsyncRegisterLevelContext();
-	
-	TUniquePtr<FAsyncRegisterLevelContext> AsyncRegisterLevelContext;
 	TSet<UActorComponent*> PreUnregisteringActorComponents;
 	TSet<UActorComponent*> PreRegisteringActorComponents;
 	friend struct FLevelRegistrationAccessor;
@@ -1618,21 +1611,6 @@ private:
 		return InLevel->HasPreUnregisteringComponents();
 	}
 
-	static void CreateAsyncRegisterLevelContext(ULevel* InLevel)
-	{
-		InLevel->CreateAsyncRegisterLevelContext();
-	}
-
-	static void DestroyAsyncRegisterLevelContext(ULevel* InLevel)
-	{
-		InLevel->DestroyAsyncRegisterLevelContext();
-	}
-
-	static FAsyncRegisterLevelContext* GetAsyncRegisterLevelContext(ULevel* InLevel)
-	{
-		return InLevel ? InLevel->AsyncRegisterLevelContext.Get() : nullptr;
-	}
-
 	friend class FRegisterComponentContext;
 	friend class UActorComponent;
 	friend class UWorld;
diff --git a/Engine/Source/Runtime/Engine/Internal/Streaming/AsyncRegisterLevelContext.cpp b/Engine/Source/Runtime/Engine/Internal/Streaming/AsyncRegisterLevelContext.cpp
deleted file mode 100644
diff --git a/Engine/Source/Runtime/Engine/Internal/Streaming/AsyncRegisterLevelContext.h b/Engine/Source/Runtime/Engine/Internal/Streaming/AsyncRegisterLevelContext.h
deleted file mode 100644
diff --git a/Engine/Source/Runtime/Engine/Private/Components/ActorComponent.cpp b/Engine/Source/Runtime/Engine/Private/Components/ActorComponent.cpp
--- a/Engine/Source/Runtime/Engine/Private/Components/ActorComponent.cpp
+++ b/Engine/Source/Runtime/Engine/Private/Components/ActorComponent.cpp
@@ -44,7 +44,6 @@
 #include "Materials/Material.h"
 #include "ObjectCacheContext.h"
 #include "AutoRTFM.h"
-#include "Streaming/AsyncRegisterLevelContext.h"
 
 #if WITH_EDITOR
 #include "Kismet2/ComponentEditorUtils.h"
@@ -140,45 +139,27 @@ FRegisterComponentContext::FRegisterComponentContext(UWorld* InWorld)
 	, Timeout(UE::FTimeout::AlwaysExpired())
 	, IncrementalUpdateAddPrimitiveBatchSize(0)
 	, IncrementalUpdateNumComponentsToUpdate(0)
-	, AsyncRegisterLevelContext(nullptr)
 {}
 
 FRegisterComponentContext::~FRegisterComponentContext()
 {
 	Process();
-
-	// Make sure level incremental registration has completed its asynchronous tasks
-	check(!AsyncRegisterLevelContext || !AsyncRegisterLevelContext->IsRunningAsync());
 }
 
 void FRegisterComponentContext::AddPrimitive(UPrimitiveComponent* PrimitiveComponent)
 {
-	if (AsyncRegisterLevelContext)
-	{
-		AsyncRegisterLevelContext->AddPrimitive(PrimitiveComponent);
-	}
-	else
+	checkSlow(!AddPrimitiveBatches.Contains(PrimitiveComponent));
+	AddPrimitiveBatches.Add(PrimitiveComponent);
+	if ((IncrementalUpdateAddPrimitiveBatchSize > 0) && (AddPrimitiveBatches.Num() >= IncrementalUpdateAddPrimitiveBatchSize))
 	{
-		checkSlow(!AddPrimitiveBatches.Contains(PrimitiveComponent));
-		AddPrimitiveBatches.Add(PrimitiveComponent);
-		if ((IncrementalUpdateAddPrimitiveBatchSize > 0) && (AddPrimitiveBatches.Num() >= IncrementalUpdateAddPrimitiveBatchSize))
-		{
-			Process();
-		}
+		Process();
 	}
 }
 
 void FRegisterComponentContext::AddSendRenderDynamicData(UPrimitiveComponent* PrimitiveComponent)
 {
-	if (AsyncRegisterLevelContext)
-	{
-		AsyncRegisterLevelContext->AddSendRenderDynamicData(PrimitiveComponent);
-	}
-	else
-	{
-		checkSlow(!SendRenderDynamicDataPrimitives.Contains(PrimitiveComponent));
-		SendRenderDynamicDataPrimitives.Add(PrimitiveComponent);
-	}
+	checkSlow(!SendRenderDynamicDataPrimitives.Contains(PrimitiveComponent));
+	SendRenderDynamicDataPrimitives.Add(PrimitiveComponent);
 }
 
 void FRegisterComponentContext::SendRenderDynamicData(FRegisterComponentContext* Context, UPrimitiveComponent* PrimitiveComponent)
@@ -213,7 +194,6 @@ FRegisterComponentContext& FRegisterComponentContext::SetIncrementalUpdateNumCom
 
 FRegisterComponentContext& FRegisterComponentContext::SetAsyncRegisterLevelContext(FAsyncRegisterLevelContext* Value)
 {
-	AsyncRegisterLevelContext = Value;
 	return *this;
 }
 
@@ -221,33 +201,18 @@ bool FRegisterComponentContext::OnIncrementalRegisterComponentsDone()
 {
 	// We need to process pending adds prior to rerunning the construction scripts, which may internally
 	// perform removals / adds themselves.
-	if (AsyncRegisterLevelContext)
-	{
-		// Notify the context that we registered all actor components
-		AsyncRegisterLevelContext->SetIncrementalRegisterComponentsDone(true);
-		AsyncRegisterLevelContext->Tick();
-		return !AsyncRegisterLevelContext->HasRemainingWork();
-	}
-	else
-	{
-		Process();
-		return true;
-	}
+	Process();
+	return true;
 }
 
 void FRegisterComponentContext::Process()
 {
 	TRACE_CPUPROFILER_EVENT_SCOPE(FRegisterComponentContext::Process)
-	if (AsyncRegisterLevelContext)
-	{
-		AsyncRegisterLevelContext->Tick();
-	}
-	else
-	{
+
 	if (AddPrimitiveBatches.Num())
 	{
-			FSceneInterface* Scene = World->Scene;
-			const bool bAppCanEverRender = FApp::CanEverRender();
+		FSceneInterface* Scene = World->Scene;
+		const bool bAppCanEverRender = FApp::CanEverRender();
 
 		bool bSingleThreaded = !FApp::ShouldUseThreadingForPerformance();
 #if WITH_EDITOR
@@ -294,7 +259,6 @@ void FRegisterComponentContext::Process()
 		}
 	}
 	SendRenderDynamicDataPrimitives.Empty();
-}
 
 	check(AddPrimitiveBatches.IsEmpty());
 	check(SendRenderDynamicDataPrimitives.IsEmpty());
diff --git a/Engine/Source/Runtime/Engine/Private/Level.cpp b/Engine/Source/Runtime/Engine/Private/Level.cpp
--- a/Engine/Source/Runtime/Engine/Private/Level.cpp
+++ b/Engine/Source/Runtime/Engine/Private/Level.cpp
@@ -1865,7 +1865,6 @@ void ULevel::IncrementalUpdateComponents(int32 NumComponentsToUpdate, bool bReru
 	Context.SetIncrementalUpdateTimeout(InContext ? InContext->GetIncrementalUpdateTimeout() : UE::FTimeout::AlwaysExpired());
 	Context.SetIncrementalUpdateNumComponentsToUpdate(NumComponentsToUpdate);
 	Context.SetIncrementalUpdateAddPrimitiveBatchSize(InContext ? InContext->GetIncrementalUpdateAddPrimitiveBatchSize() : 0);
-	Context.SetAsyncRegisterLevelContext(AsyncRegisterLevelContext.Get());
 
 	bool bAllowLoop;
 	do
@@ -1925,7 +1924,6 @@ void ULevel::IncrementalUpdateComponents(int32 NumComponentsToUpdate, bool bReru
 #endif
 
 		case EIncrementalComponentState::Finalize:
-			check(!AsyncRegisterLevelContext.IsValid() || !AsyncRegisterLevelContext->HasRemainingWork());
 			IncrementalComponentState = EIncrementalComponentState::Init;
 			CurrentActorIndexForIncrementalUpdate = 0;
 			bAreComponentsCurrentlyRegistered = true;
@@ -1989,17 +1987,6 @@ bool ULevel::IncrementalPreRegisterComponents(FRegisterComponentContext& Context
 	return false;
 }
 
-void ULevel::CreateAsyncRegisterLevelContext()
-{
-	check(!AsyncRegisterLevelContext.IsValid());
-	AsyncRegisterLevelContext = FAsyncRegisterLevelContext::CreateInstance(this);
-}
-
-void ULevel::DestroyAsyncRegisterLevelContext()
-{
-	AsyncRegisterLevelContext.Reset();
-}
-
 bool ULevel::IncrementalRegisterComponents(FRegisterComponentContext& Context)
 {
 	TRACE_CPUPROFILER_EVENT_SCOPE(ULevel::IncrementalRegisterComponents());
diff --git a/Engine/Source/Runtime/Engine/Private/World.cpp b/Engine/Source/Runtime/Engine/Private/World.cpp
--- a/Engine/Source/Runtime/Engine/Private/World.cpp
+++ b/Engine/Source/Runtime/Engine/Private/World.cpp
@@ -132,7 +132,6 @@
 #include "Net/NetPing.h"
 #include "AssetRegistry/AssetRegistryModule.h"
 #include "ProfilingDebugging/CountersTrace.h"
-#include "Streaming/AsyncRegisterLevelContext.h"
 
 #include "Iris/ReplicationSystem/ReplicationSystem.h"
 #include "Net/Iris/ReplicationSystem/EngineReplicationBridge.h"
@@ -3747,7 +3746,6 @@ void UWorld::AddToWorld(ULevel* Level, const FTransform& LevelTransform, bool bC
 		
 		// Mark level as being in process of being made visible.
 		SetLevelPendingVisibilityState(Level, ELevelPendingVisibilityState::MakingVisible);
-		FLevelRegistrationAccessor::CreateAsyncRegisterLevelContext(Level);
 
 		// Add to the UWorld's array of levels, which causes it to be rendered et al.
 		Levels.AddUnique( Level );
@@ -3863,15 +3861,6 @@ void UWorld::AddToWorld(ULevel* Level, const FTransform& LevelTransform, bool bC
 		// from disk rather than duplicated
 		const bool bRerunConstructionScript = !FPlatformProperties::RequiresCookedData() && (!bIsGameWorld || !Level->bHasRerunConstructionScripts);
 
-		FAsyncRegisterLevelContext* AsyncRegisterLevelContext = FLevelRegistrationAccessor::GetAsyncRegisterLevelContext(Level);
-		auto SetLevelContextCanLaunchNewTasks = [AsyncRegisterLevelContext](bool bValue)
-		{
-			if (AsyncRegisterLevelContext)
-			{
-				AsyncRegisterLevelContext->SetCanLaunchNewTasks(bValue);
-			}
-		};
-
 		// Incremental update components
 		{
 			// Prepare context used to store batch/parallelize AddPrimitive calls
@@ -3881,7 +3870,6 @@ void UWorld::AddToWorld(ULevel* Level, const FTransform& LevelTransform, bool bC
 
 			// Incrementally update components.
 			bool bTimeLimitExceeded = false;
-			bool bHasRunningTasks = false;
 			const int32 NumComponentsToUpdate = (!bConsiderTimeLimit || !bIsGameWorld || IsRunningCommandlet() ? 0 : GLevelStreamingComponentsRegistrationGranularity);
 			do
 			{
@@ -3891,22 +3879,9 @@ void UWorld::AddToWorld(ULevel* Level, const FTransform& LevelTransform, bool bC
 					break;
 				}
 				bTimeLimitExceeded = IsTimeLimitExceeded(Timeout, TEXT("updating components"), Level);
-				bHasRunningTasks = AsyncRegisterLevelContext ? AsyncRegisterLevelContext->IsRunningAsync() : false;
-				if (AsyncRegisterLevelContext && bTimeLimitExceeded && bHasRunningTasks)
-				{
-					// Block level context to launch new tasks since we reached the time limit
-					SetLevelContextCanLaunchNewTasks(false);
-				}
-			} while (!bTimeLimitExceeded || bHasRunningTasks);
-			
-			// Block level context to launch new tasks since we exit the IncrementalUpdateComponents scope
-			SetLevelContextCanLaunchNewTasks(false);
+			} while (!bTimeLimitExceeded);
 		}
 
-		check(!Level->bAreComponentsCurrentlyRegistered || !AsyncRegisterLevelContext || !AsyncRegisterLevelContext->HasRemainingWork());
-		// Allow level context to launch new tasks
-		SetLevelContextCanLaunchNewTasks(true);
-
 		// We are done once all components are attached.
 		Level->bAlreadyUpdatedComponents	= Level->bAreComponentsCurrentlyRegistered;
 		bExecuteNextStep					= Level->bAreComponentsCurrentlyRegistered && (!bConsiderTimeLimit || !IsTimeLimitExceeded(Timeout, TEXT("updating components"), Level));
@@ -4012,7 +3987,6 @@ void UWorld::AddToWorld(ULevel* Level, const FTransform& LevelTransform, bool bC
 
 		// Finished making level visible - allow other levels to be added to the world.
 		SetLevelPendingVisibilityState(Level, ELevelPendingVisibilityState::None);
-		FLevelRegistrationAccessor::DestroyAsyncRegisterLevelContext(Level);
 
 		if (bIsGameWorld && !Level->bClientOnlyVisible)
 		{
