[Cloud gaming] Clean resolution logs added for xCloud resolution investigation

#rnx

#Synced-CL 51584369

diff --git a/Engine/Source/Runtime/Engine/Private/GameEngine.cpp b/Engine/Source/Runtime/Engine/Private/GameEngine.cpp
--- a/Engine/Source/Runtime/Engine/Private/GameEngine.cpp
+++ b/Engine/Source/Runtime/Engine/Private/GameEngine.cpp
@@ -472,7 +472,6 @@ void UGameEngine::DetermineGameWindowResolution( int32& ResolutionX, int32& Reso
 		// We need to pass the resolution back out to GameUserSettings, or it will just override it again
 		ResolutionX = DisplayMetrics.PrimaryDisplayWorkAreaRect.Right - DisplayMetrics.PrimaryDisplayWorkAreaRect.Left;
 		ResolutionY = DisplayMetrics.PrimaryDisplayWorkAreaRect.Bottom - DisplayMetrics.PrimaryDisplayWorkAreaRect.Top;
-		UE_LOG(LogEngine, Display, TEXT("UGameEngine::DetermineGameWindowResolution: calling RequestResolutionChange %dx%d %s"), ResolutionX, ResolutionY, LexToString(EWindowMode::Fullscreen));
 		FSystemResolution::RequestResolutionChange(ResolutionX, ResolutionY, EWindowMode::Fullscreen);
 	}
 
@@ -496,7 +495,6 @@ TSharedRef<SWindow> UGameEngine::CreateGameWindow()
 		if (ResX != GSystemResolution.ResX || ResY != GSystemResolution.ResY || WindowMode != GSystemResolution.WindowMode)
 		{
 			FSystemResolution::RequestResolutionChange(ResX, ResY, WindowMode);
-			UE_LOG(LogEngine, Display, TEXT("UGameEngine::CreateGameWindow: calling RequestResolutionChange %dx%d %s"), ResX, ResY, LexToString(WindowMode));
 			IConsoleManager::Get().CallAllConsoleVariableSinks();
 		}
 	}
@@ -860,7 +858,6 @@ void UGameEngine::OnViewportResized(FViewport* Viewport, uint32 Unused)
 		{
 			GSystemResolution.ResX = ViewportSize.X;
 			GSystemResolution.ResY = ViewportSize.Y;
-			UE_LOG(LogEngine, Display, TEXT("UGameEngine::OnViewportResized: calling RequestResolutionChange %dx%d %s"), GSystemResolution.ResX, GSystemResolution.ResY, LexToString(EWindowMode::Windowed));
 			FSystemResolution::RequestResolutionChange(GSystemResolution.ResX, GSystemResolution.ResY, EWindowMode::Windowed);
 
 			UGameUserSettings* Settings = GetGameUserSettings();
diff --git a/Engine/Source/Runtime/Engine/Private/GameUserSettings.cpp b/Engine/Source/Runtime/Engine/Private/GameUserSettings.cpp
--- a/Engine/Source/Runtime/Engine/Private/GameUserSettings.cpp
+++ b/Engine/Source/Runtime/Engine/Private/GameUserSettings.cpp
@@ -576,7 +576,6 @@ void UGameUserSettings::ApplyResolutionSettings(bool bCheckForCommandLineOverrid
 
 	// Request a resolution change
 	const int32 MonitorIndex = GetClosestMonitorIndex();
-	UE_LOG(LogGameUserSettings, Display, TEXT("UGameUserSettings::ApplyResolutionSettings: calling RequestResolutionChange %dx%d %s"), ResolutionSizeX, ResolutionSizeY, LexToString(NewFullscreenMode));
 	RequestResolutionChange(ResolutionSizeX, ResolutionSizeY, NewFullscreenMode, bCheckForCommandLineOverrides, MonitorIndex);
 
 	if (NewFullscreenMode == EWindowMode::Fullscreen || NewFullscreenMode == EWindowMode::WindowedFullscreen)
@@ -667,8 +666,6 @@ void UGameUserSettings::RequestResolutionChange(int32 InResolutionX, int32 InRes
 	}
 
 	FSystemResolution::SetMonitorIndex(MonitorIndex);
-
-	UE_LOG(LogGameUserSettings, Display, TEXT("UGameUserSettings::RequestResolutionChange: calling RequestResolutionChange %dx%d %s"), InResolutionX, InResolutionY, LexToString(InWindowMode));
 	FSystemResolution::RequestResolutionChange(InResolutionX, InResolutionY, InWindowMode);
 }
 
@@ -816,7 +813,6 @@ void UGameUserSettings::PreloadResolutionSettings(bool bAllowCmdLineOverrides /*
 #endif
 
 	const int32 MonitorIndex = GetClosestMonitorFromIDAndIndex(DisplayID, DisplayIndex);
-	UE_LOG(LogGameUserSettings, Display, TEXT("UGameUserSettings::PreloadResolutionSettings: calling RequestResolutionChange %dx%d %s"), ResolutionX, ResolutionY, LexToString(WindowMode));
 	RequestResolutionChange(ResolutionX, ResolutionY, WindowMode, bAllowCmdLineOverrides, MonitorIndex);
 
 	IConsoleManager::Get().CallAllConsoleVariableSinks();
diff --git a/Engine/Source/Runtime/Engine/Private/GameViewportClient.cpp b/Engine/Source/Runtime/Engine/Private/GameViewportClient.cpp
--- a/Engine/Source/Runtime/Engine/Private/GameViewportClient.cpp
+++ b/Engine/Source/Runtime/Engine/Private/GameViewportClient.cpp
@@ -3969,7 +3969,6 @@ bool UGameViewportClient::SetDisplayConfiguration(const FIntPoint* Dimensions, E
 			NewY = Dimensions->Y;
 		}
 
-		UE_LOG(LogViewport, Display, TEXT("UGameViewportClient::SetDisplayConfiguration: calling RequestResolutionChange %dx%d %s"), NewX, NewY, LexToString(WindowMode));
 		FSystemResolution::RequestResolutionChange(NewX, NewY, WindowMode);
 	}
 
@@ -4047,7 +4046,6 @@ bool UGameViewportClient::HandleToggleFullscreenCommand()
 		if (!bNewModeApplied)
 		{
 			FSystemResolution::SetMonitorIndex(MonitorIndex);
-			UE_LOG(LogViewport, Display, TEXT("UGameViewportClient::HandleToggleFullscreenCommand: calling RequestResolutionChange %dx%d %s"), ResolutionX, ResolutionY, LexToString(FullScreenMode));
 			FSystemResolution::RequestResolutionChange(ResolutionX, ResolutionY, FullScreenMode);
 		}
 
@@ -4085,7 +4083,6 @@ bool UGameViewportClient::HandleSetResCommand( const TCHAR* Cmd, FOutputDevice&
 		}
 		if( X && Y )
 		{
-			UE_LOG(LogViewport, Display, TEXT("UGameEngine::DetermineGameWindowResolution: calling RequestResolutionChange %dx%d %s"), X, Y, LexToString(WindowMode));
 			FSystemResolution::RequestResolutionChange(X, Y, WindowMode);
 		}
 	}
diff --git a/Engine/Source/Runtime/Engine/Private/UnrealEngine.cpp b/Engine/Source/Runtime/Engine/Private/UnrealEngine.cpp
--- a/Engine/Source/Runtime/Engine/Private/UnrealEngine.cpp
+++ b/Engine/Source/Runtime/Engine/Private/UnrealEngine.cpp
@@ -18981,7 +18981,6 @@ void FSystemResolution::RequestResolutionChange(int32 InResX, int32 InResY, EWin
 void FSystemResolution::ForceRefresh()
 {
 	SetMonitorIndex(MonitorIndex);
-	UE_LOG(LogEngine, Display, TEXT("FSystemResolution::ForceRefresh: calling RequestResolutionChange %dx%d %s"), ResX, ResY, LexToString(WindowMode));
 	RequestResolutionChange(ResX, ResY, WindowMode);
 	bForceRefresh = true;
 }
