Substrate - Unified how shader quality is marked as needed for substrate, similar to quality switch nodes..

#rnx
#Synced-CL 51590311

#Synced-CL 51590311

diff --git a/Engine/Source/Runtime/Engine/Private/Materials/MaterialCachedData.cpp b/Engine/Source/Runtime/Engine/Private/Materials/MaterialCachedData.cpp
--- a/Engine/Source/Runtime/Engine/Private/Materials/MaterialCachedData.cpp
+++ b/Engine/Source/Runtime/Engine/Private/Materials/MaterialCachedData.cpp
@@ -490,7 +490,10 @@ void FMaterialCachedExpressionData::UpdateForExpressions(const FMaterialCachedEx
 		uint32 ReferenceClosurePerPixelCount = Substrate::GetClosurePerPixel(GMaxRHIShaderPlatform, Substrate::IgnoreMaterialQualityLevelsForClosurePerPixel);
 		for (int32 MaterialQualityLevel = 0; MaterialQualityLevel < EMaterialQualityLevel::Num; MaterialQualityLevel++)
 		{
-			QualityLevelsUsed[MaterialQualityLevel] = ReferenceClosurePerPixelCount != Substrate::GetClosurePerPixel(GMaxRHIShaderPlatform, MaterialQualityLevel);
+			if (ReferenceClosurePerPixelCount != Substrate::GetClosurePerPixel(GMaxRHIShaderPlatform, MaterialQualityLevel))
+			{
+				QualityLevelsUsed[MaterialQualityLevel] = true;
+			}
 		}
 	}
 
