Non-exhaustive removal of deprecated code from Shader and Material files.

#Synced-CL 51565377

diff --git a/Engine/Source/Developer/MaterialUtilities/Private/MaterialUtilities.cpp b/Engine/Source/Developer/MaterialUtilities/Private/MaterialUtilities.cpp
--- a/Engine/Source/Developer/MaterialUtilities/Private/MaterialUtilities.cpp
+++ b/Engine/Source/Developer/MaterialUtilities/Private/MaterialUtilities.cpp
@@ -997,11 +997,6 @@ bool FMaterialUtilities::SupportsExport(bool bIsOpaque, EMaterialProperty InMate
 	return FExportMaterialProxy::WillFillData(bIsOpaque, InMaterialProperty);
 }
 
-bool FMaterialUtilities::SupportsExport(EBlendMode InBlendMode, EMaterialProperty InMaterialProperty)
-{
-	return FExportMaterialProxy::WillFillData(IsOpaqueBlendMode(InBlendMode), InMaterialProperty);
-}
-
 static bool ExportLandscapeMaterial(const ALandscapeProxy* InLandscape, const TSet<FPrimitiveComponentId>& ShowOnlyPrimitives, const TSet<FPrimitiveComponentId>& HiddenPrimitives, FFlattenMaterial& OutFlattenMaterial)
 {
 	check(InLandscape);
diff --git a/Engine/Source/Developer/MaterialUtilities/Public/MaterialUtilities.h b/Engine/Source/Developer/MaterialUtilities/Public/MaterialUtilities.h
--- a/Engine/Source/Developer/MaterialUtilities/Public/MaterialUtilities.h
+++ b/Engine/Source/Developer/MaterialUtilities/Public/MaterialUtilities.h
@@ -309,8 +309,6 @@ public:
 	/**
 	 * Whether material utilities support exporting specified material blend mode and property 
 	 */
-	UE_DEPRECATED(5.2, "Use SupportsExport(bool bIsOpaque, ...) function instead")
-	static UE_API bool SupportsExport(EBlendMode InBlendMode, EMaterialProperty InMaterialProperty);
 	static UE_API bool SupportsExport(bool bIsOpaque, EMaterialProperty InMaterialProperty);
 
 	/**
diff --git a/Engine/Source/Developer/ShaderCompilerCommon/Private/ShaderCompilerCommon.cpp b/Engine/Source/Developer/ShaderCompilerCommon/Private/ShaderCompilerCommon.cpp
--- a/Engine/Source/Developer/ShaderCompilerCommon/Private/ShaderCompilerCommon.cpp
+++ b/Engine/Source/Developer/ShaderCompilerCommon/Private/ShaderCompilerCommon.cpp
@@ -2159,11 +2159,6 @@ namespace UE::ShaderCompilerCommon
 		Ctx.DebugSourceFiles.Add({ Input.Target.GetFrequency(), MoveTemp(DebugSourcePath) });
 	}
 
-	void DumpDebugShaderData(const FShaderCompilerInput& Input, const FString& PreprocessedSource, const FDebugShaderDataOptions& Options)
-	{
-		// deprecated
-	}
-
 	void FBaseShaderFormat::OutputDebugDataMinimal(const FShaderCompilerInput& Input, FShaderDebugDataContext& Ctx) const
 	{
 		const FStringView FailedSourceStr = TEXTVIEW(R"(
diff --git a/Engine/Source/Developer/ShaderCompilerCommon/Public/ShaderCompilerCommon.h b/Engine/Source/Developer/ShaderCompilerCommon/Public/ShaderCompilerCommon.h
--- a/Engine/Source/Developer/ShaderCompilerCommon/Public/ShaderCompilerCommon.h
+++ b/Engine/Source/Developer/ShaderCompilerCommon/Public/ShaderCompilerCommon.h
@@ -119,17 +119,6 @@ namespace UE::ShaderCompilerCommon
 		SHADERCOMPILERCOMMON_API FString GetDebugShaderPath(const FShaderCompilerInput& Input, const TCHAR* Suffix = nullptr) const;
 	};
 
-	/*
-	 * Dumps common debug information (preprocessed .usf as constructed by GetDebugShaderContents, and a directcompile.txt file 
-	 * containing the commandline for launching ShaderCompileWorker manually) for the given shader compile input
-	 * and preprocessed source according to the options provided.
-	 * @param	Input The input of the compilation job
-	 * @param	PreprocessedSource The unmodified preprocessed source (used as input to the compilation)
-	 * @param	Options Options which can change behaviour of the debug dump; see above.
-	 */
-	UE_DEPRECATED(5.6, "This function should only be used internally and will be removed from API")
-	extern SHADERCOMPILERCOMMON_API void DumpDebugShaderData(const FShaderCompilerInput& Input, const FString& PreprocessedSource, const FDebugShaderDataOptions& Options = FDebugShaderDataOptions());
-
 	UE_DEPRECATED(5.7, "DumpExtendedDebugShaderData now requires a FShaderDebugDataContext struct parameter")
 		extern SHADERCOMPILERCOMMON_API void DumpExtendedDebugShaderData(
 			const FShaderCompilerInput& Input,
diff --git a/Engine/Source/Developer/ShaderFormatOpenGL/Public/ShaderFormatOpenGL.h b/Engine/Source/Developer/ShaderFormatOpenGL/Public/ShaderFormatOpenGL.h
--- a/Engine/Source/Developer/ShaderFormatOpenGL/Public/ShaderFormatOpenGL.h
+++ b/Engine/Source/Developer/ShaderFormatOpenGL/Public/ShaderFormatOpenGL.h
@@ -36,8 +36,6 @@ enum GLSLVersion
 	GLSL_ES2_IOS_REMOVED,
 	GLSL_310_ES_EXT_REMOVED,
 	GLSL_ES3_1_ANDROID,
-	GLSL_SWITCH UE_DEPRECATED(5.5, "no longer supported in this shader backend"),
-	GLSL_SWITCH_FORWARD UE_DEPRECATED(5.5, "no longer supported in this shader backend"),
 
 	GLSL_MAX
 };
diff --git a/Engine/Source/Developer/SharedSettingsWidgets/Public/ShaderFormatsPropertyDetails.h b/Engine/Source/Developer/SharedSettingsWidgets/Public/ShaderFormatsPropertyDetails.h
--- a/Engine/Source/Developer/SharedSettingsWidgets/Public/ShaderFormatsPropertyDetails.h
+++ b/Engine/Source/Developer/SharedSettingsWidgets/Public/ShaderFormatsPropertyDetails.h
@@ -41,14 +41,6 @@ public:
 	/** Create the UI to select which windows shader formats we are targeting */
 	UE_API void CreateTargetShaderFormatsPropertyView(ITargetPlatformSettings* TargetPlatform, GetFriendlyNameFromRHINameFnc* FriendlyNameFnc, FilterShaderPlatformFnc* FilterShaderPlatformFunc = nullptr, ECategoryPriority::Type InPriority = ECategoryPriority::Default);
 	
-	UE_DEPRECATED(5.5, "CreateTargetShaderFormatsPropertyView now takes a ITargetPlatformSettings* as first argument instead of ITargetPlatform*. Please change your callback function as this is a breaking change.")
-	UE_API void CreateTargetShaderFormatsPropertyView(ITargetPlatform* TargetPlatform, GetFriendlyNameFromRHINameFnc* FriendlyNameFnc, FilterShaderPlatformFnc* FilterShaderPlatformFunc = nullptr, ECategoryPriority::Type InPriority = ECategoryPriority::Default);
-
-	UE_DEPRECATED(5.1, "CreateTargetShaderFormatsPropertyView now gets RHI names via FName instead of FString. Please change your callback function as this is a breaking change.")
-	void CreateTargetShaderFormatsPropertyView(ITargetPlatform* TargetPlatform, Deprecated_GetFriendlyNameFromRHINameFnc* FriendlyNameFnc)
-	{
-	}
-	
 	/**
 	 * @param InRHIName - The input RHI to check
 	 * @returns Whether this RHI is currently enabled
diff --git a/Engine/Source/Developer/TargetPlatform/Public/Interfaces/IShaderFormat.h b/Engine/Source/Developer/TargetPlatform/Public/Interfaces/IShaderFormat.h
--- a/Engine/Source/Developer/TargetPlatform/Public/Interfaces/IShaderFormat.h
+++ b/Engine/Source/Developer/TargetPlatform/Public/Interfaces/IShaderFormat.h
@@ -106,9 +106,6 @@ public:
 	 */
 	virtual void ModifyShaderCompilerInput(FShaderCompilerInput& Input) const { }
 
-	UE_DEPRECATED(5.6, "Use overload accepting a const FCompressedBuffer& for symbol data")
-	virtual void NotifyShaderCompiled(const TConstArrayView<uint8>& SymbolData, FName Format, const FString& DebugInfo = FString()) const { }
-
 	/**
 	 * Called when a shader resource is cooked, so the shader format can perform platform-specific operations on the debug data.
 	 * Does nothing on platforms that make no use of the symbol data.
diff --git a/Engine/Source/Editor/MaterialEditor/Public/MaterialEditorSettings.h b/Engine/Source/Editor/MaterialEditor/Public/MaterialEditorSettings.h
--- a/Engine/Source/Editor/MaterialEditor/Public/MaterialEditorSettings.h
+++ b/Engine/Source/Editor/MaterialEditor/Public/MaterialEditorSettings.h
@@ -163,9 +163,6 @@ public:
 	*/
 	UPROPERTY(config, EditAnywhere, Category = "Offline Shader Compilers", meta = (DisplayName = "Offline Shader Compiler"))
 	EOfflineShaderCompiler OfflineCompiler;
-
-	UE_DEPRECATED(5.5, "This property has been deprecated, please use OfflineCompilerPath instead.")
-	FFilePath MaliOfflineCompilerPath;
 	
 	/**
 	Path to user installed shader compiler that can be used by the material editor to compile and extract shader informations for Android platforms.
diff --git a/Engine/Source/Runtime/Engine/Public/MaterialCompiler.h b/Engine/Source/Runtime/Engine/Public/MaterialCompiler.h
--- a/Engine/Source/Runtime/Engine/Public/MaterialCompiler.h
+++ b/Engine/Source/Runtime/Engine/Public/MaterialCompiler.h
@@ -581,20 +581,6 @@ public:
 		FTransformParameters Parameters {};
 		return TransformPosition(SourceCoordBasis, DestCoordBasis, Parameters, A);
 	}
-	UE_DEPRECATED(5.6, "Use TransformPosition(EMaterialCommonBasis, EMaterialCommonBasis, FTransformParameters, int32) instead")
-	int32 TransformFromPeriodicWorldPosition(EMaterialCommonBasis DestCoordBasis, int TileScaleIndex, int32 A)
-	{
-		FTransformParameters Parameters{};
-		Parameters.PeriodicWorldTileSizeIndex = TileScaleIndex;
-		return TransformPosition(EMaterialCommonBasis::MCB_PeriodicWorld, DestCoordBasis, Parameters, A);
-	}
-	UE_DEPRECATED(5.6, "Use TransformPosition(EMaterialCommonBasis, EMaterialCommonBasis, FTransformParameters, int32) instead")
-	int32 TransformToPeriodicWorldPosition(EMaterialCommonBasis SourceCoordBasis, int TileScaleIndex, int32 A)
-	{
-		FTransformParameters Parameters{};
-		Parameters.PeriodicWorldTileSizeIndex = TileScaleIndex;
-		return TransformPosition(SourceCoordBasis, EMaterialCommonBasis::MCB_PeriodicWorld, Parameters, A);
-	}
 	virtual int32 TransformNormalFromRequestedBasisToWorld(int32 NormalCodeChunk) = 0;
 
 	virtual int32 DynamicParameter(FLinearColor& DefaultValue, uint32 ParameterIndex = 0) = 0;
diff --git a/Engine/Source/Runtime/Engine/Public/Materials/MaterialExpressionFunctionInput.h b/Engine/Source/Runtime/Engine/Public/Materials/MaterialExpressionFunctionInput.h
--- a/Engine/Source/Runtime/Engine/Public/Materials/MaterialExpressionFunctionInput.h
+++ b/Engine/Source/Runtime/Engine/Public/Materials/MaterialExpressionFunctionInput.h
@@ -87,10 +87,6 @@ class UMaterialExpressionFunctionInput : public UMaterialExpression
 	UPROPERTY(EditAnywhere, Category = MaterialExpressionFunctionInput)
 	TEnumAsByte<enum EBlendInputRelevance> BlendInputRelevance = EBlendInputRelevance::General;
 
-	UE_DEPRECATED(5.6, "bCompilingFunctionPreview has been deprecated.")
-	UPROPERTY(transient, meta = (DeprecatedProperty, DeprecationMessage = "bCompilingFunctionPreview has been removed from function inputs and has no functional usage, it has been replaced by automatic state tracking between function calls."))
-	uint32 bCompilingFunctionPreview_DEPRECATED : 1;
-
 #if WITH_EDITOR
 	/** 
 	* These create and delete new FExpression objects on the heap, but are public to allow creation and removal from external systems. 
diff --git a/Engine/Source/Runtime/Engine/Public/Materials/MaterialFunction.h b/Engine/Source/Runtime/Engine/Public/Materials/MaterialFunction.h
--- a/Engine/Source/Runtime/Engine/Public/Materials/MaterialFunction.h
+++ b/Engine/Source/Runtime/Engine/Public/Materials/MaterialFunction.h
@@ -65,15 +65,15 @@ class UMaterialFunction : public UMaterialFunctionInterface
 	UPROPERTY(EditAnywhere, Category=MaterialFunction, AssetRegistrySearchable)
 	uint8 bExposeToLibrary:1;
 
-#if WITH_EDITORONLY_DATA	
-	UE_DEPRECATED(5.6, "bPrefixParameterNames has been deprecated.")
-	UPROPERTY()
-	uint8 bPrefixParameterNames:1;
-#endif // WITH_EDITORONLY_DATA
+#if WITH_EDITORONLY_DATA
+	UE_DEPRECATED(All, "bPrefixParameterNames has been deprecated.")
+	UPROPERTY(meta = (DeprecatedProperty, DeprecationMessage = "bPrefixParameterNames has been deprecated."))
+	uint8 bPrefixParameterNames_DEPRECATED:1;
 
 	UE_DEPRECATED(5.8, "Experimental bEnableExecWire is no longer supported.")
-	UPROPERTY(EditAnywhere, Category = MaterialFunction)
-	uint8 bEnableExecWire : 1;
+	UPROPERTY(meta = (DeprecatedProperty, DeprecationMessage = "Experimental bEnableExecWire is no longer supported."))
+	uint8 bEnableExecWire_DEPRECATED : 1;
+#endif // WITH_EDITORONLY_DATA
 
 	UPROPERTY(EditAnywhere, Category = MaterialFunction, meta=(DisplayName="Enable new material translator"))
 	uint8 bEnableNewHLSLGenerator : 1;
diff --git a/Engine/Source/Runtime/Engine/Public/PSOPrecacheMaterial.h b/Engine/Source/Runtime/Engine/Public/PSOPrecacheMaterial.h
--- a/Engine/Source/Runtime/Engine/Public/PSOPrecacheMaterial.h
+++ b/Engine/Source/Runtime/Engine/Public/PSOPrecacheMaterial.h
@@ -27,14 +27,6 @@ public:
 	IPSOCollector(int32 InPSOCollectorIndex) : PSOCollectorIndex(InPSOCollectorIndex) {}
 	virtual ~IPSOCollector() = default;
 
-	UE_DEPRECATED(5.2, "Call CollectPSOInitializers with FPSOPrecacheVertexFactoryData instead.")
-	void CollectPSOInitializers(const FSceneTexturesConfig& SceneTexturesConfig, const FMaterial& Material, const FVertexFactoryType* VertexFactoryType, const FPSOPrecacheParams& PreCacheParams, TArray<FPSOPrecacheData>& PSOInitializers)
-	{
-		FPSOPrecacheVertexFactoryData VertexFactoryData;
-		VertexFactoryData.VertexFactoryType = VertexFactoryType;
-		return CollectPSOInitializers(SceneTexturesConfig, Material, VertexFactoryData, PreCacheParams, PSOInitializers);
-	}
-
 	// Collect all PSO for given material, vertex factory & params
 	virtual void CollectPSOInitializers(const FSceneTexturesConfig& SceneTexturesConfig, const FMaterial& Material, const FPSOPrecacheVertexFactoryData& VertexFactoryData, const FPSOPrecacheParams& PreCacheParams, TArray<FPSOPrecacheData>& PSOInitializers) = 0;
 
diff --git a/Engine/Source/Runtime/Engine/Public/ShaderCompiler.h b/Engine/Source/Runtime/Engine/Public/ShaderCompiler.h
--- a/Engine/Source/Runtime/Engine/Public/ShaderCompiler.h
+++ b/Engine/Source/Runtime/Engine/Public/ShaderCompiler.h
@@ -746,17 +746,6 @@ struct FShaderCompilerSinglePermutationStat
 		, CookedDouble(0)
 	{}
 
-	UE_DEPRECATED(5.6, "Use constructor accepting a u64 permutationstring hash instead of a string")
-	FShaderCompilerSinglePermutationStat(FString PermutationString, uint32 Compiled, uint32 Cooked)
-		: Compiled(Compiled)
-		, Cooked(Cooked)
-		, CompiledDouble(0)
-		, CookedDouble(0)
-
-	{
-		PermutationHash = GetPermutationHash(PermutationString);
-	}
-
 	FShaderCompilerSinglePermutationStat(uint64 PermutationHash, uint32 Compiled, uint32 Cooked)
 		: PermutationHash(PermutationHash)
 		, Compiled(Compiled)
@@ -766,15 +755,6 @@ struct FShaderCompilerSinglePermutationStat
 
 	{}
 
-	PRAGMA_DISABLE_DEPRECATION_WARNINGS
-	// Explicitly-defaulted ctors & assignment operators are needed temporarily due to deprecation
-	// of PermutationString field. Can be removed once the deprecation window for said field ends.
-	FShaderCompilerSinglePermutationStat(FShaderCompilerSinglePermutationStat&&) = default;
-	FShaderCompilerSinglePermutationStat(const FShaderCompilerSinglePermutationStat&) = default;
-	FShaderCompilerSinglePermutationStat& operator=(FShaderCompilerSinglePermutationStat&&) = default;
-	FShaderCompilerSinglePermutationStat& operator=(const FShaderCompilerSinglePermutationStat&) = default;
-	PRAGMA_ENABLE_DEPRECATION_WARNINGS
-
 	static uint64 GetPermutationHash(const FString& PermutationString)
 	{
 		return CityHash64(reinterpret_cast<const char*>(*PermutationString), PermutationString.Len() * sizeof(TCHAR));
@@ -783,10 +763,6 @@ struct FShaderCompilerSinglePermutationStat
 	UPROPERTY()
 	uint64 PermutationHash;
 
-	UE_DEPRECATED(5.6, "PermutationString is no longer stored due to memory overhead; use PermutationHash to uniquely identify permutation stats")
-	UPROPERTY()
-	uint32 PermutationString = 0u;
-
 	UPROPERTY()
 	uint32 Compiled = 0u;
 
@@ -1319,10 +1295,6 @@ public:
 	ENGINE_API FShaderCompileJob* PrepareShaderCompileJob(uint32 Id, const FShaderCompileJobKey& Key, EShaderCompileJobPriority Priority);
 	ENGINE_API FShaderPipelineCompileJob* PreparePipelineCompileJob(uint32 Id, const FShaderPipelineCompileJobKey& Key, EShaderCompileJobPriority Priority);
 
-
-	UE_DEPRECATED(5.6, "ProcessFinishedJob must now be passed an EShaderCompileJobStatus")
-	void ProcessFinishedJob(FShaderCommonCompileJob* FinishedJob) {}
-
 	/** This is an entry point for all jobs that have finished the compilation. Can be called from multiple threads.*/
 	ENGINE_API void ProcessFinishedJob(FShaderCommonCompileJob* FinishedJob, EShaderCompileJobStatus Status);
 
diff --git a/Engine/Source/Runtime/Renderer/Private/ShaderPrint.cpp b/Engine/Source/Runtime/Renderer/Private/ShaderPrint.cpp
--- a/Engine/Source/Runtime/Renderer/Private/ShaderPrint.cpp
+++ b/Engine/Source/Runtime/Renderer/Private/ShaderPrint.cpp
@@ -277,11 +277,6 @@ namespace ShaderPrint
 		}
 	}
 
-	void SetParameters(FRDGBuilder& GraphBuilder, const FViewInfo & View, FShaderParameters& OutParameters)
-	{
-		SetParameters(GraphBuilder, View.ShaderPrintData, OutParameters);
-	}
-
 	bool IsSupported(EShaderPlatform InShaderPlatform)
 	{
 		// Avoid FXC since it struggles pretty hard with the shaders
@@ -355,13 +350,6 @@ namespace ShaderPrint
 		GQueuedRenderItems.Add({InData, InScene, GFrameNumberRenderThread + 10});
 	}
 
-	void SubmitShaderPrintData(FFrozenShaderPrintData& InData)
-	{
-		PRAGMA_DISABLE_DEPRECATION_WARNINGS
-		SubmitShaderPrintData(InData, nullptr);
-		PRAGMA_ENABLE_DEPRECATION_WARNINGS
-	}
-
 	//////////////////////////////////////////////////////////////////////////////////////////////////
 	// Common Shaders
 
diff --git a/Engine/Source/Runtime/Renderer/Public/ShaderPrintParameters.h b/Engine/Source/Runtime/Renderer/Public/ShaderPrintParameters.h
--- a/Engine/Source/Runtime/Renderer/Public/ShaderPrintParameters.h
+++ b/Engine/Source/Runtime/Renderer/Public/ShaderPrintParameters.h
@@ -154,17 +154,11 @@ namespace ShaderPrint
 	/** Submit shader print data for display in the next rendered frame. The data is displayed in views from the scene, or all views if InScene==nullptr. */
 	RENDERER_API void SubmitShaderPrintData(FFrozenShaderPrintData& InData, FSceneInterface const* InScene);
 
-	UE_DEPRECATED(5.2, "Use the version of this function that takes a FSceneInterface")
-	RENDERER_API void SubmitShaderPrintData(FFrozenShaderPrintData& InData);
-
 	/** Fill the FShaderParameters with an explicit FShaderPrintData managed by the calling code. */
 	RENDERER_API void SetParameters(FRDGBuilder& GraphBuilder, FShaderPrintData const& InData, FShaderParameters& OutParameters);
 	/** Fill the FShaderParameters with the opaque FShaderPrintData from the current default view. */
 	RENDERER_API void SetParameters(FRDGBuilder& GraphBuilder, FShaderParameters& OutParameters);
 	
-	UE_DEPRECATED(5.1, "Use one of the other implementations of SetParameters()")
-	void SetParameters(FRDGBuilder& GraphBuilder, const FViewInfo& View, FShaderParameters& OutParameters);
-
 	// Experimental GPU string
 	//
 	// In Cpp:
