Fixed crash in Virtual Texture expression Build function.

#rnx

#Synced-CL 51592004

diff --git a/Engine/Source/Runtime/Engine/Private/Materials/MaterialExpressionsToMIR.cpp b/Engine/Source/Runtime/Engine/Private/Materials/MaterialExpressionsToMIR.cpp
--- a/Engine/Source/Runtime/Engine/Private/Materials/MaterialExpressionsToMIR.cpp
+++ b/Engine/Source/Runtime/Engine/Private/Materials/MaterialExpressionsToMIR.cpp
@@ -2365,11 +2365,16 @@ void UMaterialExpressionRuntimeVirtualTextureSample::Build(MIR::FEmitter& Em)
 			UniformName = GetParameterName();
 			TheVirtualTexture = Metadata.Value.RuntimeVirtualTexture;
 		}
-		else
+		else if (VirtualTexture)
 		{
 			UniformName = VirtualTexture->GetFName();
 			TheVirtualTexture = VirtualTexture;
 		}
+		else
+		{
+			Em.Error(TEXT("No virtual texture object specified."));
+			return;
+		}
 
 		TextureObjects[TextureLayerIndex] = Em.VirtualTextureUniform(UniformName, TheVirtualTexture, TextureLayerIndex, PageTableLayerIndex);
 	}
