SpringMath: Added AutoCreateRefTerm for various functions

#Synced-CL 51592181

diff --git a/Engine/Source/Runtime/Engine/Classes/Kismet/BlueprintSpringMathLibrary.h b/Engine/Source/Runtime/Engine/Classes/Kismet/BlueprintSpringMathLibrary.h
--- a/Engine/Source/Runtime/Engine/Classes/Kismet/BlueprintSpringMathLibrary.h
+++ b/Engine/Source/Runtime/Engine/Classes/Kismet/BlueprintSpringMathLibrary.h
@@ -42,7 +42,7 @@ public:
 	 * @param SmoothingTime The smoothing time to use for the spring. Longer times lead to more damped behaviour. A time of 0 will snap InOutX to TargetX
 	 * @param DeltaTime Timestep in seconds
 	 */
-	UFUNCTION(BlueprintCallable, meta=(Category="Math|Spring|Experimental"))
+	UFUNCTION(BlueprintCallable, meta=(Category="Math|Spring|Experimental", AutoCreateRefTerm="TargetX"))
 	static ENGINE_API void CriticalSpringDampFloat(UPARAM(ref) float& InOutX, UPARAM(ref) float& InOutV, const float& TargetX, float DeltaTime, float SmoothingTime = 0.2f);
 
 	UE_EXPERIMENTAL(5.7, "SpringMath is experimental")
@@ -54,7 +54,7 @@ public:
 	 * @param SmoothingTime The smoothing time to use for the spring. Longer times lead to more damped behaviour. A time of 0 will snap InOutAngle to TargetAngle
 	 * @param DeltaTime Timestep in seconds
 	 */
-	UFUNCTION(BlueprintCallable, meta = (Category = "Math|Spring|Experimental"))
+	UFUNCTION(BlueprintCallable, meta = (Category = "Math|Spring|Experimental", AutoCreateRefTerm = "TargetAngle"))
 	static ENGINE_API void CriticalSpringDampAngle(UPARAM(ref) float& InOutAngle, UPARAM(ref) float& InOutAngularVelocity, const float& TargetAngle, float DeltaTime, float SmoothingTime = 0.2f);
 
 	UE_EXPERIMENTAL(5.7, "SpringMath is experimental")
@@ -66,7 +66,7 @@ public:
 	 * @param SmoothingTime The smoothing time to use for the spring. Longer times lead to more damped behaviour. A time of 0 will snap InOutX to TargetX
 	 * @param DeltaTime Timestep in seconds
 	 */
-	UFUNCTION(BlueprintCallable, meta=(Category="Math|Spring|Experimental"))
+	UFUNCTION(BlueprintCallable, meta=(Category="Math|Spring|Experimental", AutoCreateRefTerm = "TargetX"))
 	static ENGINE_API void CriticalSpringDampVector(UPARAM(ref) FVector& InOutX, UPARAM(ref) FVector& InOutV, const FVector& TargetX, float DeltaTime, float SmoothingTime = 0.2f);
 
 	UE_EXPERIMENTAL(5.7, "SpringMath is experimental")
@@ -78,7 +78,7 @@ public:
 	 * @param SmoothingTime The smoothing time to use for the spring. Longer times lead to more damped behaviour. A time of 0 will snap InOutX to TargetX
 	 * @param DeltaTime Timestep in seconds
 	 */
-	UFUNCTION(BlueprintCallable, meta=(Category="Math|Spring|Experimental"))
+	UFUNCTION(BlueprintCallable, meta=(Category="Math|Spring|Experimental", AutoCreateRefTerm = "TargetX"))
 	static ENGINE_API void CriticalSpringDampVector2D(UPARAM(ref) FVector2D& InOutX, UPARAM(ref) FVector2D& InOutV, const FVector2D& TargetX, float DeltaTime, float SmoothingTime = 0.2f);
 
 	UE_EXPERIMENTAL(5.7, "SpringMath is experimental")
@@ -90,7 +90,7 @@ public:
 	 * @param SmoothingTime The smoothing time to use for the spring. Longer times lead to more damped behaviour. A time of 0 will snap InOutRotation to TargetRotation
 	 * @param DeltaTime Timestep in seconds
 	 */
-	UFUNCTION(BlueprintCallable, meta=(Category="Math|Spring|Experimental"))
+	UFUNCTION(BlueprintCallable, meta=(Category="Math|Spring|Experimental", AutoCreateRefTerm = "TargetRotation"))
 	static ENGINE_API void CriticalSpringDampQuat(UPARAM(ref) FQuat& InOutRotation, UPARAM(ref) FVector& InOutAngularVelocity, const FQuat& TargetRotation, float DeltaTime, float SmoothingTime = 0.2f);
 
 	UE_EXPERIMENTAL(5.7, "SpringMath is experimental")
@@ -102,7 +102,7 @@ public:
 	 * @param SmoothingTime The smoothing time to use for the spring. Longer times lead to more damped behaviour. A time of 0 will snap InOutRotation to TargetRotation
 	 * @param DeltaTime Timestep in seconds
 	 */
-	UFUNCTION(BlueprintCallable, meta=(Category="Math|Spring|Experimental"))
+	UFUNCTION(BlueprintCallable, meta=(Category="Math|Spring|Experimental", AutoCreateRefTerm = "TargetRotation"))
 	static ENGINE_API void CriticalSpringDampRotator(UPARAM(ref) FRotator& InOutRotation, UPARAM(ref) FVector& InOutAngularVelocity, const FRotator& TargetRotation, float DeltaTime, float SmoothingTime = 0.2f);
 
 	UE_EXPERIMENTAL(5.8, "SpringMath is experimental")
@@ -114,7 +114,7 @@ public:
 	 * @param SmoothingTime The smoothing time to use for the spring. Longer times lead to more damped behaviour. A time of 0 will snap InOutScale to TargetScale
 	 * @param DeltaTime Timestep in seconds
 	 */
-	UFUNCTION(BlueprintCallable, meta = (Category = "Math|Spring|Experimental"))
+	UFUNCTION(BlueprintCallable, meta = (Category = "Math|Spring|Experimental", AutoCreateRefTerm = "TargetScale"))
 	static ENGINE_API void CriticalSpringDampScale(UPARAM(ref) FVector& InOutScale, UPARAM(ref) FVector& InOutScalarVelocity, const FVector& TargetScale, float DeltaTime, float SmoothingTime = 0.2f);
 
 	UE_EXPERIMENTAL(5.8, "SpringMath is experimental")
@@ -126,7 +126,7 @@ public:
 	 * @param SmoothingTime The smoothing time to use for the spring. Longer times lead to more damped behaviour. A time of 0 will snap InOutTransform to TargetTransform
 	 * @param DeltaTime Timestep in seconds
 	 */
-	UFUNCTION(BlueprintCallable, meta = (Category = "Math|Spring|Experimental"))
+	UFUNCTION(BlueprintCallable, meta = (Category = "Math|Spring|Experimental", AutoCreateRefTerm = "TargetTransform"))
 	static ENGINE_API void CriticalSpringDampTransform(UPARAM(ref) FTransform& InOutTransform, UPARAM(ref) FSpringTransformVelocity& InOutVelocity, const FTransform& TargetTransform, float DeltaTime, float SmoothingTime = 0.2f);
 
 
@@ -143,7 +143,7 @@ public:
 	 * @param SmoothingTime The smoothing time to use for the spring. Longer times lead to more damped behaviour. A time of 0 will snap InOutX to TargetX
 	 * @param DeltaTime Timestep in seconds
 	 */
-	UFUNCTION(BlueprintCallable, meta = (Category = "Math|Spring|Experimental"))
+	UFUNCTION(BlueprintCallable, meta = (Category = "Math|Spring|Experimental", AutoCreateRefTerm = "TargetX"))
 	static ENGINE_API void CriticalDoubleSpringDampFloat(UPARAM(ref) float& InOutX, UPARAM(ref) float& InOutV, UPARAM(ref) float& InOutXi, UPARAM(ref) float& InOutVi, const float& TargetX, float DeltaTime, float SmoothingTime = 0.2f);
 
 	UE_EXPERIMENTAL(5.7, "SpringMath is experimental")
@@ -157,7 +157,7 @@ public:
 	 * @param SmoothingTime The smoothing time to use for the spring. Longer times lead to more damped behaviour. A time of 0 will snap InOutAngle to TargetAngle
 	 * @param DeltaTime Timestep in seconds
 	 */
-	UFUNCTION(BlueprintCallable, meta = (Category = "Math|Spring|Experimental"))
+	UFUNCTION(BlueprintCallable, meta = (Category = "Math|Spring|Experimental", AutoCreateRefTerm = "TargetAngle"))
 	static ENGINE_API void CriticalDoubleSpringDampAngle(UPARAM(ref) float& InOutAngle,UPARAM(ref) float& InOutAngularVelocity, UPARAM(ref) float& InOutIntermediateAngle, UPARAM(ref) float& InOutIntermediateAngularVelocity, const float& TargetAngle, float DeltaTime, float SmoothingTime = 0.2f);
 
 	UE_EXPERIMENTAL(5.7, "SpringMath is experimental")
@@ -171,9 +171,9 @@ public:
 	 * @param SmoothingTime The smoothing time to use for the spring. Longer times lead to more damped behaviour. A time of 0 will snap InOutX to TargetX
 	 * @param DeltaTime Timestep in seconds
 	 */
-	UFUNCTION(BlueprintCallable, meta = (Category = "Math|Spring|Experimental"))
+	UFUNCTION(BlueprintCallable, meta = (Category = "Math|Spring|Experimental", AutoCreateRefTerm = "TargetAngle"))
 	static ENGINE_API void CriticalDoubleSpringDampVector(UPARAM(ref) FVector& InOutX, UPARAM(ref) FVector& InOutV, UPARAM(ref) FVector& InOutXi, UPARAM(ref) FVector& InOutVi, const FVector& TargetX, float DeltaTime, float SmoothingTime = 0.2f);
-
+	
 	UE_EXPERIMENTAL(5.7, "SpringMath is experimental")
 	/** Interpolates the value InOutX towards TargetX using a critically damped spring via an intermediate state - producing a more S-shaped curve. The velocity of X is stored in InOutV.
 	 *
@@ -185,7 +185,7 @@ public:
 	 * @param SmoothingTime The smoothing time to use for the spring. Longer times lead to more damped behaviour. A time of 0 will snap InOutX to TargetX
 	 * @param DeltaTime Timestep in seconds
 	 */
-	UFUNCTION(BlueprintCallable, meta = (Category = "Math|Spring|Experimental"))
+	UFUNCTION(BlueprintCallable, meta = (Category = "Math|Spring|Experimental", AutoCreateRefTerm = "TargetX"))
 	static ENGINE_API void CriticalDoubleSpringDampVector2D(UPARAM(ref) FVector2D& InOutX, UPARAM(ref) FVector2D& InOutV, UPARAM(ref) FVector2D& InOutXi, UPARAM(ref) FVector2D& InOutVi, const FVector2D& TargetX, float DeltaTime, float SmoothingTime = 0.2f);
 
 	UE_EXPERIMENTAL(5.7, "SpringMath is experimental")
@@ -199,7 +199,7 @@ public:
 	 * @param SmoothingTime The smoothing time to use for the spring. Longer times lead to more damped behaviour. A time of 0 will snap InOutRotation to TargetRotation
 	 * @param DeltaTime Timestep in seconds
 	 */
-	UFUNCTION(BlueprintCallable, meta = (Category = "Math|Spring|Experimental"))
+	UFUNCTION(BlueprintCallable, meta = (Category = "Math|Spring|Experimental", AutoCreateRefTerm = "TargetRotation"))
 	static ENGINE_API void CriticalDoubleSpringDampQuat(UPARAM(ref) FQuat& InOutRotation, UPARAM(ref) FVector& InOutAngularVelocity, UPARAM(ref) FQuat& InOutIntermediateRotation, UPARAM(ref) FVector& InOutIntermediateAngularVelocity, const FQuat& TargetRotation, float DeltaTime, float SmoothingTime = 0.2f);
 
 	UE_EXPERIMENTAL(5.7, "SpringMath is experimental")
@@ -213,7 +213,7 @@ public:
 	 * @param SmoothingTime The smoothing time to use for the spring. Longer times lead to more damped behaviour. A time of 0 will snap InOutRotation to TargetRotation
 	 * @param DeltaTime Timestep in seconds
 	 */
-	UFUNCTION(BlueprintCallable, meta = (Category = "Math|Spring|Experimental"))
+	UFUNCTION(BlueprintCallable, meta = (Category = "Math|Spring|Experimental", AutoCreateRefTerm = "TargetRotation"))
 	static ENGINE_API void CriticalDoubleSpringDampRotator(UPARAM(ref) FRotator& InOutRotation, UPARAM(ref) FVector& InOutAngularVelocity, UPARAM(ref) FRotator& InOutIntermediateRotation, UPARAM(ref) FVector& InOutIntermediateAngularVelocity, const FRotator& TargetRotation, float DeltaTime, float SmoothingTime = 0.2f);
 
 	UE_EXPERIMENTAL(5.8, "SpringMath is experimental")
@@ -227,7 +227,7 @@ public:
 	 * @param SmoothingTime The smoothing time to use for the spring. Longer times lead to more damped behaviour. A time of 0 will snap InOutScale to TargetScale
 	 * @param DeltaTime Timestep in seconds
 	 */
-	UFUNCTION(BlueprintCallable, meta = (Category = "Math|Spring|Experimental"))
+	UFUNCTION(BlueprintCallable, meta = (Category = "Math|Spring|Experimental", AutoCreateRefTerm = "TargetScale"))
 	static ENGINE_API void CriticalDoubleSpringDampScale(UPARAM(ref) FVector& InOutScale, UPARAM(ref) FVector& InOutScalarVelocity, UPARAM(ref) FVector& InOutIntermediateScale, UPARAM(ref) FVector& InOutIntermediateScalarVelocity, const FVector& TargetScale, float DeltaTime, float SmoothingTime = 0.2f);
 
 	UE_EXPERIMENTAL(5.8, "SpringMath is experimental")
@@ -241,7 +241,7 @@ public:
 	 * @param SmoothingTime The smoothing time to use for the spring. Longer times lead to more damped behaviour. A time of 0 will snap InOutTransform to TargetTransform
 	 * @param DeltaTime Timestep in seconds
 	 */
-	UFUNCTION(BlueprintCallable, meta = (Category = "Math|Spring|Experimental"))
+	UFUNCTION(BlueprintCallable, meta = (Category = "Math|Spring|Experimental", AutoCreateRefTerm = "TargetTransform"))
 	static ENGINE_API void CriticalDoubleSpringDampTransform(UPARAM(ref) FTransform& InOutTransform, UPARAM(ref) FSpringTransformVelocity& InOutVelocity, UPARAM(ref) FTransform& InOutIntermediateTransform, UPARAM(ref) FSpringTransformVelocity& InOutIntermediateTransformVelocity, const FTransform& TargetTransform, float DeltaTime, float SmoothingTime = 0.2f);
 
 
@@ -331,7 +331,7 @@ public:
 	 * @param  DeltaTime Time interval
 	 * @param  SmoothingTime Timescale over which to smooth. Larger values result in more smoothed behaviour. Can be zero.
 	 */
-	UFUNCTION(BlueprintPure, meta=(Category="Math|Spring|Experimental"))
+	UFUNCTION(BlueprintPure, meta=(Category="Math|Spring|Experimental", AutoCreateRefTerm = "Value, Target"))
 	static ENGINE_API FVector DampVector(const FVector& Value, const FVector& Target, float DeltaTime, float SmoothingTime = 0.2f);
 
 	UE_EXPERIMENTAL(5.7, "SpringMath is experimental")
@@ -344,7 +344,7 @@ public:
 	 * @param  DeltaTime Time interval
 	 * @param  SmoothingTime Timescale over which to smooth. Larger values result in more smoothed behaviour. Can be zero.
 	 */
-	UFUNCTION(BlueprintPure, meta=(Category="Math|Spring|Experimental"))
+	UFUNCTION(BlueprintPure, meta=(Category="Math|Spring|Experimental", AutoCreateRefTerm = "Value, Target"))
 	static ENGINE_API FVector2D DampVector2D(const FVector2D& Value, const FVector2D& Target, float DeltaTime, float SmoothingTime = 0.2f);
 
 	UE_EXPERIMENTAL(5.7, "SpringMath is experimental")
@@ -357,7 +357,7 @@ public:
 	 * @param  DeltaTime Time interval
 	 * @param  SmoothingTime Timescale over which to smooth. Larger values result in more smoothed behaviour. Can be zero.
 	 */
-	UFUNCTION(BlueprintPure, meta=(Category="Math|Spring|Experimental"))
+	UFUNCTION(BlueprintPure, meta=(Category="Math|Spring|Experimental", AutoCreateRefTerm = "Rotation, TargetRotation"))
 	static ENGINE_API FQuat DampQuat(const FQuat& Rotation, const FQuat& TargetRotation, float DeltaTime, float SmoothingTime = 0.2f);
 
 	UE_EXPERIMENTAL(5.7, "SpringMath is experimental")
@@ -370,7 +370,7 @@ public:
 	 * @param  DeltaTime Time interval
 	 * @param  SmoothingTime Timescale over which to smooth. Larger values result in more smoothed behaviour. Can be zero.
 	 */
-	UFUNCTION(BlueprintPure, meta=(Category="Math|Spring|Experimental"))
+	UFUNCTION(BlueprintPure, meta=(Category="Math|Spring|Experimental", AutoCreateRefTerm = "Rotation, TargetRotation"))
 	static ENGINE_API FRotator DampRotator(const FRotator& Rotation, const FRotator& TargetRotation, float DeltaTime, float SmoothingTime = 0.2f);
 
 	UE_EXPERIMENTAL(5.8, "SpringMath is experimental")
@@ -383,7 +383,7 @@ public:
 	 * @param  DeltaTime Time interval
 	 * @param  SmoothingTime Timescale over which to smooth. Larger values result in more smoothed behaviour. Can be zero.
 	 */
-	UFUNCTION(BlueprintPure, meta = (Category = "Math|Spring|Experimental"))
+	UFUNCTION(BlueprintPure, meta = (Category = "Math|Spring|Experimental", AutoCreateRefTerm = "Scale, TargetScale"))
 	static ENGINE_API FVector DampScale(const FVector& Scale, const FVector& TargetScale, float DeltaTime, float SmoothingTime = 0.2f);
 
 	UE_EXPERIMENTAL(5.8, "SpringMath is experimental")
@@ -396,7 +396,7 @@ public:
 	 * @param  DeltaTime Time interval
 	 * @param  SmoothingTime Timescale over which to smooth. Larger values result in more smoothed behaviour. Can be zero.
 	 */
-	UFUNCTION(BlueprintPure, meta = (Category = "Math|Spring|Experimental"))
+	UFUNCTION(BlueprintPure, meta = (Category = "Math|Spring|Experimental", AutoCreateRefTerm = "Transform, TargetTransform"))
 	static ENGINE_API FTransform DampTransform(const FTransform& Transform, const FTransform& TargetTransform, float DeltaTime, float SmoothingTime = 0.2f);
 
 
@@ -412,7 +412,7 @@ public:
 	 * @param DeltaTime The delta time to tick the character
 	 * @param SmoothingTime The time over which to smooth velocity. It takes roughly the smoothing time in order for the character to reach the target velocity.
 	 */
-	UFUNCTION(BlueprintCallable, meta=(Category="Math|Spring|Experimental"))
+	UFUNCTION(BlueprintCallable, meta=(Category="Math|Spring|Experimental", AutoCreateRefTerm = "TargetVelocity"))
 	static ENGINE_API void SpringCharacterUpdate(UPARAM(ref) FVector& InOutPosition, UPARAM(ref) FVector& InOutVelocity, UPARAM(ref) FVector& InOutAcceleration, const FVector& TargetVelocity, float DeltaTime, float SmoothingTime = 0.2f);
 
 	UE_EXPERIMENTAL(5.8, "SpringMath is experimental")
@@ -584,7 +584,7 @@ public:
 	 * @param SmoothingTime The time over which to smooth velocity. It takes roughly the smoothing time in order for the character to reach the target velocity.
 	 * @param MaxAcceleration Puts a limit on the maximum acceleration that the intermediate velocity can do each frame. If MaxAccel is very large, the behaviour will be the same as SpringCharacterUpdate
 	 */
-	UFUNCTION(BlueprintCallable, meta=(Category="Math|Spring|Experimental"))
+	UFUNCTION(BlueprintCallable, meta=(Category="Math|Spring|Experimental", AutoCreateRefTerm = "TargetVelocity"))
 	static ENGINE_API void VelocitySpringCharacterUpdate(UPARAM(ref) FVector& InOutPosition, UPARAM(ref) FVector& InOutVelocity, UPARAM(ref) FVector& InOutVelocityIntermediate, UPARAM(ref) FVector& InOutAcceleration, const FVector& TargetVelocity, float DeltaTime, float SmoothingTime = 0.2f, float MaxAcceleration = 1.0f);
 
 
