Substrate - NMT - added support for ConvertToDecal node.

#rnx
#Synced-CL 51590311

#Synced-CL 51590311

diff --git a/Engine/Source/Runtime/Engine/Private/Materials/MaterialExpressionsToMIR.cpp b/Engine/Source/Runtime/Engine/Private/Materials/MaterialExpressionsToMIR.cpp
--- a/Engine/Source/Runtime/Engine/Private/Materials/MaterialExpressionsToMIR.cpp
+++ b/Engine/Source/Runtime/Engine/Private/Materials/MaterialExpressionsToMIR.cpp
@@ -6644,7 +6644,62 @@ void UMaterialExpressionSubstrateUI::Build(MIR::FEmitter& Em)
 	);
 
 	FinalizeBSDFAndPromoteToOperatorIfNeeded(Em, SubstrateOperator, UIBSDF);
+}
+
+void UMaterialExpressionSubstrateConvertToDecal::Build(MIR::FEmitter& Em)
+{
+	if (!DecalMaterial.GetTracedInput().Expression)
+	{
+		return Em.Errorf(TEXT("Missing ConvertToDecal node input."));
+	}
+
+	FValueRef CoverageArg = CompileWithDefaultFloat1(Em, Coverage, 1.0f);
+
+	// Note: we no longer do SubstrateTreeStackPush/Pop since the path guid is regenerated for every substrate node from BuildTopMaterialExpression.
+	FValueRef DecalMaterialArg = Em.TryInput(&DecalMaterial);
+
+	FSubstrateTranslatorData* SubstrateTranslatorData = Em.GetSubstrateTranslatorData();
+	FGuid PathUniqueId = SubstrateTranslatorData->SubstrateTreeStackGetPathUniqueId();
+	FSubstrateOperator& SubstrateOperator = SubstrateTranslatorData->SubstrateCompilationGetOperator(PathUniqueId);
+
+	if (!SubstrateOperator.bUseParameterBlending)
+	{
+		return Em.Errorf(TEXT("Substrate Convert To Decal node must receive SubstrateData a parameter blended Substrate material sub tree."));
+	}
+	if (!SubstrateOperator.bRootOfParameterBlendingSubTree)
+	{
+		return Em.Errorf(TEXT("Substrate Convert To Decal node must be the root of a parameter blending sub tree: no more Substrate operations can be applied a over its output."));
+	}
+
+	// Propagate the parameter blended normal
+	FSubstrateOperator* Operator = SubstrateTranslatorData->SubstrateCompilationGetOperatorFromIndex(SubstrateOperator.LeftIndex);
+	SubstrateOperator.BSDFRegisteredSharedLocalBasis = Operator->BSDFRegisteredSharedLocalBasis;
+
+	// Now emit the WEight operator acting as conversion to decal with coverage.
+	if (!DecalMaterialArg.IsValid())
+	{
+		return Em.Errorf(TEXT("A input input graphs could not be evaluated for parameter blending."));
+	}
+	if (!CoverageArg.IsValid())
+	{
+		return Em.Errorf(TEXT("Weight input graphs could not be evaluated for parameter blending."));
+	}
+
+	check(SubstrateOperator.Index != INDEX_NONE);
+	check(SubstrateOperator.BSDFIndex != INDEX_NONE);
+
+	FValueRef DecalMaterialWeightOpArg = Em.SubstrateCoverageWeightParameterBlending(
+		DecalMaterialArg,
+		CoverageArg);
+
+	FValueRef PromotedDecalMaterialWeightOpArg = Em.SubstratePromoteToOperator(
+		DecalMaterialWeightOpArg,
+		Em.ConstantInt(SubstrateOperator.Index),
+		Em.ConstantInt(SubstrateOperator.BSDFIndex),
+		Em.ConstantInt(SubstrateOperator.LayerDepth),
+		Em.ConstantInt(SubstrateOperator.bIsBottom ? 1 : 0));
 
+	Em.Output(0, PromotedDecalMaterialWeightOpArg);
 }
 
 static FValueRef SubstrateBlendNormal(MIR::FEmitter& Em, FValueRef NormalCodeChunk0, FValueRef NormalCodeChunk1, FValueRef MixCodeChunk)
diff --git a/Engine/Source/Runtime/Engine/Public/Materials/MaterialExpressionSubstrate.h b/Engine/Source/Runtime/Engine/Public/Materials/MaterialExpressionSubstrate.h
--- a/Engine/Source/Runtime/Engine/Public/Materials/MaterialExpressionSubstrate.h
+++ b/Engine/Source/Runtime/Engine/Public/Materials/MaterialExpressionSubstrate.h
@@ -1026,7 +1026,7 @@ class UMaterialExpressionSubstrateConvertToDecal : public UMaterialExpressionSub
 	//~ Begin UMaterialExpression Interface
 #if WITH_EDITOR
 	virtual int32 Compile(class FMaterialCompiler* Compiler, int32 OutputIndex) override;
-	// SUBSTRATE_TODO virtual void Build(MIR::FEmitter& Emitter) override;
+	virtual void Build(MIR::FEmitter& Emitter) override;
 	virtual void GetCaption(TArray<FString>& OutCaptions) const override;
 	virtual EMaterialValueType GetOutputValueType(int32 OutputIndex) override;
 	virtual EMaterialValueType GetInputValueType(int32 InputIndex) override;
