[HWRT] Move core of UpdateReferenceBasedResidency(...) off the render thread.
- New struct FReferenceBasedResidencyState which contains a shadow copy of geometry/group metadata that can be processed on a background task thread.
- The old UpdateReferenceBasedResidency() (which ran entirely on the render thread) is split into:
- UpdateReferenceBasedResidencyState() syncs metadata from the manager to the state struct (render thread).
- FReferenceBasedResidencyState::Update() does the expensive computation (background task).
- ProcessReferenceBasedResidency()` consumes the results (render thread).

#Synced-CL 51569112

diff --git a/Engine/Source/Runtime/Engine/Private/Rendering/RayTracingGeometryManager.cpp b/Engine/Source/Runtime/Engine/Private/Rendering/RayTracingGeometryManager.cpp
--- a/Engine/Source/Runtime/Engine/Private/Rendering/RayTracingGeometryManager.cpp
+++ b/Engine/Source/Runtime/Engine/Private/Rendering/RayTracingGeometryManager.cpp
@@ -188,10 +188,224 @@ DECLARE_DWORD_ACCUMULATOR_STAT(TEXT("In-flight Streaming Requests"), STAT_RayTra
 
 CSV_DEFINE_CATEGORY(RayTracingGeometry, true);
 
+struct FReferenceBasedResidencyState
+{
+	struct FGeometryMetadata
+	{
+		uint64 LastReferencedFrame = 0;
+		uint32 Size = 0;
+		RayTracing::FGeometryGroupHandle GroupHandle;
+
+		bool bAlwaysResident : 1 = false;
+		bool bEvicted : 1 = false;
+		bool bBVHStreamedOut : 1 = false;
+		bool bBuffersStreamedOut : 1 = false;
+	};
+
+	struct FGeometryGroupMetadata
+	{
+		TArray<RayTracing::FGeometryHandle, TInlineAllocator<2>> GeometryHandles;
+		int32 CurrentFirstLODIdx = INDEX_NONE;
+	};
+
+	FGeometryMetadata& GetGeometryMetadata(RayTracing::FGeometryHandle Handle);
+	FGeometryGroupMetadata& GetGroupMetadata(RayTracing::FGeometryGroupHandle Handle);
+
+	void Update();
+
+	uint64 FrameCounter = ~0ULL;
+
+	int64 TotalResidentSize = 0;
+	int64 TotalAlwaysResidentSize = 0;
+	int64 TotalStreamingSize = 0;
+
+	TSparseArray<FGeometryMetadata> GeometryMetadatas;
+	TSparseArray<FGeometryGroupMetadata> GroupMetadatas;
+
+	TSet<RayTracing::FGeometryHandle> EvictableGeometries;
+
+	TSet<RayTracing::FGeometryHandle> ReferencedGeometries;
+	TSet<RayTracing::FGeometryGroupHandle> ReferencedGeometryGroups;
+	TSet<RayTracing::FGeometryGroupHandle> ReferencedGeometryGroupsForDynamicUpdate;
+
+	// Update output
+
+	TArray<RayTracing::FGeometryHandle> GeometriesToEvict;
+	TArray<RayTracing::FGeometryHandle> GeometriesToMakeResidentAndStreamIn;
+
+	int64 ReferencedSize = 0;
+	int64 RequestedSize = 0;
+	int64 RequestedButEvictedSize = 0;
+};
+
+FReferenceBasedResidencyState::FGeometryMetadata& FReferenceBasedResidencyState::GetGeometryMetadata(RayTracing::FGeometryHandle Handle)
+{
+	checkf(!Handle.IsNull(), TEXT("Must provide valid RayTracing::FGeometryHandle."));
+
+	const uint32 GeometryIndex = FRayTracingGeometryManager::GetHandleIndex(Handle);
+
+	checkf(GeometryMetadatas.IsValidIndex(GeometryIndex), TEXT("Stale RayTracing::FGeometryHandle detected - invalid index %u."), GeometryIndex);
+
+	return GeometryMetadatas[GeometryIndex];
+}
+
+FReferenceBasedResidencyState::FGeometryGroupMetadata& FReferenceBasedResidencyState::GetGroupMetadata(RayTracing::FGeometryGroupHandle Handle)
+{
+	checkf(!Handle.IsNull(), TEXT("Must provide valid RayTracing::FGeometryGroupHandle."));
+
+	const uint32 GroupIndex = FRayTracingGeometryManager::GetHandleIndex(Handle);
+
+	checkf(GroupMetadatas.IsValidIndex(GroupIndex), TEXT("Stale RayTracing::FGeometryGroupHandle detected - invalid index %u."), GroupIndex);
+
+	return GroupMetadatas[GroupIndex];
+}
+
+void FReferenceBasedResidencyState::Update()
+{
+	TRACE_CPUPROFILER_EVENT_SCOPE(FReferenceBasedResidencyState::Update);
+
+	// TODO: Should also add support for evicting resident VB/IB buffers when no longer referenced in ReferencedGeometryGroupsForDynamicUpdate
+	// Need a separate set of "EvictableBuffers" set and logic to handle it.
+
+	TSet<RayTracing::FGeometryHandle> NotReferencedResidentGeometries = EvictableGeometries;
+
+	GeometriesToMakeResidentAndStreamIn.Reserve(GeometryMetadatas.Num());
+
+	ReferencedSize = 0;
+	RequestedSize = TotalAlwaysResidentSize;
+	RequestedButEvictedSize = 0;
+
+	// Step 1
+	// - update LastReferencedFrame of referenced geometries and calculate memory required to make evicted geometries resident
+	for (RayTracing::FGeometryHandle GeometryHandle : ReferencedGeometries)
+	{
+		FGeometryMetadata& GeometryMetadata = GetGeometryMetadata(GeometryHandle);
+		GeometryMetadata.LastReferencedFrame = FrameCounter;
+
+		if (GeometryMetadata.bEvicted || GeometryMetadata.bBVHStreamedOut || (GeometryMetadata.bBuffersStreamedOut && ReferencedGeometryGroupsForDynamicUpdate.Contains(GeometryMetadata.GroupHandle)))
+		{
+			GeometriesToMakeResidentAndStreamIn.Add(GeometryHandle);
+		}
+
+		NotReferencedResidentGeometries.Remove(GeometryHandle);
+
+		ReferencedSize += GeometryMetadata.Size;
+
+		if (!GeometryMetadata.bAlwaysResident)
+		{
+			RequestedSize += GeometryMetadata.Size;
+		}
+
+		if (GeometryMetadata.bEvicted)
+		{
+			RequestedButEvictedSize += GeometryMetadata.Size;
+		}
+	}
+
+	// Step 2
+	// - add all geometries in referenced groups to ReferencedGeometriesEvictedOrRequiringStreaming
+	//		- need to make all geometries in group resident otherwise might not have valid geometry when reducing LOD
+	//		- TODO: Could track TargetLOD and only make [TargetLOD ... LastLOD] range resident
+	// - also update LastReferencedFrame and calculate memory required to make evicted geometries resident
+	for (RayTracing::FGeometryGroupHandle GroupHandle : ReferencedGeometryGroups)
+	{
+		const FGeometryGroupMetadata& GroupMetadata = GetGroupMetadata(GroupHandle);
+		const int32 NumGeometriesInGroup = GroupMetadata.GeometryHandles.Num();
+
+		for (int32 LODIndex = GroupMetadata.CurrentFirstLODIdx; LODIndex < NumGeometriesInGroup; ++LODIndex)
+		{
+			RayTracing::FGeometryHandle GeometryHandle = GroupMetadata.GeometryHandles[LODIndex];
+
+			if (!GeometryHandle.IsNull()) // some LODs might be stripped during cook
+			{
+				FGeometryMetadata& GeometryMetadata = GetGeometryMetadata(GeometryHandle);
+				GeometryMetadata.LastReferencedFrame = FrameCounter;
+
+				ReferencedSize += GeometryMetadata.Size;
+
+				if (GeometryMetadata.bAlwaysResident)
+				{
+					checkf(!GeometryMetadata.bEvicted, TEXT("Always-resident ray tracing geometry was unexpectedly evicted."));
+				}
+				else
+				{
+					if (GeometryMetadata.bEvicted || GeometryMetadata.bBVHStreamedOut || (GeometryMetadata.bBuffersStreamedOut && ReferencedGeometryGroupsForDynamicUpdate.Contains(GeometryMetadata.GroupHandle)))
+					{
+						GeometriesToMakeResidentAndStreamIn.Add(GeometryHandle);
+					}
+
+					NotReferencedResidentGeometries.Remove(GeometryHandle);
+
+					RequestedSize += GeometryMetadata.Size;
+
+					if (GeometryMetadata.bEvicted)
+					{
+						RequestedButEvictedSize += GeometryMetadata.Size;
+					}
+				}
+			}
+		}
+	}
+
+#if DO_CHECK
+	// ensure(GeometriesToMakeResidentAndStreamIn.Num() == TSet(GeometriesToMakeResidentAndStreamIn).Num());
+#endif
+
+	const int64 ResidentGeometryMemoryPoolSize = FUnitConversion::Convert(GRayTracingResidentGeometryMemoryPoolSizeInMB, EUnit::Megabytes, EUnit::Bytes);
+
+	// Step 3
+	// - if making requested geometries resident will put us over budget -> evict some geometry not referenced by TLAS
+	if (TotalResidentSize + TotalStreamingSize + RequestedButEvictedSize > ResidentGeometryMemoryPoolSize)
+	{
+		GeometriesToEvict = NotReferencedResidentGeometries.Array();
+
+		// sort to evict geometries in the following order:
+		//	- least recently used
+		//	- largest geometries
+		Algo::Sort(GeometriesToEvict, [this](RayTracing::FGeometryHandle& LHSHandle, RayTracing::FGeometryHandle& RHSHandle)
+			{
+				FGeometryMetadata& LHS = GetGeometryMetadata(LHSHandle);
+				FGeometryMetadata& RHS = GetGeometryMetadata(RHSHandle);
+
+				// TODO: evict unreferenced dynamic geometries using shared buffers first since they need to be rebuild anyway
+				// (and then dynamic geometries requiring update?
+
+				// 1st - last referenced frame
+				if (LHS.LastReferencedFrame != RHS.LastReferencedFrame)
+				{
+					return LHS.LastReferencedFrame < RHS.LastReferencedFrame;
+				}
+
+				// 2nd - size
+				return LHS.Size > RHS.Size;
+			});
+	}
+	else
+	{
+		GeometriesToEvict.Empty();
+	}
+
+	// Step 4
+	// - make referenced geometries resident until we go over budget
+	if (TotalResidentSize + TotalStreamingSize < ResidentGeometryMemoryPoolSize)
+	{
+		// sort by size to prioritize smaller geometries
+		Algo::Sort(GeometriesToMakeResidentAndStreamIn, [this](RayTracing::FGeometryHandle& LHSHandle, RayTracing::FGeometryHandle& RHSHandle)
+			{
+				FGeometryMetadata& LHS = GetGeometryMetadata(LHSHandle);
+				FGeometryMetadata& RHS = GetGeometryMetadata(RHSHandle);
+
+				return LHS.Size < RHS.Size;
+			});
+	}
+}
+
 FRayTracingGeometryManager::FRayTracingGeometryManager()
 {
 	StreamingRequests.SetNum(GRayTracingStreamingMaxPendingRequests);
 
+	ReferenceBasedResidencyState = MakePimpl<FReferenceBasedResidencyState>();
+
 #if CSV_PROFILER_STATS
 	// Need to defer until RHI is initialized.
 	FCoreDelegates::GetOnPostEngineInit().AddLambda([]
@@ -206,6 +420,9 @@ FRayTracingGeometryManager::~FRayTracingGeometryManager()
 	ensure(GeometryBuildRequests.IsEmpty());
 	ensure(RegisteredGeometries.IsEmpty());
 	ensure(RegisteredGroups.IsEmpty());
+
+	ReferenceBasedResidencyTask.Wait();
+	ReferenceBasedResidencyTask = {};
 }
 
 FRayTracingGeometryManager::FBuildRequest& FRayTracingGeometryManager::GetBuildRequest(RayTracing::FBuildRequestHandle Handle)
@@ -221,6 +438,22 @@ FRayTracingGeometryManager::FBuildRequest& FRayTracingGeometryManager::GetBuildR
 	return GeometryBuildRequests[RequestIndex];
 }
 
+bool FRayTracingGeometryManager::IsHandleValid(RayTracing::FGeometryHandle Handle) const
+{
+	const uint32 GeometryIndex = GetHandleIndex(Handle);
+	const uint32 GeometryVersion = GetHandleVersion(Handle);
+
+	return RegisteredGeometries.IsValidIndex(GeometryIndex) && (GeometryVersion == RegisteredGeometriesVersions[GeometryIndex]);
+}
+
+bool FRayTracingGeometryManager::IsHandleValid(RayTracing::FGeometryGroupHandle Handle) const
+{
+	const uint32 GroupIndex = GetHandleIndex(Handle);
+	const uint32 GroupVersion = GetHandleVersion(Handle);
+
+	return RegisteredGroups.IsValidIndex(GroupIndex) && (GroupVersion == RegisteredGroupsVersions[GroupIndex]);
+}
+
 uint32 FRayTracingGeometryManager::GetRegisteredGeometryIndex(RayTracing::FGeometryHandle Handle) const
 {
 	checkf(!Handle.IsNull(), TEXT("Must provide valid RayTracing::FGeometryHandle."));
@@ -353,6 +586,8 @@ RayTracing::FGeometryGroupHandle FRayTracingGeometryManager::RegisterRayTracingG
 
 	INC_DWORD_STAT(STAT_RayTracingGeometryGroupCount);
 
+	ModifiedGroups.Add(Handle);
+
 	return Handle;
 }
 
@@ -386,6 +621,13 @@ void FRayTracingGeometryManager::ReleaseRayTracingGeometryGroupReference(RayTrac
 		ReferencedGeometryGroupsForDynamicUpdate.Remove(Handle);
 
 		DEC_DWORD_STAT(STAT_RayTracingGeometryGroupCount);
+
+		ReleasedGroups.Add(Handle);
+		ModifiedGroups.Remove(Handle);
+	}
+	else
+	{
+		ModifiedGroups.Add(Handle);
 	}
 }
 
@@ -405,7 +647,6 @@ RayTracing::FGeometryHandle FRayTracingGeometryManager::RegisterRayTracingGeomet
 
 	FRegisteredGeometry& RegisteredGeometry = RegisteredGeometries[GeometryIndex];
 	RegisteredGeometry.Geometry = InGeometry;
-	RegisteredGeometry.LastReferencedFrame = 0;
 
 	const RayTracing::FGeometryGroupHandle GroupHandle = InGeometry->GroupHandle;
 	const int8 LODIndex = InGeometry->LODIndex;
@@ -420,6 +661,8 @@ RayTracing::FGeometryHandle FRayTracingGeometryManager::RegisterRayTracingGeomet
 		Group.GeometryHandles[LODIndex] = Handle;
 		++Group.NumReferences;
 
+		ModifiedGroups.Add(GroupHandle);
+
 		RegisteredGeometry.bAlwaysResident = IsAlwaysResidentGeometry(InGeometry, Group);
 
 		if (RegisteredGeometry.bAlwaysResident)
@@ -488,7 +731,6 @@ void FRayTracingGeometryManager::ReleaseRayTracingGeometryHandle(RayTracing::FGe
 	RegisteredGeometries.RemoveAt(GeometryIndex);
 	RegisteredGeometriesVersions[GeometryIndex] = (RegisteredGeometriesVersions[GeometryIndex] + 1) & HandleVersionMask;
 
-	EvictableGeometries.Remove(Handle);
 	ReferencedGeometries.Remove(Handle);
 	if (PendingStreamingRequests.Remove(Handle) > 0)
 	{
@@ -496,6 +738,9 @@ void FRayTracingGeometryManager::ReleaseRayTracingGeometryHandle(RayTracing::FGe
 	}
 
 	DEC_DWORD_STAT(STAT_RayTracingGeometryCount);
+
+	ReleasedGeometries.Add(Handle);
+	ModifiedGeometries.Remove(Handle);
 }
 
 void FRayTracingGeometryManager::SetRayTracingGeometryStreamingData(const FRayTracingGeometry* Geometry, FRayTracingStreamableAsset& StreamableAsset)
@@ -550,7 +795,7 @@ void FRayTracingGeometryManager::SetRayTracingGeometryGroupCurrentFirstLODIndex(
 		{
 			if (!Group.GeometryHandles[LODIdx].IsNull())
 			{
-				// TODO: should do this for always resident mips even when using reference based residency
+				// TODO: should do this for always-resident mips even when using reference based residency
 				PendingStreamingRequests.Add(Group.GeometryHandles[LODIdx]);
 				INC_DWORD_STAT(STAT_RayTracingPendingStreamingRequests);
 			}
@@ -558,6 +803,8 @@ void FRayTracingGeometryManager::SetRayTracingGeometryGroupCurrentFirstLODIndex(
 	}
 
 	Group.CurrentFirstLODIdx = NewCurrentFirstLODIdx;
+
+	ModifiedGroups.Add(Handle);
 }
 
 static bool ShouldCompactAfterBuild(const FRayTracingGeometryInitializer& Initializer)
@@ -630,11 +877,6 @@ void FRayTracingGeometryManager::RefreshRegisteredGeometry(RayTracing::FGeometry
 			}
 
 			TotalResidentSize += RegisteredGeometry.Size;
-
-			if (!RegisteredGeometry.bAlwaysResident)
-			{
-				EvictableGeometries.Add(Handle);
-			}
 		}
 		else
 		{
@@ -644,11 +886,11 @@ void FRayTracingGeometryManager::RefreshRegisteredGeometry(RayTracing::FGeometry
 			{
 				TotalResidentSize -= OldSize;
 			}
-
-			EvictableGeometries.Remove(Handle);
 		}
 
-		checkf(!AlwaysResidentGeometries.Contains(Handle) || !RegisteredGeometry.bEvicted || !IsRayTracingEnabled(), TEXT("Always resident geometries can't be evicted"));
+		checkf(!AlwaysResidentGeometries.Contains(Handle) || !RegisteredGeometry.bEvicted || !IsRayTracingEnabled(), TEXT("Always-resident geometries can't be evicted"));
+
+		ModifiedGeometries.Add(Handle);
 	}
 }
 
@@ -657,6 +899,14 @@ void FRayTracingGeometryManager::PreRender()
 	bRenderedFrame = true;
 }
 
+void FRayTracingGeometryManager::Update(FRHICommandList& RHICmdList)
+{
+	ReferenceBasedResidencyTask.Wait();
+	ReferenceBasedResidencyTask = {};
+
+	ProcessReferenceBasedResidency(RHICmdList);
+}
+
 void FRayTracingGeometryManager::Tick(FRHICommandList& RHICmdList)
 {
 	if (IsRunningCommandlet())
@@ -670,6 +920,9 @@ void FRayTracingGeometryManager::Tick(FRHICommandList& RHICmdList)
 	QUICK_SCOPE_CYCLE_COUNTER(STAT_FRayTracingGeometryManager_Tick);
 	CSV_SCOPED_TIMING_STAT_EXCLUSIVE(RayTracingGeometryManager_Tick);
 
+	ReferenceBasedResidencyTask.Wait();
+	ReferenceBasedResidencyTask = {};
+
 	// TODO: investigate fine grained locking to minimize blocking progress on render command pipes
 	// - Don't touch registered geometry/group arrays from render command pipes
 	//   - Separate arrays of free geometry/group handles + HandleAllocationCS
@@ -692,6 +945,8 @@ void FRayTracingGeometryManager::Tick(FRHICommandList& RHICmdList)
 
 	RefreshAlwaysResidentRayTracingGeometries(RHICmdList);
 
+	UpdateReferenceBasedResidencyState();
+
 #if DO_CHECK
 	CheckIntegrity();
 #endif
@@ -732,30 +987,26 @@ void FRayTracingGeometryManager::Tick(FRHICommandList& RHICmdList)
 			return;
 		}
 
-		bRenderedFrame = false;
-
-		if (bHasRayTracingEnableChanged)
+#if !UE_BUILD_SHIPPING
+		if (bDumpUnreferencedAlwaysResidentGeometriesNextFrame)
 		{
-			// make always resident geometries actually resident
-
-			for (RayTracing::FGeometryHandle GeometryHandle : AlwaysResidentGeometries)
-			{
-				FRegisteredGeometry& RegisteredGeometry = GetRegisteredGeometry(GeometryHandle);
-
-				if (RegisteredGeometry.Geometry->IsEvicted())
-				{
-					MakeGeometryResident(RHICmdList, RegisteredGeometry);
-				}
-				if (!RequestRayTracingGeometryStreamIn(RHICmdList, GeometryHandle))
-				{
-					PendingStreamingRequests.Add(GeometryHandle);
+			bDumpUnreferencedAlwaysResidentGeometriesNextFrame = false;
 
-					INC_DWORD_STAT(STAT_RayTracingPendingStreamingRequests);
-				}
-			}
+			DumpUnreferencedAlwaysResidentRayTracingGeometries();
 		}
+#endif // !UE_BUILD_SHIPPING
+
+		// TODO: swap instead of copying?
+		ReferenceBasedResidencyState->ReferencedGeometries = ReferencedGeometries;
+		ReferenceBasedResidencyState->ReferencedGeometryGroups = ReferencedGeometryGroups;
+		ReferenceBasedResidencyState->ReferencedGeometryGroupsForDynamicUpdate = ReferencedGeometryGroupsForDynamicUpdate;
 
-		UpdateReferenceBasedResidency(RHICmdList);
+		ReferenceBasedResidencyTask = UE::Tasks::Launch(UE_SOURCE_LOCATION,
+			[ReferenceBasedResidencyState = ReferenceBasedResidencyState.Get()]()
+			{
+				FTaskTagScope TaskTagScope(ETaskTag::EParallelRenderingThread);
+				ReferenceBasedResidencyState->Update();
+			});
 	}
 	else
 	{
@@ -795,6 +1046,7 @@ void FRayTracingGeometryManager::Tick(FRHICommandList& RHICmdList)
 	ReferencedGeometryGroupsForDynamicUpdate.Reset();
 
 	bHasRayTracingEnableChanged = false;
+	bRenderedFrame = false;
 
 	SET_MEMORY_STAT(STAT_RayTracingGeometryResidentMemory, TotalResidentSize);
 	SET_MEMORY_STAT(STAT_RayTracingGeometryAlwaysResidentMemory, TotalAlwaysResidentSize);
@@ -803,180 +1055,148 @@ void FRayTracingGeometryManager::Tick(FRHICommandList& RHICmdList)
 	CSV_CUSTOM_STAT(RayTracingGeometry, TotalAlwaysResidentSizeMB, TotalAlwaysResidentSize / 1024.0f / 1024.0f, ECsvCustomStatOp::Set);
 }
 
-void FRayTracingGeometryManager::UpdateReferenceBasedResidency(FRHICommandList& RHICmdList)
+void FRayTracingGeometryManager::UpdateReferenceBasedResidencyState()
 {
-	TSet<RayTracing::FGeometryHandle> NotReferencedResidentGeometries = EvictableGeometries;
+	TRACE_CPUPROFILER_EVENT_SCOPE(FRayTracingGeometryManager::UpdateReferenceBasedResidencyState);
 
-	TArray<RayTracing::FGeometryHandle> ReferencedGeometriesEvictedOrRequiringStreaming;
-	ReferencedGeometriesEvictedOrRequiringStreaming.Reserve(RegisteredGeometries.Num());
+	ReferenceBasedResidencyState->FrameCounter = GFrameCounterRenderThread;
+	ReferenceBasedResidencyState->TotalResidentSize = TotalResidentSize;
+	ReferenceBasedResidencyState->TotalAlwaysResidentSize = TotalAlwaysResidentSize;
+	ReferenceBasedResidencyState->TotalStreamingSize = TotalStreamingSize;
 
-	int64 ReferencedSize = 0;
-	int64 RequestedSize = TotalAlwaysResidentSize;
-	int64 RequestedButEvictedSize = 0;
-
-	// Step 1
-	// - update LastReferencedFrame of referenced geometries and calculate memory required to make evicted geometries resident
-	for (RayTracing::FGeometryHandle GeometryHandle : ReferencedGeometries)
+	// Remove metadata for geometries/groups released during the previous frame
+	for (RayTracing::FGeometryHandle Handle : ReleasedGeometries)
 	{
-		FRegisteredGeometry& RegisteredGeometry = GetRegisteredGeometry(GeometryHandle);
-		RegisteredGeometry.LastReferencedFrame = GFrameCounterRenderThread;
-
-		bool bStreamBVH;
-		bool bStreamBuffers;
-		DoesRayTracingGeometryRequireStreaming(RegisteredGeometry, bStreamBVH, bStreamBuffers);
+		const uint32 GeometryIndex = GetHandleIndex(Handle);
 
-		if (RegisteredGeometry.bEvicted || bStreamBVH || bStreamBuffers)
+		if (ReferenceBasedResidencyState->GeometryMetadatas.IsValidIndex(GeometryIndex))
 		{
-			ReferencedGeometriesEvictedOrRequiringStreaming.Add(GeometryHandle);
+			ReferenceBasedResidencyState->GeometryMetadatas.RemoveAt(GeometryIndex);
 		}
 
-		NotReferencedResidentGeometries.Remove(GeometryHandle);
-
-		ReferencedSize += RegisteredGeometry.Size;
+		ReferenceBasedResidencyState->EvictableGeometries.Remove(Handle);
+	}
+	ReleasedGeometries.Empty();
 
-		if (!RegisteredGeometry.bAlwaysResident)
-		{
-			RequestedSize += RegisteredGeometry.Size;
-		}
+	for (RayTracing::FGeometryGroupHandle Handle : ReleasedGroups)
+	{
+		const uint32 GroupIndex = GetHandleIndex(Handle);
 
-		if (RegisteredGeometry.bEvicted)
+		if (ReferenceBasedResidencyState->GroupMetadatas.IsValidIndex(GroupIndex))
 		{
-			RequestedButEvictedSize += RegisteredGeometry.Size;
+			ReferenceBasedResidencyState->GroupMetadatas.RemoveAt(GroupIndex);
 		}
 	}
+	ReleasedGroups.Empty();
 
-	// Step 2
-	// - add all geometries in referenced groups to ReferencedGeometriesEvictedOrRequiringStreaming
-	//		- need to make all geometries in group resident otherwise might not have valid geometry when reducing LOD
-	//		- TODO: Could track TargetLOD and only make [TargetLOD ... LastLOD] range resident
-	// - also update LastReferencedFrame and calculate memory required to make evicted geometries resident
-	for (RayTracing::FGeometryGroupHandle GroupHandle : ReferencedGeometryGroups)
+	// Add/update metadata for geometries/groups added/modified during the previous frame.
+	for (RayTracing::FGeometryHandle Handle : ModifiedGeometries)
 	{
-		const FRayTracingGeometryGroup& Group = GetRegisteredGroup(GroupHandle);
-		const int32 NumGeometriesInGroup = Group.GeometryHandles.Num();
+		const FRegisteredGeometry& RegisteredGeometry = GetRegisteredGeometry(Handle);
 
-		for (int32 LODIndex = Group.CurrentFirstLODIdx; LODIndex < NumGeometriesInGroup; ++LODIndex)
+		const uint32 GeometryIndex = GetHandleIndex(Handle);
+
+		if (!ReferenceBasedResidencyState->GeometryMetadatas.IsValidIndex(GeometryIndex))
 		{
-			RayTracing::FGeometryHandle GeometryHandle = Group.GeometryHandles[LODIndex];
+			ReferenceBasedResidencyState->GeometryMetadatas.Insert(GeometryIndex, {});
+		}
 
-			if (!GeometryHandle.IsNull()) // some LODs might be stripped during cook
-			{
-				FRegisteredGeometry& RegisteredGeometry = GetRegisteredGeometry(GeometryHandle);
-				RegisteredGeometry.LastReferencedFrame = GFrameCounterRenderThread;
+		FReferenceBasedResidencyState::FGeometryMetadata& Metadata = ReferenceBasedResidencyState->GeometryMetadatas[GeometryIndex];
+		Metadata.Size = RegisteredGeometry.Size;
+		Metadata.GroupHandle = RegisteredGeometry.Geometry->GroupHandle;
+		Metadata.bEvicted = RegisteredGeometry.bEvicted;
+		Metadata.bBVHStreamedOut = RegisteredGeometry.Geometry->GetInitializer().Type == ERayTracingGeometryInitializerType::StreamingDestination;
+		Metadata.bBuffersStreamedOut = RegisteredGeometry.StreamableBuffersSize > 0 && !RegisteredGeometry.StreamableAsset->AreBuffersStreamedIn();
+		Metadata.bAlwaysResident = RegisteredGeometry.bAlwaysResident;
 
-				ReferencedSize += RegisteredGeometry.Size;
+		if (RegisteredGeometry.Geometry->IsValid() && !RegisteredGeometry.bEvicted && !RegisteredGeometry.bAlwaysResident)
+		{
+			ReferenceBasedResidencyState->EvictableGeometries.Add(Handle);
+		}
+		else
+		{
+			ReferenceBasedResidencyState->EvictableGeometries.Remove(Handle);
+		}
+	}
+	ModifiedGeometries.Empty();
 
-				if (RegisteredGeometry.bAlwaysResident)
-				{
-					checkf(!RegisteredGeometry.bEvicted, TEXT("Always resident ray tracing geometry was unexpectely evicted."));
-				}
-				else
-				{
-					bool bStreamBVH;
-					bool bStreamBuffers;
-					DoesRayTracingGeometryRequireStreaming(RegisteredGeometry, bStreamBVH, bStreamBuffers);
+	for (RayTracing::FGeometryGroupHandle Handle : ModifiedGroups)
+	{
+		const FRayTracingGeometryGroup& RegisteredGroup = GetRegisteredGroup(Handle);
 
-					if (RegisteredGeometry.bEvicted || bStreamBVH || bStreamBuffers)
-					{
-						ReferencedGeometriesEvictedOrRequiringStreaming.Add(GeometryHandle);
-					}
+		const uint32 GroupIndex = GetHandleIndex(Handle);
 
-					NotReferencedResidentGeometries.Remove(GeometryHandle);
+		if (!ReferenceBasedResidencyState->GroupMetadatas.IsValidIndex(GroupIndex))
+		{
+			ReferenceBasedResidencyState->GroupMetadatas.Insert(GroupIndex, {});
+		}
 
-					RequestedSize += RegisteredGeometry.Size;
+		FReferenceBasedResidencyState::FGeometryGroupMetadata& Metadata = ReferenceBasedResidencyState->GroupMetadatas[GroupIndex];
+		Metadata.GeometryHandles = RegisteredGroup.GeometryHandles;
+		Metadata.CurrentFirstLODIdx = RegisteredGroup.CurrentFirstLODIdx;
 
-					if (RegisteredGeometry.bEvicted)
-					{
-						RequestedButEvictedSize += RegisteredGeometry.Size;
-					}
-				}
-			}
-		}
 	}
+	ModifiedGroups.Empty();
+}
 
-#if DO_CHECK
-	// ensure(ReferencedGeometriesEvictedOrRequiringStreaming.Num() == TSet(ReferencedGeometriesEvictedOrRequiringStreaming).Num());
-#endif
+void FRayTracingGeometryManager::ProcessReferenceBasedResidency(FRHICommandList& RHICmdList)
+{
+	TRACE_CPUPROFILER_EVENT_SCOPE(FRayTracingGeometryManager::ProcessReferenceBasedResidency);
+
+	const int64 ReferencedSize = ReferenceBasedResidencyState->ReferencedSize;
+	const int64 RequestedSize = ReferenceBasedResidencyState->RequestedSize;
+	const int64 RequestedButEvictedSize = ReferenceBasedResidencyState->RequestedButEvictedSize;
+	TArray<RayTracing::FGeometryHandle>& GeometriesToEvict = ReferenceBasedResidencyState->GeometriesToEvict;
+	TArray<RayTracing::FGeometryHandle>& GeometriesToMakeResidentAndStreamIn = ReferenceBasedResidencyState->GeometriesToMakeResidentAndStreamIn;
 
 	const int64 ResidentGeometryMemoryPoolSize = FUnitConversion::Convert(GRayTracingResidentGeometryMemoryPoolSizeInMB, EUnit::Megabytes, EUnit::Bytes);
 
-	// Step 3
-	// - if making requested geometries resident will put us over budget -> evict some geometry not referenced by TLAS
-	if (TotalResidentSize + TotalStreamingSize + RequestedButEvictedSize > ResidentGeometryMemoryPoolSize)
+	// Evict geometries until we are in budget
 	{
-		TArray<RayTracing::FGeometryHandle> NotReferencedResidentGeometriesArray = NotReferencedResidentGeometries.Array();
-
-		// Step 3.1
-		// - sort to evict geometries in the following order:
-		//		- least recently used
-		//		- largest geometries
-		Algo::Sort(NotReferencedResidentGeometriesArray, [this](RayTracing::FGeometryHandle& LHSHandle, RayTracing::FGeometryHandle& RHSHandle)
-			{
-				FRegisteredGeometry& LHS = GetRegisteredGeometry(LHSHandle);
-				FRegisteredGeometry& RHS = GetRegisteredGeometry(RHSHandle);
-
-				// TODO: evict unreferenced dynamic geometries using shared buffers first since they need to be rebuild anyway
-				// (and then dynamic geometries requiring update?
-
-				// 1st - last referenced frame
-				if (LHS.LastReferencedFrame != RHS.LastReferencedFrame)
-				{
-					return LHS.LastReferencedFrame < RHS.LastReferencedFrame;
-				}
-
-				// 2nd - size
-				return LHS.Size > RHS.Size;
-			});
-
-		// Step 3.2
-		// - evict geometries until we are in budget
 		int32 Index = 0;
-		while (TotalResidentSize + TotalStreamingSize + RequestedButEvictedSize > ResidentGeometryMemoryPoolSize && Index < NotReferencedResidentGeometriesArray.Num())
+		while (TotalResidentSize + TotalStreamingSize + RequestedButEvictedSize > ResidentGeometryMemoryPoolSize && Index < GeometriesToEvict.Num())
 		{
-			RayTracing::FGeometryHandle GeometryHandle = NotReferencedResidentGeometriesArray[Index];
-			FRegisteredGeometry& RegisteredGeometry = GetRegisteredGeometry(GeometryHandle);
+			RayTracing::FGeometryHandle GeometryHandle = GeometriesToEvict[Index];
 
-			check(RegisteredGeometry.Geometry->IsValid() && !RegisteredGeometry.Geometry->IsEvicted());
+			if (IsHandleValid(GeometryHandle))
+			{
+				FRegisteredGeometry& RegisteredGeometry = GetRegisteredGeometry(GeometryHandle);
 
-			EvictGeometry(RHICmdList, RegisteredGeometry);
+				check(RegisteredGeometry.Geometry->IsValid() && !RegisteredGeometry.Geometry->IsEvicted());
+
+				EvictGeometry(RHICmdList, RegisteredGeometry);
+			}
 
 			++Index;
 		}
 	}
 
-	// Step 4
-	// - make referenced geometries resident until we go over budget
-	if (TotalResidentSize + TotalStreamingSize < ResidentGeometryMemoryPoolSize)
+	// Make referenced geometries resident until we go over budget
 	{
-		// Step 4.1
-		//	- sort by size to prioritize smaller geometries
-		Algo::Sort(ReferencedGeometriesEvictedOrRequiringStreaming, [this](RayTracing::FGeometryHandle& LHSHandle, RayTracing::FGeometryHandle& RHSHandle)
-			{
-				FRegisteredGeometry& LHS = GetRegisteredGeometry(LHSHandle);
-				FRegisteredGeometry& RHS = GetRegisteredGeometry(RHSHandle);
-
-				return LHS.Size < RHS.Size;
-			});
-
-		// Step 4.2
-		// - make geometries resident until we go over budget
 		int32 Index = 0;
-		while (TotalResidentSize + TotalStreamingSize < ResidentGeometryMemoryPoolSize && Index < ReferencedGeometriesEvictedOrRequiringStreaming.Num())
+		while (TotalResidentSize + TotalStreamingSize < ResidentGeometryMemoryPoolSize && Index < GeometriesToMakeResidentAndStreamIn.Num())
 		{
-			RayTracing::FGeometryHandle GeometryHandle = ReferencedGeometriesEvictedOrRequiringStreaming[Index];
-			FRegisteredGeometry& RegisteredGeometry = GetRegisteredGeometry(GeometryHandle);
+			RayTracing::FGeometryHandle GeometryHandle = GeometriesToMakeResidentAndStreamIn[Index];
 
-			if (RegisteredGeometry.Geometry->IsEvicted())
+			if (IsHandleValid(GeometryHandle))
 			{
-				MakeGeometryResident(RHICmdList, RegisteredGeometry);
-			}
+				FRegisteredGeometry& RegisteredGeometry = GetRegisteredGeometry(GeometryHandle);
+
+				if (RegisteredGeometry.Geometry->IsEvicted())
+				{
+					MakeGeometryResident(RHICmdList, RegisteredGeometry);
+				}
 
-			RequestRayTracingGeometryStreamIn(RHICmdList, GeometryHandle);
+				RequestRayTracingGeometryStreamIn(RHICmdList, GeometryHandle);
+			}
 
 			++Index;
 		}
 	}
 
+	GeometriesToEvict.Empty();
+	GeometriesToMakeResidentAndStreamIn.Empty();
+
 	SET_MEMORY_STAT(STAT_RayTracingGeometryReferencedMemory, ReferencedSize);
 	SET_MEMORY_STAT(STAT_RayTracingGeometryRequestedMemory, RequestedSize);
 	CSV_CUSTOM_STAT(RayTracingGeometry, ReferencedSizeMB, ReferencedSize / 1024.0f / 1024.0f, ECsvCustomStatOp::Set);
@@ -1005,13 +1225,6 @@ void FRayTracingGeometryManager::UpdateReferenceBasedResidency(FRHICommandList&
 					*FText::AsMemory(ResidentGeometryMemoryPoolSize).ToString()));
 		}
 	}
-
-	if (bDumpUnreferencedAlwaysResidentGeometriesNextFrame)
-	{
-		bDumpUnreferencedAlwaysResidentGeometriesNextFrame = false;
-
-		DumpUnreferencedAlwaysResidentRayTracingGeometries();
-	}
 #endif
 }
 
@@ -1061,7 +1274,7 @@ void FRayTracingGeometryManager::RefreshAlwaysResidentRayTracingGeometries(FRHIC
 
 			if (GroupHandle.IsNull())
 			{
-				checkf(!RegisteredGeometry.bAlwaysResident, TEXT("Ray tracing geometry unexpectely marked as always resident even though it is not assigned to a group."));
+				checkf(!RegisteredGeometry.bAlwaysResident, TEXT("Ray tracing geometry unexpectedly marked as always-resident even though it is not assigned to a group."));
 				continue;
 			}
 
@@ -1085,13 +1298,33 @@ void FRayTracingGeometryManager::RefreshAlwaysResidentRayTracingGeometries(FRHIC
 
 					INC_DWORD_STAT(STAT_RayTracingPendingStreamingRequests);
 				}
-
-				EvictableGeometries.Remove(RegisteredGeometry.Geometry->RayTracingGeometryHandle);
 			}
 			else if (bUsingReferenceBasedResidency && RegisteredGeometry.Geometry->GetRHI() != nullptr)
 			{
 				EvictGeometry(RHICmdList, RegisteredGeometry);
 			}
+
+			ModifiedGeometries.Add(RegisteredGeometry.Geometry->RayTracingGeometryHandle);
+		}
+	}
+	else if (bUsingReferenceBasedResidency && bHasRayTracingEnableChanged)
+	{
+		// make always-resident geometries actually resident
+
+		for (RayTracing::FGeometryHandle GeometryHandle : AlwaysResidentGeometries)
+		{
+			FRegisteredGeometry& RegisteredGeometry = GetRegisteredGeometry(GeometryHandle);
+
+			if (RegisteredGeometry.Geometry->IsEvicted())
+			{
+				MakeGeometryResident(RHICmdList, RegisteredGeometry);
+			}
+			if (!RequestRayTracingGeometryStreamIn(RHICmdList, GeometryHandle))
+			{
+				PendingStreamingRequests.Add(GeometryHandle);
+
+				INC_DWORD_STAT(STAT_RayTracingPendingStreamingRequests);
+			}
 		}
 	}
 }
@@ -1105,7 +1338,6 @@ void FRayTracingGeometryManager::CheckIntegrity() const
 	if (GRayTracingTestCheckIntegrity)
 	{
 		uint32 NumAlwaysResidentGeometries = 0;
-		uint32 NumEvictableGeometries = 0;
 
 		for (const FRegisteredGeometry& RegisteredGeometry : RegisteredGeometries)
 		{
@@ -1132,20 +1364,53 @@ void FRayTracingGeometryManager::CheckIntegrity() const
 			{
 				++NumAlwaysResidentGeometries;
 			}
+		}
+
+		check(NumAlwaysResidentGeometries == AlwaysResidentGeometries.Num());
+
+		uint32 NumEvictableGeometries = 0;
+
+		for (auto It = RegisteredGeometries.CreateConstIterator(); It; ++It)
+		{
+			const FRegisteredGeometry& RegisteredGeometry = *It;
+			const uint32 GeometryIndex = It.GetIndex();
+
+			const RayTracing::FGeometryHandle GeometryHandle = RegisteredGeometry.Geometry->RayTracingGeometryHandle;
+			const RayTracing::FGeometryGroupHandle GroupHandle = RegisteredGeometry.Geometry->GroupHandle;
+
+			const bool bBVHStreamedOut = RegisteredGeometry.Geometry->GetInitializer().Type == ERayTracingGeometryInitializerType::StreamingDestination;
+			const bool bBuffersStreamedOut = RegisteredGeometry.StreamableBuffersSize > 0 && !RegisteredGeometry.StreamableAsset->AreBuffersStreamedIn();
+
+			const FReferenceBasedResidencyState::FGeometryMetadata& Metadata = ReferenceBasedResidencyState->GeometryMetadatas[GeometryIndex];
+			check(Metadata.Size == RegisteredGeometry.Size);
+			check(Metadata.GroupHandle == RegisteredGeometry.Geometry->GroupHandle);
+			check(Metadata.bEvicted == RegisteredGeometry.bEvicted);
+			check(Metadata.bBVHStreamedOut == bBVHStreamedOut);
+			check(Metadata.bBuffersStreamedOut == bBuffersStreamedOut);
+			check(Metadata.bAlwaysResident == RegisteredGeometry.bAlwaysResident);
 
 			if (RegisteredGeometry.Geometry->IsValid() && !RegisteredGeometry.bEvicted && !RegisteredGeometry.bAlwaysResident)
 			{
-				check(EvictableGeometries.Contains(GeometryHandle));
+				check(ReferenceBasedResidencyState->EvictableGeometries.Contains(GeometryHandle));
 				++NumEvictableGeometries;
 			}
 			else
 			{
-				check(!EvictableGeometries.Contains(GeometryHandle));
+				check(!ReferenceBasedResidencyState->EvictableGeometries.Contains(GeometryHandle));
 			}
 		}
 
-		check(NumAlwaysResidentGeometries == AlwaysResidentGeometries.Num());
-		check(NumEvictableGeometries == EvictableGeometries.Num());
+		check(NumEvictableGeometries == ReferenceBasedResidencyState->EvictableGeometries.Num());
+
+		for (auto It = RegisteredGroups.CreateConstIterator(); It; ++It)
+		{
+			const FRayTracingGeometryGroup& RegisteredGroup = *It;
+			const uint32 GroupIndex = It.GetIndex();
+
+			FReferenceBasedResidencyState::FGeometryGroupMetadata& Metadata = ReferenceBasedResidencyState->GroupMetadatas[GroupIndex];
+			check(Metadata.GeometryHandles == RegisteredGroup.GeometryHandles);
+			check(Metadata.CurrentFirstLODIdx == RegisteredGroup.CurrentFirstLODIdx);
+		}
 
 		if (!bIsRayTracingEnabled)
 		{
@@ -1284,6 +1549,8 @@ bool FRayTracingGeometryManager::RequestRayTracingGeometryStreamIn(FRHICommandLi
 
 void FRayTracingGeometryManager::ProcessCompletedStreamingRequests(FRHICommandList& RHICmdList)
 {
+	TRACE_CPUPROFILER_EVENT_SCOPE(FRayTracingGeometryManager::ProcessCompletedStreamingRequests);
+
 	const bool bOnDemandGeometryBuffersStreaming = IsRayTracingUsingReferenceBasedResidency() && CVarRayTracingOnDemandGeometryBuffersStreaming.GetValueOnRenderThread();
 
 	const int32 StartPendingRequestIndex = (NextStreamingRequestIndex + GRayTracingStreamingMaxPendingRequests - NumStreamingRequests) % GRayTracingStreamingMaxPendingRequests;
diff --git a/Engine/Source/Runtime/Engine/Public/Rendering/RayTracingGeometryManager.h b/Engine/Source/Runtime/Engine/Public/Rendering/RayTracingGeometryManager.h
--- a/Engine/Source/Runtime/Engine/Public/Rendering/RayTracingGeometryManager.h
+++ b/Engine/Source/Runtime/Engine/Public/Rendering/RayTracingGeometryManager.h
@@ -7,6 +7,7 @@
 #include "Containers/Map.h"
 #include "IO/IoBuffer.h"
 #include "Serialization/BulkData.h"
+#include "Tasks/Task.h"
 
 #if RHI_RAYTRACING
 
@@ -40,6 +41,7 @@ public:
 	ENGINE_API void SetRayTracingGeometryGroupCurrentFirstLODIndex(FRHICommandListBase& RHICmdList, RayTracing::FGeometryGroupHandle Handle, uint8 CurrentFirstLODIdx);
 
 	ENGINE_API virtual void PreRender() override;
+	ENGINE_API virtual void Update(FRHICommandList& RHICmdList) override;
 	ENGINE_API virtual void Tick(FRHICommandList& RHICmdList) override;
 
 	ENGINE_API void RegisterProxyWithCachedRayTracingState(FPrimitiveSceneProxy* Proxy, RayTracing::FGeometryGroupHandle InRayTracingGeometryGroupHandle);
@@ -80,7 +82,8 @@ private:
 	bool RequestRayTracingGeometryStreamIn(FRHICommandList& RHICmdList, RayTracing::FGeometryHandle GeometryHandle);
 	void ProcessCompletedStreamingRequests(FRHICommandList& RHICmdList);
 
-	void UpdateReferenceBasedResidency(FRHICommandList& RHICmdList);
+	void UpdateReferenceBasedResidencyState();
+	void ProcessReferenceBasedResidency(FRHICommandList& RHICmdList);
 
 	FCriticalSection RequestCS;
 
@@ -112,7 +115,6 @@ private:
 	struct FRegisteredGeometry
 	{
 		FRayTracingGeometry* Geometry = nullptr;
-		uint64 LastReferencedFrame = 0;
 		uint32 Size = 0;
 
 		FRayTracingStreamableAsset* StreamableAsset = nullptr;
@@ -193,6 +195,9 @@ private:
 
 	FBuildRequest& GetBuildRequest(RayTracing::FBuildRequestHandle Handle);
 
+	bool IsHandleValid(RayTracing::FGeometryHandle Handle) const;
+	bool IsHandleValid(RayTracing::FGeometryGroupHandle Handle) const;
+
 	uint32 GetRegisteredGeometryIndex(RayTracing::FGeometryHandle Handle) const;
 	uint32 GetRegisteredGroupIndex(RayTracing::FGeometryGroupHandle Handle) const;
 
@@ -216,8 +221,6 @@ private:
 	TSet<RayTracing::FGeometryHandle> AlwaysResidentGeometries;
 	int64 TotalAlwaysResidentSize = 0;
 
-	TSet<RayTracing::FGeometryHandle> EvictableGeometries;
-
 	TSet<RayTracing::FGeometryHandle> ReferencedGeometries;
 	TSet<RayTracing::FGeometryGroupHandle> ReferencedGeometryGroups;
 
@@ -237,6 +240,18 @@ private:
 	TSet<RayTracing::FGeometryHandle> VisibleGeometryHandles;
 
 	bool bRenderedFrame = false;
+
+	TSet<RayTracing::FGeometryHandle> ReleasedGeometries;
+	TSet<RayTracing::FGeometryGroupHandle> ReleasedGroups;
+
+	TSet<RayTracing::FGeometryHandle> ModifiedGeometries;
+	TSet<RayTracing::FGeometryGroupHandle> ModifiedGroups;
+
+	TPimplPtr<struct FReferenceBasedResidencyState> ReferenceBasedResidencyState;
+
+	UE::Tasks::FTask ReferenceBasedResidencyTask = {};
+
+	friend struct FReferenceBasedResidencyState;
 };
 
 #endif // RHI_RAYTRACING
diff --git a/Engine/Source/Runtime/RenderCore/Public/RayTracingGeometryManagerInterface.h b/Engine/Source/Runtime/RenderCore/Public/RayTracingGeometryManagerInterface.h
--- a/Engine/Source/Runtime/RenderCore/Public/RayTracingGeometryManagerInterface.h
+++ b/Engine/Source/Runtime/RenderCore/Public/RayTracingGeometryManagerInterface.h
@@ -49,6 +49,7 @@ public:
 	virtual void RefreshRegisteredGeometry(RayTracing::FGeometryHandle Handle) = 0;
 
 	virtual void PreRender() = 0;
+	virtual void Update(FRHICommandList& RHICmdList) = 0;
 	virtual void Tick(FRHICommandList& RHICmdList) = 0;
 
 	virtual bool IsGeometryVisible(RayTracing::FGeometryHandle GeometryHandle) const = 0;
diff --git a/Engine/Source/Runtime/Renderer/Private/DeferredShadingRenderer.cpp b/Engine/Source/Runtime/Renderer/Private/DeferredShadingRenderer.cpp
--- a/Engine/Source/Runtime/Renderer/Private/DeferredShadingRenderer.cpp
+++ b/Engine/Source/Runtime/Renderer/Private/DeferredShadingRenderer.cpp
@@ -1770,6 +1770,8 @@ void FDeferredShadingSceneRenderer::Render(FRDGBuilder& GraphBuilder, const FSce
 	{
 		RHI_BREADCRUMB_EVENT_STAT(GraphBuilder.RHICmdList, RayTracingGeometry, "RayTracingGeometry");
 
+		GRayTracingGeometryManager->Update(GraphBuilder.RHICmdList);
+
 		// TODO: should only process build requests once per frame
 		GRayTracingGeometryManager->ProcessBuildRequests(GraphBuilder.RHICmdList);
 	}
