[World] Old DataLayer deprecation continuation part 5 - Deprecate UDataLayerSubsystem, utils methods. Final data layer deprecation CL.

#Synced-CL 51591079

diff --git a/Engine/Source/Runtime/Engine/Private/WorldPartition/DataLayer/DataLayerManager.cpp b/Engine/Source/Runtime/Engine/Private/WorldPartition/DataLayer/DataLayerManager.cpp
--- a/Engine/Source/Runtime/Engine/Private/WorldPartition/DataLayer/DataLayerManager.cpp
+++ b/Engine/Source/Runtime/Engine/Private/WorldPartition/DataLayer/DataLayerManager.cpp
@@ -3,7 +3,6 @@
 #include "WorldPartition/DataLayer/DataLayerManager.h"
 #include "WorldPartition/DataLayer/WorldDataLayers.h"
 #include "WorldPartition/DataLayer/DataLayerInstanceWithAsset.h"
-#include "WorldPartition/DataLayer/DataLayerSubsystem.h"
 #include "WorldPartition/DataLayer/DataLayerLoadingPolicy.h"
 #include "WorldPartition/DataLayer/ExternalDataLayerAsset.h"
 #include "WorldPartition/DataLayer/ExternalDataLayerInstance.h"
@@ -324,14 +323,6 @@ bool UDataLayerManager::SetDataLayerRuntimeState(const UDataLayerAsset* InDataLa
 
 void UDataLayerManager::BroadcastOnDataLayerInstanceRuntimeStateChanged(const UDataLayerInstance* InDataLayer, EDataLayerRuntimeState InState)
 {
-	// For backward compatibility
-	PRAGMA_DISABLE_DEPRECATION_WARNINGS
-	if (UDataLayerSubsystem* DataLayerSubsystem = UWorld::GetSubsystem<UDataLayerSubsystem>(GetWorld()))
-	{
-		DataLayerSubsystem->OnDataLayerRuntimeStateChanged.Broadcast(InDataLayer, InState);
-	}
-	PRAGMA_ENABLE_DEPRECATION_WARNINGS
-
 	OnDataLayerInstanceRuntimeStateChanged.Broadcast(InDataLayer, InState);
 }
 
@@ -692,14 +683,6 @@ TSubclassOf<UDataLayerLoadingPolicy> UDataLayerManager::GetDataLayerLoadingPolic
 	// Use UDataLayerManager::DataLayerLoadingPolicyClass
 	UClass* DataLayerLoadingPolicyClassValue = DataLayerLoadingPolicyClass.Get();
 	if (!DataLayerLoadingPolicyClassValue)
-	{
-		PRAGMA_DISABLE_DEPRECATION_WARNINGS
-		// Fallback on UDataLayerSubsystem::DataLayerLoadingPolicyClass
-		UDataLayerSubsystem* DataLayerSubsystem = UWorld::GetSubsystem<UDataLayerSubsystem>(GetWorld());
-		DataLayerLoadingPolicyClassValue = DataLayerSubsystem ? DataLayerSubsystem->DataLayerLoadingPolicyClass.Get() : nullptr;
-		PRAGMA_ENABLE_DEPRECATION_WARNINGS
-	}
-	if (!DataLayerLoadingPolicyClassValue)
 	{
 		// Default to UDataLayerLoadingPolicy class
 		DataLayerLoadingPolicyClassValue = UDataLayerLoadingPolicy::StaticClass();
diff --git a/Engine/Source/Runtime/Engine/Private/WorldPartition/DataLayer/DataLayerSubsystem.cpp b/Engine/Source/Runtime/Engine/Private/WorldPartition/DataLayer/DataLayerSubsystem.cpp
deleted file mode 100644
diff --git a/Engine/Source/Runtime/Engine/Public/WorldPartition/DataLayer/DataLayerManager.h b/Engine/Source/Runtime/Engine/Public/WorldPartition/DataLayer/DataLayerManager.h
--- a/Engine/Source/Runtime/Engine/Public/WorldPartition/DataLayer/DataLayerManager.h
+++ b/Engine/Source/Runtime/Engine/Public/WorldPartition/DataLayer/DataLayerManager.h
@@ -165,7 +165,6 @@ private:
 
 	friend class UWorldPartition;
 	friend class AWorldDataLayers;
-	friend class UDataLayerSubsystem;
 	friend class UDataLayerInstanceWithAsset;
 	friend class UExternalDataLayerManager;
 	friend class UExternalDataLayerInstance;
diff --git a/Engine/Source/Runtime/Engine/Public/WorldPartition/DataLayer/DataLayerSubsystem.h b/Engine/Source/Runtime/Engine/Public/WorldPartition/DataLayer/DataLayerSubsystem.h
--- a/Engine/Source/Runtime/Engine/Public/WorldPartition/DataLayer/DataLayerSubsystem.h
+++ b/Engine/Source/Runtime/Engine/Public/WorldPartition/DataLayer/DataLayerSubsystem.h
@@ -2,269 +2,14 @@
 
 #pragma once
 
+#if UE_ENABLE_INCLUDE_ORDER_DEPRECATED_IN_5_8
 #include "Subsystems/WorldSubsystem.h"
 #include "WorldPartition/DataLayer/DataLayerInstance.h"
 #include "WorldPartition/DataLayer/DataLayerManager.h"
-#include "DataLayerSubsystem.generated.h"
+#endif // UE_ENABLE_INCLUDE_ORDER_DEPRECATED_IN_5_8
 
-/**
- * UDataLayerSubsystem
- */
+#include "Misc/CoreMiscDefines.h"
 
-class UDataLayerLoadingPolicy;
-class UCanvas;
+UE_DEPRECATED_HEADER(5.8, "The contents of this header have been fully deprecated. See DeprecatedDataLayer.h for details on particular cases.") 
 
-DECLARE_DYNAMIC_MULTICAST_DELEGATE_TwoParams(FOnDataLayerRuntimeStateChanged, const UDataLayerInstance*, DataLayer, EDataLayerRuntimeState, State);
-
-#if WITH_EDITOR
-DECLARE_MULTICAST_DELEGATE_OneParam(FOnWorldDataLayersPostRegister, AWorldDataLayers*);
-DECLARE_MULTICAST_DELEGATE_OneParam(FOnWorldDataLayersPreUnregister, AWorldDataLayers*);
-#endif
-
-/** This class is deprecated, it has been replaced by DataLayerManager. */
-UCLASS(Config = Engine, MinimalAPI)
-class UDataLayerSubsystem : public UWorldSubsystem
-{
-	GENERATED_BODY()
-
-protected:
-	ENGINE_API virtual bool DoesSupportWorldType(const EWorldType::Type WorldType) const override;
-
-public:
-	//~ Begin Blueprint callable functions
-
-	/** Find a Data Layer by its asset. */
-	UE_DEPRECATED(5.3, "Use UDataLayerManager::GetDataLayerInstanceFromAsset instead.")
-	UFUNCTION(BlueprintCallable, Category = DataLayers, meta = (DeprecatedFunction, DeprecationMessage = "DataLayerSubsystem is deprecated, use the equivalent function in DataLayerManager"))
-	ENGINE_API UDataLayerInstance* GetDataLayerInstanceFromAsset(const UDataLayerAsset* InDataLayerAsset) const;
-
-	UE_DEPRECATED(5.3, "Use UDataLayerManager::GetDataLayerInstanceRuntimeState instead.")
-	UFUNCTION(BlueprintCallable, Category = DataLayers, meta = (DeprecatedFunction, DeprecationMessage = "DataLayerSubsystem is deprecated, use the equivalent function in DataLayerManager"))
-	ENGINE_API EDataLayerRuntimeState GetDataLayerInstanceRuntimeState(const UDataLayerAsset* InDataLayerAsset) const;
-
-	UE_DEPRECATED(5.3, "Use UDataLayerManager::GetDataLayerInstanceEffectiveRuntimeState instead.")
-	UFUNCTION(BlueprintCallable, Category = DataLayers, meta = (DeprecatedFunction, DeprecationMessage = "DataLayerSubsystem is deprecated, use the equivalent function in DataLayerManager"))
-	ENGINE_API EDataLayerRuntimeState GetDataLayerInstanceEffectiveRuntimeState(const UDataLayerAsset* InDataLayerAsset) const;
-
-	/** Set the Data Layer state using its name. */
-	UE_DEPRECATED(5.3, "Use UDataLayerManager::SetDataLayerInstanceRuntimeState instead.")
-	UFUNCTION(BlueprintCallable, Category = DataLayers, BlueprintAuthorityOnly, meta = (DeprecatedFunction, DeprecationMessage = "DataLayerSubsystem is deprecated, use the equivalent function in DataLayerManager"))
-	ENGINE_API void SetDataLayerInstanceRuntimeState(const UDataLayerAsset* InDataLayerAsset, EDataLayerRuntimeState InState, bool bInIsRecursive = false);
-
-	/** Called when a Data Layer changes state. */
-	UE_DEPRECATED(5.3, "Use UDataLayerManager::OnDataLayerInstanceRuntimeStateChanged instead.")
-	UPROPERTY(BlueprintAssignable)
-	FOnDataLayerRuntimeStateChanged OnDataLayerRuntimeStateChanged;
-
-	//~ End Blueprint callable functions
-
-	template<class T>
-	UE_DEPRECATED(5.3, "Use UDataLayerManager::GetDataLayerInstance instead.")
-	UDataLayerInstance* GetDataLayerInstance(const T& InDataLayerIdentifier, const ULevel* InLevelContext = nullptr) const;
-
-	template<class T>
-	UE_DEPRECATED(5.3, "Use UDataLayerManager::GetDataLayerInstances instead.")
-	TArray<const UDataLayerInstance*> GetDataLayerInstances(const TArray<T>& InDataLayerIdentifiers, const ULevel* InLevelContext = nullptr) const;
-
-	template<class T>
-	UE_DEPRECATED(5.3, "Use UDataLayerManager::GetDataLayerInstanceNames instead.")
-	TArray<FName> GetDataLayerInstanceNames(const TArray<T>& InDataLayerIdentifiers, const ULevel* InLevelContext = nullptr) const;
-
-	UE_DEPRECATED(5.3, "This function has been deprecated.")
-	ENGINE_API const UDataLayerInstance* GetDataLayerInstanceFromAssetName(const FName& InDataLayerAssetPathName) const;
-
-	UE_DEPRECATED(5.3, "Use UDataLayerManager::SetDataLayerRuntimeState instead.")
-	ENGINE_API void SetDataLayerRuntimeState(const UDataLayerInstance* InDataLayerInstance, EDataLayerRuntimeState InState, bool bInIsRecursive = false);
-
-	UE_DEPRECATED(5.3, "Use UDataLayerManager::GetDataLayerRuntimeState instead.")
-	ENGINE_API EDataLayerRuntimeState GetDataLayerRuntimeState(const UDataLayerInstance* InDataLayer) const;
-
-	UE_DEPRECATED(5.3, "Use UDataLayerManager::GetDataLayerRuntimeStateByName instead.")
-	ENGINE_API EDataLayerRuntimeState GetDataLayerRuntimeStateByName(const FName& InDataLayerName) const;
-
-	UE_DEPRECATED(5.3, "Use UDataLayerManager::GetDataLayerEffectiveRuntimeState instead.")
-	ENGINE_API EDataLayerRuntimeState GetDataLayerEffectiveRuntimeState(const UDataLayerInstance* InDataLayer) const;
-
-	UE_DEPRECATED(5.3, "Use UDataLayerManager::GetDataLayerEffectiveRuntimeStateByName instead.")
-	ENGINE_API EDataLayerRuntimeState GetDataLayerEffectiveRuntimeStateByName(const FName& InDataLayerName) const;
-
-	UE_DEPRECATED(5.3, "Use UDataLayerManager::IsAnyDataLayerInEffectiveRuntimeState instead.")
-	ENGINE_API bool IsAnyDataLayerInEffectiveRuntimeState(const TArray<FName>& InDataLayerNames, EDataLayerRuntimeState InState) const;
-
-	UE_DEPRECATED(5.3, "Use UDataLayerManager::GetEffectiveActiveDataLayerNames instead.")
-	ENGINE_API const TSet<FName>& GetEffectiveActiveDataLayerNames() const;
-
-	UE_DEPRECATED(5.3, "Use UDataLayerManager::GetEffectiveLoadedDataLayerNames instead.")
-	ENGINE_API const TSet<FName>& GetEffectiveLoadedDataLayerNames() const;
-
-	UE_DEPRECATED(5.3, "This function has been deprecated.")
-	void GetDataLayerDebugColors(TMap<FName, FColor>& OutMapping) const {}
-
-	UE_DEPRECATED(5.3, "This function has been deprecated.")
-	ENGINE_API void DrawDataLayersStatus(UCanvas* Canvas, FVector2D& Offset) const;
-
-	UE_DEPRECATED(5.3, "This function has been deprecated.")
-	ENGINE_API void DumpDataLayers(FOutputDevice& OutputDevice) const;
-
-	UE_DEPRECATED(5.3, "This function has been deprecated.")
-	void RegisterWorldDataLayer(AWorldDataLayers* WorldDataLayers) {}
-
-	UE_DEPRECATED(5.3, "This function has been deprecated.")
-	void UnregisterWorldDataLayer(AWorldDataLayers* WorldDataLayers) {}
-
-	UE_DEPRECATED(5.3, "This function has been deprecated.")
-	void ForEachWorldDataLayer(TFunctionRef<bool(AWorldDataLayers*)> Func) {}
-
-	UE_DEPRECATED(5.3, "This function has been deprecated.")
-	void ForEachWorldDataLayer(TFunctionRef<bool(AWorldDataLayers*)> Func) const {}
-
-	UE_DEPRECATED(5.3, "This function has been deprecated.")
-	void ForEachDataLayer(TFunctionRef<bool(UDataLayerInstance*)> Func, const ULevel* InLevelContext = nullptr) {}
-
-	UE_DEPRECATED(5.3, "This function has been deprecated.")
-	void ForEachDataLayer(TFunctionRef<bool(UDataLayerInstance*)> Func, const ULevel* InLevelContext = nullptr) const {}
-
-#if WITH_EDITOR
-	UE_DEPRECATED(5.3, "This function has been deprecated.")
-	static ENGINE_API TArray<const UDataLayerInstance*> GetRuntimeDataLayerInstances(UWorld* InWorld, const TArray<FName>& InDataLayerInstanceNames);
-
-	UE_DEPRECATED(5.3, "This function has been deprecated.")
-	bool ResolveIsLoadedInEditor(const TArray<FName>& InDataLayerInstanceNames) const { return false; }
-
-	UE_DEPRECATED(5.3, "This function has been deprecated.")
-	bool CanResolveDataLayers() const { return false; }
-
-	UE_DEPRECATED(5.3, "This function has been deprecated.")
-	bool RemoveDataLayer(const UDataLayerInstance* InDataLayer) { return false; }
-
-	UE_DEPRECATED(5.3, "This function has been deprecated.")
-	int32 RemoveDataLayers(const TArray<UDataLayerInstance*>& InDataLayerInstances) { return 0; }
-
-	UE_DEPRECATED(5.3, "This function has been deprecated.")
-	void UpdateDataLayerEditorPerProjectUserSettings() const {}
-
-	UE_DEPRECATED(5.3, "This function has been deprecated.")
-	void GetUserLoadedInEditorStates(TArray<FName>& OutDataLayersLoadedInEditor, TArray<FName>& OutDataLayersNotLoadedInEditor) const {}
-
-	UE_DEPRECATED(5.3, "This function has been deprecated.")
-	void PushActorEditorContext() const {}
-
-	UE_DEPRECATED(5.3, "This function has been deprecated.")
-	void PopActorEditorContext() const {}
-
-	UE_DEPRECATED(5.3, "This function has been deprecated.")
-	TArray<UDataLayerInstance*> GetActorEditorContextDataLayers() const { return TArray<UDataLayerInstance*>(); }
-
-	UE_DEPRECATED(5.3, "This function has been deprecated.")
-	uint32 GetDataLayerEditorContextHash() const { return 0; }
-#endif
-
-	PRAGMA_DISABLE_DEPRECATION_WARNINGS // Suppress compiler warning on override of deprecated function
-
-	UE_DEPRECATED(5.0, "Use SetDataLayerRuntimeState instead.")
-	UFUNCTION(BlueprintCallable, Category = DataLayers, BlueprintAuthorityOnly, meta = (DeprecatedFunction, DeprecationMessage = "DataLayerSubsystem is deprecated, use the equivalent function in DataLayerManager"))
-	void SetDataLayerState(const FActorDataLayer& InDataLayer, EDataLayerState InState) { SetDataLayerRuntimeState(InDataLayer, (EDataLayerRuntimeState)InState); }
-
-	UE_DEPRECATED(5.0, "Use SetDataLayerRuntimeStateByLabel instead.")
-	UFUNCTION(BlueprintCallable, Category = DataLayers, BlueprintAuthorityOnly, meta = (DeprecatedFunction, DeprecationMessage = "DataLayerSubsystem is deprecated, use the equivalent function in DataLayerManager"))
-	void SetDataLayerStateByLabel(const FName& InDataLayerLabel, EDataLayerState InState) { SetDataLayerRuntimeStateByLabel(InDataLayerLabel, (EDataLayerRuntimeState)InState); }
-
-	UE_DEPRECATED(5.0, "Use GetDataLayerRuntimeState instead.")
-	UFUNCTION(BlueprintCallable, Category = DataLayers, meta = (DeprecatedFunction, DeprecationMessage = "DataLayerSubsystem is deprecated, use the equivalent function in DataLayerManager"))
-	EDataLayerState GetDataLayerState(const FActorDataLayer& InDataLayer) const { return (EDataLayerState)GetDataLayerRuntimeState(InDataLayer); }
-
-	UE_DEPRECATED(5.0, "Use GetDataLayerRuntimeStateByLabel instead.")
-	UFUNCTION(BlueprintCallable, Category = DataLayers, meta = (DeprecatedFunction, DeprecationMessage = "DataLayerSubsystem is deprecated, use the equivalent function in DataLayerManager"))
-	EDataLayerState GetDataLayerStateByLabel(const FName& InDataLayerLabel) const { return (EDataLayerState)GetDataLayerRuntimeStateByLabel(InDataLayerLabel); }
-
-	UE_DEPRECATED(5.0, "Use GetDataLayerRuntimeState instead.")
-	EDataLayerState GetDataLayerState(const UDataLayerInstance* InDataLayer) const { return (EDataLayerState)GetDataLayerRuntimeState(InDataLayer); }
-
-	UE_DEPRECATED(5.0, "Use UDataLayerManager::GetDataLayerInstanceRuntimeState instead.")
-	EDataLayerState GetDataLayerStateByName(const FName& InDataLayerName) const { return (EDataLayerState)GetDataLayerRuntimeStateByName(InDataLayerName); }
-
-	UE_DEPRECATED(5.0, "Use UDataLayerManager::IsAnyDataLayerInEffectiveRuntimeState instead.")
-	bool IsAnyDataLayerInState(const TArray<FName>& InDataLayerNames, EDataLayerState InState) const { return IsAnyDataLayerInEffectiveRuntimeState(InDataLayerNames, (EDataLayerRuntimeState)InState); }
-
-	PRAGMA_ENABLE_DEPRECATION_WARNINGS
-
-	UE_DEPRECATED(5.0, "GetActiveDataLayerNames will be removed.")
-	UFUNCTION(BlueprintCallable, Category = DataLayers, meta = (DeprecatedFunction, DeprecationMessage = "GetActiveDataLayerNames will be removed."))
-	const TSet<FName>& GetActiveDataLayerNames() const { static TSet<FName> EmptySet; return EmptySet; }
-
-	UE_DEPRECATED(5.0, "GetLoadedDataLayerNames will be removed.")
-	UFUNCTION(BlueprintCallable, Category = DataLayers, meta = (DeprecatedFunction, DeprecationMessage = "GetLoadedDataLayerNames will be removed."))
-	const TSet<FName>& GetLoadedDataLayerNames() const { static TSet<FName> EmptySet; return EmptySet; }
-
-	PRAGMA_DISABLE_DEPRECATION_WARNINGS
-	UE_DEPRECATED(5.1, "Use GetDataLayerInstanceFromAsset instead.")
-	UFUNCTION(BlueprintCallable, Category = DataLayers, meta = (DeprecatedFunction, DeprecationMessage = "DataLayerSubsystem is deprecated, use GetDataLayerInstanceFromAsset in DataLayerManager"))
-	ENGINE_API UDataLayerInstance* GetDataLayer(const FActorDataLayer& InDataLayer) const;
-	PRAGMA_ENABLE_DEPRECATION_WARNINGS
-
-	UE_DEPRECATED(5.1, "Use GetDataLayerInstanceFromAsset instead.")
-	UFUNCTION(BlueprintCallable, Category = DataLayers, meta = (DeprecatedFunction, DeprecationMessage = "DataLayerSubsystem is deprecated, use GetDataLayerInstanceFromAsset in DataLayerManager"))
-	ENGINE_API UDataLayerInstance* GetDataLayerFromName(FName InDataLayerName) const;
-
-	UE_DEPRECATED(all, "Use GetDataLayerInstanceFromAsset instead.")
-	UFUNCTION(meta = (DeprecatedFunction, DeprecationMessage = "DataLayerSubsystem is deprecated, use GetDataLayerInstanceFromAsset in DataLayerManager"))
-	UDataLayerInstance* GetDataLayerFromLabel(FName InDataLayerLabel) const { return nullptr; }
-
-	PRAGMA_DISABLE_DEPRECATION_WARNINGS
-	UE_DEPRECATED(5.1, "Use GetDataLayerInstanceRuntimeState instead.")
-	UFUNCTION(BlueprintCallable, Category = DataLayers, meta = (DeprecatedFunction, DeprecationMessage = "DataLayerSubsystem is deprecated, use GetDataLayerInstanceRuntimeState in DataLayerManager"))
-	ENGINE_API EDataLayerRuntimeState GetDataLayerRuntimeState(const FActorDataLayer& InDataLayer) const;
-	PRAGMA_ENABLE_DEPRECATION_WARNINGS
-
-	UE_DEPRECATED(5.1, "Use GetDataLayerInstanceRuntimeState instead.")
-	UFUNCTION(BlueprintCallable, Category = DataLayers, meta = (DeprecatedFunction, DeprecationMessage = "DataLayerSubsystem is deprecated, use GetDataLayerInstanceRuntimeState in DataLayerManager"))
-	ENGINE_API EDataLayerRuntimeState GetDataLayerRuntimeStateByLabel(const FName& InDataLayerLabel) const;
-
-	PRAGMA_DISABLE_DEPRECATION_WARNINGS
-	UE_DEPRECATED(5.1, "Use GetDataLayerInstanceEffectiveRuntimeState instead.")
-	UFUNCTION(BlueprintCallable, Category = DataLayers, meta = (DeprecatedFunction, DeprecationMessage = "DataLayerSubsystem is deprecated, use GetDataLayerInstanceEffectiveRuntimeState in DataLayerManager"))
-	ENGINE_API EDataLayerRuntimeState GetDataLayerEffectiveRuntimeState(const FActorDataLayer& InDataLayer) const;
-	PRAGMA_ENABLE_DEPRECATION_WARNINGS
-
-	UE_DEPRECATED(5.1, "Use GetDataLayerInstanceEffectiveRuntimeState instead.")
-	UFUNCTION(BlueprintCallable, Category = DataLayers, meta = (DeprecatedFunction, DeprecationMessage = "DataLayerSubsystem is deprecated, use GetDataLayerInstanceEffectiveRuntimeState in DataLayerManager"))
-	ENGINE_API EDataLayerRuntimeState GetDataLayerEffectiveRuntimeStateByLabel(const FName& InDataLayerLabel) const;
-
-	PRAGMA_DISABLE_DEPRECATION_WARNINGS
-	UE_DEPRECATED(5.1, "Use SetDataLayerRuntimeState() with UDataLayerAsset* overload instead.")
-	UFUNCTION(BlueprintCallable, Category = DataLayers, BlueprintAuthorityOnly)
-	ENGINE_API void SetDataLayerRuntimeState(const FActorDataLayer& InDataLayer, EDataLayerRuntimeState InState, bool bInIsRecursive = false);
-	PRAGMA_ENABLE_DEPRECATION_WARNINGS
-
-	UE_DEPRECATED(5.1, "Use SetDataLayerInstanceRuntimeState instead.")
-	UFUNCTION(BlueprintCallable, Category = DataLayers, BlueprintAuthorityOnly)
-	ENGINE_API void SetDataLayerRuntimeStateByLabel(const FName& InDataLayerLabel, EDataLayerRuntimeState InState, bool bInIsRecursive = false);
-
-private:
-	UE_DEPRECATED(5.3, "Use UDataLayerManager::DataLayerLoadingPolicyClass instead.")
-	UPROPERTY(Config)
-	TSoftClassPtr<UDataLayerLoadingPolicy> DataLayerLoadingPolicyClass;
-
-	friend class UDataLayerManager;
-};
-
-template<class T>
-UDataLayerInstance* UDataLayerSubsystem::GetDataLayerInstance(const T& InDataLayerIdentifier, const ULevel* InLevelContext /* = nullptr */) const
-{
-	const UDataLayerManager* DataLayerManager = UDataLayerManager::GetDataLayerManager(GetWorld());
-	return DataLayerManager ? const_cast<UDataLayerInstance*>(DataLayerManager->GetDataLayerInstance(InDataLayerIdentifier)) : nullptr;
-}
-
-template<class T>
-TArray<FName> UDataLayerSubsystem::GetDataLayerInstanceNames(const TArray<T>& InDataLayerIdentifiers, const ULevel* InLevelContext /* = nullptr */) const
-{
-	const UDataLayerManager* DataLayerManager = UDataLayerManager::GetDataLayerManager(GetWorld());
-	return DataLayerManager ? DataLayerManager->GetDataLayerInstanceNames(InDataLayerIdentifiers) : TArray<FName>();
-}
-
-template<class T>
-TArray<const UDataLayerInstance*> UDataLayerSubsystem::GetDataLayerInstances(const TArray<T>& InDataLayerIdentifiers, const ULevel* InLevelContext /* = nullptr */) const
-{
-	const UDataLayerManager* DataLayerManager = UDataLayerManager::GetDataLayerManager(GetWorld());
-	return DataLayerManager ? DataLayerManager->GetDataLayerInstances(InDataLayerIdentifiers) : TArray<const UDataLayerInstance*>();
-}
+#include "WorldPartition/DataLayer/Deprecated/DeprecatedDataLayer.h"
\ No newline at end of file
diff --git a/Engine/Source/Runtime/Engine/Public/WorldPartition/DataLayer/DataLayerUtils.h b/Engine/Source/Runtime/Engine/Public/WorldPartition/DataLayer/DataLayerUtils.h
--- a/Engine/Source/Runtime/Engine/Public/WorldPartition/DataLayer/DataLayerUtils.h
+++ b/Engine/Source/Runtime/Engine/Public/WorldPartition/DataLayer/DataLayerUtils.h
@@ -27,22 +27,6 @@ public:
 		return IconNameByType[static_cast<uint32>(DataLayerType)];
 	}
 
-	UE_DEPRECATED(5.4, "Use ResolveDataLayerInstanceNames instead")
-	static TArray<FName> ResolvedDataLayerInstanceNames(const UDataLayerManager* InDataLayerManager, const FWorldPartitionActorDesc* InActorDesc, const TArray<const FWorldDataLayersActorDesc*>& InWorldDataLayersActorDescs = TArray<const FWorldDataLayersActorDesc*>())
-	{
-		return ResolveDataLayerInstanceNames(InDataLayerManager, InActorDesc, InWorldDataLayersActorDescs).ToArray();
-	}
-
-	PRAGMA_DISABLE_DEPRECATION_WARNINGS
-
-	UE_DEPRECATED(5.4, "Use IWorldPartitionActorDescInstanceView version instead")
-	static bool ResolveRuntimeDataLayerInstanceNames(const UDataLayerManager* InDataLayerManager, const class FWorldPartitionActorDescView& InActorDescView, const FStreamingGenerationActorDescViewMap& ActorDescViewMap, TArray<FName>& OutRuntimeDataLayerInstanceNames) { return false; }
-
-	UE_DEPRECATED(5.4, "This function is no longer used")
-	static TArray<const FWorldDataLayersActorDesc*> FindWorldDataLayerActorDescs(const FStreamingGenerationActorDescViewMap& ActorDescViewMap) { return TArray<const FWorldDataLayersActorDesc*>(); }
-
-	PRAGMA_ENABLE_DEPRECATION_WARNINGS
-
 	/**
 	 * Calls Callback for each data layer instance that an actor belongs to.
 	 *
diff --git a/Engine/Source/Runtime/Engine/Public/WorldPartition/DataLayer/Deprecated/DeprecatedDataLayer.h b/Engine/Source/Runtime/Engine/Public/WorldPartition/DataLayer/Deprecated/DeprecatedDataLayer.h
--- a/Engine/Source/Runtime/Engine/Public/WorldPartition/DataLayer/Deprecated/DeprecatedDataLayer.h
+++ b/Engine/Source/Runtime/Engine/Public/WorldPartition/DataLayer/Deprecated/DeprecatedDataLayer.h
@@ -3,10 +3,11 @@
 #pragma once
 
 #include "CoreMinimal.h"
+#include "Math/Color.h"
+#include "Subsystems/WorldSubsystem.h"
 #include "UObject/ObjectMacros.h"
 #include "UObject/UObjectGlobals.h"
 #include "UObject/Object.h"
-#include "Math/Color.h"
 #include "WorldPartition/DataLayer/DataLayerInstance.h"
 #include "WorldPartition/DataLayer/DataLayerUtils.h"
 
@@ -235,4 +236,111 @@ private:
 	friend class UDataLayerConversionInfo;
 	friend class UDataLayerToAssetCommandletContext;
 	friend class UDataLayerToAssetCommandlet;
+};
+
+class UE_DEPRECATED(5.8, "This delegate is deprecated, it has been replaced by equivalent code in DataLayerManager") FOnDataLayerRuntimeStateChanged;
+
+DECLARE_DYNAMIC_MULTICAST_DELEGATE_TwoParams(FOnDataLayerRuntimeStateChanged, const UDataLayerInstance*, DataLayer, EDataLayerRuntimeState, State);
+
+class UE_DEPRECATED(5.8, "This class is deprecated, it has been replaced by DataLayerManager.") UDataLayerSubsystem;
+
+UCLASS(Config = Engine, MinimalAPI)
+class UDataLayerSubsystem : public UWorldSubsystem
+{
+	GENERATED_BODY()
+
+public:
+	/** Find a Data Layer by its asset. */
+	UFUNCTION(BlueprintCallable, Category = DataLayers, meta = (DeprecatedFunction, DeprecationMessage = "DataLayerSubsystem is deprecated, use the equivalent function in DataLayerManager"))
+	UDataLayerInstance* GetDataLayerInstanceFromAsset(const UDataLayerAsset* InDataLayerAsset) const { return nullptr; }
+
+	UFUNCTION(BlueprintCallable, Category = DataLayers, meta = (DeprecatedFunction, DeprecationMessage = "DataLayerSubsystem is deprecated, use the equivalent function in DataLayerManager"))
+	EDataLayerRuntimeState GetDataLayerInstanceRuntimeState(const UDataLayerAsset* InDataLayerAsset) const { return EDataLayerRuntimeState::Unloaded; }
+
+	UFUNCTION(BlueprintCallable, Category = DataLayers, meta = (DeprecatedFunction, DeprecationMessage = "DataLayerSubsystem is deprecated, use the equivalent function in DataLayerManager"))
+	EDataLayerRuntimeState GetDataLayerInstanceEffectiveRuntimeState(const UDataLayerAsset* InDataLayerAsset) const { return EDataLayerRuntimeState::Unloaded; }
+
+	/** Set the Data Layer state using its name. */
+	UFUNCTION(BlueprintCallable, Category = DataLayers, BlueprintAuthorityOnly, meta = (DeprecatedFunction, DeprecationMessage = "DataLayerSubsystem is deprecated, use the equivalent function in DataLayerManager"))
+	void SetDataLayerInstanceRuntimeState(const UDataLayerAsset* InDataLayerAsset, EDataLayerRuntimeState InState, bool bInIsRecursive = false) {}
+
+	PRAGMA_DISABLE_DEPRECATION_WARNINGS // Suppress compiler warning on override of deprecated function / usage of deprecated types
+
+	/** No longer called. */
+	UPROPERTY(BlueprintAssignable)
+	FOnDataLayerRuntimeStateChanged OnDataLayerRuntimeStateChanged;
+
+	UE_DEPRECATED(5.0, "Use SetDataLayerRuntimeState instead.")
+	UFUNCTION(BlueprintCallable, Category = DataLayers, BlueprintAuthorityOnly, meta = (DeprecatedFunction, DeprecationMessage = "DataLayerSubsystem is deprecated, use the equivalent function in DataLayerManager"))
+	void SetDataLayerState(const FActorDataLayer& InDataLayer, EDataLayerState InState) {}
+
+	UE_DEPRECATED(5.0, "Use SetDataLayerRuntimeStateByLabel instead.")
+	UFUNCTION(BlueprintCallable, Category = DataLayers, BlueprintAuthorityOnly, meta = (DeprecatedFunction, DeprecationMessage = "DataLayerSubsystem is deprecated, use the equivalent function in DataLayerManager"))
+	void SetDataLayerStateByLabel(const FName& InDataLayerLabel, EDataLayerState InState) {}
+
+	UE_DEPRECATED(5.0, "Use GetDataLayerRuntimeState instead.")
+	UFUNCTION(BlueprintCallable, Category = DataLayers, meta = (DeprecatedFunction, DeprecationMessage = "DataLayerSubsystem is deprecated, use the equivalent function in DataLayerManager"))
+	EDataLayerState GetDataLayerState(const FActorDataLayer& InDataLayer) const { return EDataLayerState::Unloaded; }
+
+	UE_DEPRECATED(5.0, "Use GetDataLayerRuntimeStateByLabel instead.")
+	UFUNCTION(BlueprintCallable, Category = DataLayers, meta = (DeprecatedFunction, DeprecationMessage = "DataLayerSubsystem is deprecated, use the equivalent function in DataLayerManager"))
+	EDataLayerState GetDataLayerStateByLabel(const FName& InDataLayerLabel) const { return EDataLayerState::Unloaded; }
+
+	PRAGMA_ENABLE_DEPRECATION_WARNINGS
+
+	UE_DEPRECATED(5.0, "GetActiveDataLayerNames will be removed.")
+	UFUNCTION(BlueprintCallable, Category = DataLayers, meta = (DeprecatedFunction, DeprecationMessage = "GetActiveDataLayerNames will be removed."))
+	const TSet<FName>& GetActiveDataLayerNames() const { static TSet<FName> EmptySet; return EmptySet; }
+
+	UE_DEPRECATED(5.0, "GetLoadedDataLayerNames will be removed.")
+	UFUNCTION(BlueprintCallable, Category = DataLayers, meta = (DeprecatedFunction, DeprecationMessage = "GetLoadedDataLayerNames will be removed."))
+	const TSet<FName>& GetLoadedDataLayerNames() const { static TSet<FName> EmptySet; return EmptySet; }
+
+	PRAGMA_DISABLE_DEPRECATION_WARNINGS
+	UE_DEPRECATED(5.1, "Use GetDataLayerInstanceFromAsset instead.")
+	UFUNCTION(BlueprintCallable, Category = DataLayers, meta = (DeprecatedFunction, DeprecationMessage = "DataLayerSubsystem is deprecated, use GetDataLayerInstanceFromAsset in DataLayerManager"))
+	UDataLayerInstance* GetDataLayer(const FActorDataLayer& InDataLayer) const {	return nullptr; }
+	PRAGMA_ENABLE_DEPRECATION_WARNINGS
+
+	UE_DEPRECATED(5.1, "Use GetDataLayerInstanceFromAsset instead.")
+	UFUNCTION(BlueprintCallable, Category = DataLayers, meta = (DeprecatedFunction, DeprecationMessage = "DataLayerSubsystem is deprecated, use GetDataLayerInstanceFromAsset in DataLayerManager"))
+	UDataLayerInstance* GetDataLayerFromName(FName InDataLayerName) const { return nullptr; }
+
+	UE_DEPRECATED(all, "Use GetDataLayerInstanceFromAsset instead.")
+	UFUNCTION(meta = (DeprecatedFunction, DeprecationMessage = "DataLayerSubsystem is deprecated, use GetDataLayerInstanceFromAsset in DataLayerManager"))
+	UDataLayerInstance* GetDataLayerFromLabel(FName InDataLayerLabel) const { return nullptr; }
+
+	PRAGMA_DISABLE_DEPRECATION_WARNINGS
+	UE_DEPRECATED(5.1, "Use GetDataLayerInstanceRuntimeState instead.")
+	UFUNCTION(BlueprintCallable, Category = DataLayers, meta = (DeprecatedFunction, DeprecationMessage = "DataLayerSubsystem is deprecated, use GetDataLayerInstanceRuntimeState in DataLayerManager"))
+	EDataLayerRuntimeState GetDataLayerRuntimeState(const FActorDataLayer& InDataLayer) const { return EDataLayerRuntimeState::Unloaded; }
+	PRAGMA_ENABLE_DEPRECATION_WARNINGS
+
+	UE_DEPRECATED(5.1, "Use GetDataLayerInstanceRuntimeState instead.")
+	UFUNCTION(BlueprintCallable, Category = DataLayers, meta = (DeprecatedFunction, DeprecationMessage = "DataLayerSubsystem is deprecated, use GetDataLayerInstanceRuntimeState in DataLayerManager"))
+	EDataLayerRuntimeState GetDataLayerRuntimeStateByLabel(const FName& InDataLayerLabel) const { return EDataLayerRuntimeState::Unloaded; }
+
+	PRAGMA_DISABLE_DEPRECATION_WARNINGS
+	UE_DEPRECATED(5.1, "Use GetDataLayerInstanceEffectiveRuntimeState instead.")
+	UFUNCTION(BlueprintCallable, Category = DataLayers, meta = (DeprecatedFunction, DeprecationMessage = "DataLayerSubsystem is deprecated, use GetDataLayerInstanceEffectiveRuntimeState in DataLayerManager"))
+	EDataLayerRuntimeState GetDataLayerEffectiveRuntimeState(const FActorDataLayer& InDataLayer) const { return EDataLayerRuntimeState::Unloaded; }
+	PRAGMA_ENABLE_DEPRECATION_WARNINGS
+
+	UE_DEPRECATED(5.1, "Use GetDataLayerInstanceEffectiveRuntimeState instead.")
+	UFUNCTION(BlueprintCallable, Category = DataLayers, meta = (DeprecatedFunction, DeprecationMessage = "DataLayerSubsystem is deprecated, use GetDataLayerInstanceEffectiveRuntimeState in DataLayerManager"))
+	EDataLayerRuntimeState GetDataLayerEffectiveRuntimeStateByLabel(const FName& InDataLayerLabel) const { return EDataLayerRuntimeState::Unloaded; }
+
+	PRAGMA_DISABLE_DEPRECATION_WARNINGS
+	UE_DEPRECATED(5.1, "Use SetDataLayerRuntimeState() with UDataLayerAsset* overload instead.")
+	UFUNCTION(BlueprintCallable, Category = DataLayers, BlueprintAuthorityOnly, meta = (DeprecatedFunction, DeprecationMessage = "DataLayerSubsystem is deprecated, use SetDataLayerRuntimeState in DataLayerManager"))
+	void SetDataLayerRuntimeState(const FActorDataLayer& InDataLayer, EDataLayerRuntimeState InState, bool bInIsRecursive = false) {}
+	PRAGMA_ENABLE_DEPRECATION_WARNINGS
+
+	UE_DEPRECATED(5.1, "Use SetDataLayerInstanceRuntimeState instead.")
+	UFUNCTION(BlueprintCallable, Category = DataLayers, BlueprintAuthorityOnly, meta = (DeprecatedFunction, DeprecationMessage = "DataLayerSubsystem is deprecated, use SetDataLayerInstanceRuntimeState in DataLayerManager"))
+	void SetDataLayerRuntimeStateByLabel(const FName& InDataLayerLabel, EDataLayerRuntimeState InState, bool bInIsRecursive = false) {}
+
+private:
+	UPROPERTY(Config)
+	TSoftClassPtr<class UDataLayerLoadingPolicy> DataLayerLoadingPolicyClass;
 };
\ No newline at end of file
diff --git a/Engine/Source/Runtime/Engine/Public/WorldPartition/DataLayer/DeprecatedDataLayerInstance.h b/Engine/Source/Runtime/Engine/Public/WorldPartition/DataLayer/DeprecatedDataLayerInstance.h
--- a/Engine/Source/Runtime/Engine/Public/WorldPartition/DataLayer/DeprecatedDataLayerInstance.h
+++ b/Engine/Source/Runtime/Engine/Public/WorldPartition/DataLayer/DeprecatedDataLayerInstance.h
@@ -3,7 +3,6 @@
 
 #if UE_ENABLE_INCLUDE_ORDER_DEPRECATED_IN_5_8
 #include "WorldPartition/DataLayer/DataLayerInstance.h"
-#include "WorldPartition/DataLayer/Deprecated/DeprecatedDataLayer.h"
 #endif // UE_ENABLE_INCLUDE_ORDER_DEPRECATED_IN_5_8
 
 UE_DEPRECATED_HEADER(5.8, "The contents of this header have been fully deprecated. See DeprecatedDataLayer.h for details on particular cases.") 
