[WorldPartition] Descriptor deprecation cleanup

#Synced-CL 51591079

diff --git a/Engine/Source/Runtime/Engine/Private/WorldPartition/WorldPartitionActorDescUtils.cpp b/Engine/Source/Runtime/Engine/Private/WorldPartition/WorldPartitionActorDescUtils.cpp
--- a/Engine/Source/Runtime/Engine/Private/WorldPartition/WorldPartitionActorDescUtils.cpp
+++ b/Engine/Source/Runtime/Engine/Private/WorldPartition/WorldPartitionActorDescUtils.cpp
@@ -158,16 +158,6 @@ TUniquePtr<FWorldPartitionActorDesc> FWorldPartitionActorDescUtils::GetActorDesc
 	return nullptr;
 }
 
-void FWorldPartitionActorDescUtils::AppendAssetDataTagsFromActor(const AActor* InActor, TArray<UObject::FAssetRegistryTag>& OutTags)
-{
-	FAssetRegistryTagsContextData Context(InActor, EAssetRegistryTagsCaller::Uncategorized);
-	AppendAssetDataTagsFromActor(InActor, Context);
-	for (TPair<FName, UObject::FAssetRegistryTag>& Pair : Context.Tags)
-	{
-		OutTags.Add(MoveTemp(Pair.Value));
-	}
-}
-
 void FWorldPartitionActorDescUtils::AppendAssetDataTagsFromActor(const AActor* InActor, FAssetRegistryTagsContext Context)
 {
 	// Avoid an assert when calling StaticFindObject during save to retrieve the actor's class.
@@ -266,12 +256,6 @@ void FWorldPartitionActorDescUtils::ReplaceActorDescriptorPointerFromActor(const
 	InActorDescInstance->ActorPtr = InNewActor;
 }
 
-bool FWorldPartitionActorDescUtils::FixupRedirectedAssetPath(FName& InOutAssetPath)
-{
-	UAssetRegistryHelpers::FixupRedirectedAssetPath(InOutAssetPath);
-	return true;
-}
-
 UClass* FWorldPartitionActorDescUtils::GetNativeClassFromClassName(FName InClassName)
 {
 	// Avoid an assert when calling StaticFindObject during save to retrieve the class.
diff --git a/Engine/Source/Runtime/Engine/Private/WorldPartition/WorldPartitionActorDescView.cpp b/Engine/Source/Runtime/Engine/Private/WorldPartition/WorldPartitionActorDescView.cpp
deleted file mode 100644
diff --git a/Engine/Source/Runtime/Engine/Private/WorldPartition/WorldPartitionActorDescViewProxy.cpp b/Engine/Source/Runtime/Engine/Private/WorldPartition/WorldPartitionActorDescViewProxy.cpp
deleted file mode 100644
diff --git a/Engine/Source/Runtime/Engine/Public/WorldPartition/ActorDescContainerInitParams.h b/Engine/Source/Runtime/Engine/Public/WorldPartition/ActorDescContainerInitParams.h
--- a/Engine/Source/Runtime/Engine/Public/WorldPartition/ActorDescContainerInitParams.h
+++ b/Engine/Source/Runtime/Engine/Public/WorldPartition/ActorDescContainerInitParams.h
@@ -13,12 +13,6 @@ class FWorldPartitionActorDesc;
 
 struct FActorDescContainerInitParams
 {
-	UE_DEPRECATED(5.4, "Use constructor with no UWorld param instead")
-	FActorDescContainerInitParams(UWorld* InWorld, FName InPackageName)
-		: PackageName(InPackageName)
-		, ContainerName(InPackageName.ToString())
-	{}
-
 	FActorDescContainerInitParams(FName InPackageName)
 		: PackageName(InPackageName)
 		, ContainerName(InPackageName.ToString())
diff --git a/Engine/Source/Runtime/Engine/Public/WorldPartition/PackedLevelActor/PackedLevelActorDesc.h b/Engine/Source/Runtime/Engine/Public/WorldPartition/PackedLevelActor/PackedLevelActorDesc.h
--- a/Engine/Source/Runtime/Engine/Public/WorldPartition/PackedLevelActor/PackedLevelActorDesc.h
+++ b/Engine/Source/Runtime/Engine/Public/WorldPartition/PackedLevelActor/PackedLevelActorDesc.h
@@ -14,10 +14,6 @@ public:
 	FPackedLevelActorDesc() {}
 	virtual ~FPackedLevelActorDesc() override = default;
 
-	virtual EWorldPartitionActorFilterType GetContainerFilterType() const { return EWorldPartitionActorFilterType::Packing; }
-	virtual FName GetContainerPackage() const override { return WorldAsset.GetLongPackageFName(); }
-	virtual const FWorldPartitionActorFilter* GetContainerFilter() const override { return &ContainerFilter; }
-
 protected:
 	ENGINE_API virtual void Init(const AActor* InActor) override;
 	ENGINE_API virtual void Init(const FWorldPartitionActorDescInitData& DescData) override;
diff --git a/Engine/Source/Runtime/Engine/Public/WorldPartition/WorldPartitionActorDesc.h b/Engine/Source/Runtime/Engine/Public/WorldPartition/WorldPartitionActorDesc.h
--- a/Engine/Source/Runtime/Engine/Public/WorldPartition/WorldPartitionActorDesc.h
+++ b/Engine/Source/Runtime/Engine/Public/WorldPartition/WorldPartitionActorDesc.h
@@ -301,15 +301,6 @@ public:
 		FTransform Transform = FTransform::Identity;
 		EContainerClusterMode ClusterMode;
 		TMap<FActorContainerID, TSet<FGuid>> FilteredActors;
-
-		UE_DEPRECATED(5.4, "Use FLoadedContainerInstance instead")
-		ULevel* LoadedLevel = nullptr;
-
-		UE_DEPRECATED(5.4, "Use FLoadedContainerInstance instead")
-		bool bSupportsPartialEditorLoading = false;
-
-		UE_DEPRECATED(5.4, "Use ContainerInstance instead")
-		const UActorDescContainer* Container = nullptr;
 	};
 
 	virtual ~FWorldPartitionActorDesc() = default;
@@ -391,103 +382,10 @@ public:
 	{
 		return GetTypeHash(Key.Guid);
 	}
-
+	
 	//~ Begin Deprecation
 	UE_DEPRECATED(5.1, "SetIsSpatiallyLoadedRaw is deprecated and should not be used.")
 	inline void SetIsSpatiallyLoadedRaw(bool bNewIsSpatiallyLoaded) { bIsSpatiallyLoaded = bNewIsSpatiallyLoaded; }
-
-PRAGMA_DISABLE_DEPRECATION_WARNINGS
-	UE_DEPRECATED(5.2, "GetOrigin is deprecated.")
-	inline FVector GetOrigin() const { return GetBounds().GetCenter(); }
-
-	UE_DEPRECATED(5.2, "GetBounds is deprecated, GetEditorBounds or GetRuntimeBounds should be used instead.")
-	FBox GetBounds() const { return GetEditorBounds(); }
-PRAGMA_ENABLE_DEPRECATION_WARNINGS
-
-	UE_DEPRECATED(5.2, "ShouldValidateRuntimeGrid is deprecated and should not be used")
-	virtual bool ShouldValidateRuntimeGrid() const { return true; }
-
-	UE_DEPRECATED(5.3, "GetLevelPackage is deprecated use GetChildContainerPackage instead")
-	virtual FName GetLevelPackage() const { return GetChildContainerPackage(); }
-
-	UE_DEPRECATED(5.4, "Use FWorldPartitionActorDescInstance::HasResolvedDataLayerInstanceNames")
-	inline bool HasResolvedDataLayerInstanceNames() const { return false; }
-
-	UE_DEPRECATED(5.4, "Use FWorldPartitionActorDescInstance::GetDataLayerInstanceNames")
-	TArray<FName> GetDataLayerInstanceNames() const { return TArray<FName>(); }
-
-	UE_DEPRECATED(5.4, "Use FWorldPartitionActorDescInstance::SetDataLayerInstanceNames")
-	inline void SetDataLayerInstanceNames(const TArray<FName>& InDataLayerInstanceNames) { }
-
-	UE_DEPRECATED(5.4, "Use FWorldPartitionActorDescInstance::IsRuntimeRelevant instead")
-	virtual bool IsRuntimeRelevant(const FActorContainerID& InContainerID) const { return false; }
-		
-	UE_DEPRECATED(5.4, "Use FWorldPartitionActorDescInstance::IsEditorRelevant instead")
-	virtual bool IsEditorRelevant() const { return false; }
-		
-	UE_DEPRECATED(5.4, "Please use FWorldPartitionActorDescInstance::GetUnloadedReason")
-	FText GetUnloadedReason() const { return FText(); }
-
-	UE_DEPRECATED(5.4, "Please use FWorldPartitionActorDescInstance::SetUnloadedReason")
-	void SetUnloadedReason(FText* InUnloadedReason) { }
-
-	UE_DEPRECATED(5.4, "Use IsChildContainerInstance instead")
-	virtual bool IsContainerInstance() const { return false; }
-
-	UE_DEPRECATED(5.4, "Use GetChildContainerPackage instead")
-	virtual FName GetContainerPackage() const { return NAME_None; }
-
-	UE_DEPRECATED(5.4, "Use GetChildContainerFilterType instead")
-	virtual EWorldPartitionActorFilterType GetContainerFilterType() const { return EWorldPartitionActorFilterType::None; }
-
-	UE_DEPRECATED(5.4, "Use GetChildContainerFilter instead")
-	virtual const FWorldPartitionActorFilter* GetContainerFilter() const { return nullptr; }
-
-	UE_DEPRECATED(5.4, "Use GetChildContainerInstance instead")
-	virtual bool GetContainerInstance(FContainerInstance& OutContainerInstance) const { return false; }
-
-	UE_DEPRECATED(5.4, "Use SetContainer without World param")
-	virtual void SetContainer(UActorDescContainer* InContainer, UWorld* InWorldContext)	{ }
-
-	UE_DEPRECATED(5.4, "Please use FWorldPartitionActorDescInstance::IsLoaded")
-	bool IsLoaded(bool bEvenIfPendingKill=false) const { return false; }
-	
-	UE_DEPRECATED(5.4, "Please use FWorldPartitionActorDescInstance::GetActor")
-	AActor* GetActor(bool bEvenIfPendingKill=true, bool bEvenIfUnreachable=false) const { return nullptr; }
-
-	UE_DEPRECATED(5.4, "Please use FWorldPartitionActorDescInstance::GetActor")
-	TWeakObjectPtr<AActor>* GetActorPtr(bool bEvenIfPendingKill = true, bool bEvenIfUnreachable = false) const { return nullptr; }
-
-	UE_DEPRECATED(5.4, "Please use FWorldPartitionActorDescInstance::Load")
-	AActor* Load() const { return nullptr;}
-
-	UE_DEPRECATED(5.4, "Please use FWorldPartitionActorDescInstance::Unload")
-	virtual void Unload() {}
-
-	UE_DEPRECATED(5.4, "Please use FWorldPartitionActorDescInstance")
-	void TransformInstance(const FString& From, const FString& To) {}
-
-	UE_DEPRECATED(5.4, "Use FWorldPartitionActorDescInstance version instead")
-	virtual void CheckForErrors(IStreamingGenerationErrorHandler* ErrorHandler) const {}
-
-protected:
-	UE_DEPRECATED(5.4, "Please use FWorldPartitionActorDescInstance::IncSoftRefCount")
-	inline uint32 IncSoftRefCount() const { return 0; }
-
-	UE_DEPRECATED(5.4, "Please use FWorldPartitionActorDescInstance::DecSoftRefCount")
-	inline uint32 DecSoftRefCount() const { return 0; }
-	
-	UE_DEPRECATED(5.4, "Please use FWorldPartitionActorDescInstance::IncHardRefCount")
-	inline uint32 IncHardRefCount() const { return 0;}
-	
-	UE_DEPRECATED(5.4, "Please use FWorldPartitionActorDescInstance::DecHardRefCount")
-	inline uint32 DecHardRefCount() const { return 0;}
-	
-	UE_DEPRECATED(5.4, "Please use FWorldPartitionActorDescInstance::GetSoftRefCount")
-	inline uint32 GetSoftRefCount() const { return 0; }
-	
-	UE_DEPRECATED(5.4, "Please use FWorldPartitionActorDescInstance::GetHardRefCount")
-	inline uint32 GetHardRefCount() const { return 0; }
 	//~ End Deprecation
 	
 public:
diff --git a/Engine/Source/Runtime/Engine/Public/WorldPartition/WorldPartitionActorDescUtils.h b/Engine/Source/Runtime/Engine/Public/WorldPartition/WorldPartitionActorDescUtils.h
--- a/Engine/Source/Runtime/Engine/Public/WorldPartition/WorldPartitionActorDescUtils.h
+++ b/Engine/Source/Runtime/Engine/Public/WorldPartition/WorldPartitionActorDescUtils.h
@@ -72,11 +72,9 @@ struct FWorldPartitionActorDescUtils
 	/**
 	 * Appends the actor's actor descriptor data into the provided asset registry tags.
 	 * @param InActor		The actor that will append its actor descriptor.
-	 * @param OutTags		Output tags to output into.
+	 * @param Context		Context to output into.
 	 */
 	static ENGINE_API void AppendAssetDataTagsFromActor(const AActor* InActor, FAssetRegistryTagsContext Context);
-	UE_DEPRECATED(5.4, "Call the version that takes FAssetRegistryTagsContext instead.")
-	static ENGINE_API void AppendAssetDataTagsFromActor(const AActor* InActor, TArray<UObject::FAssetRegistryTag>& OutTags);
 
 	/** 
 	 * Return the actor descriptor data.
@@ -116,12 +114,6 @@ struct FWorldPartitionActorDescUtils
 	 */
 	static ENGINE_API void ReplaceActorDescriptorPointerFromActor(const AActor* InOldActor, AActor* InNewActor, FWorldPartitionActorDescInstance* InActorDescInstance);
 
-	UE_DEPRECATED(5.4, "Use FWorldPartitionActorDescInstance version instead")
-	static void ReplaceActorDescriptorPointerFromActor(const AActor* InOldActor, AActor* InNewActor, FWorldPartitionActorDesc* InActorDesc) {}
-
-	UE_DEPRECATED(5.3, "Use FWorldPartitionHelpers::FixupRedirectedAssetPath instead")
-	static ENGINE_API bool FixupRedirectedAssetPath(FName& InOutAssetPath);
-
 	static ENGINE_API UClass* GetNativeClassFromClassName(FName InClassName);
 };
 #endif
diff --git a/Engine/Source/Runtime/Engine/Public/WorldPartition/WorldPartitionActorDescView.h b/Engine/Source/Runtime/Engine/Public/WorldPartition/WorldPartitionActorDescView.h
--- a/Engine/Source/Runtime/Engine/Public/WorldPartition/WorldPartitionActorDescView.h
+++ b/Engine/Source/Runtime/Engine/Public/WorldPartition/WorldPartitionActorDescView.h
@@ -1,122 +1,12 @@
 // Copyright Epic Games, Inc. All Rights Reserved.
 #pragma once
 
+#if UE_ENABLE_INCLUDE_ORDER_DEPRECATED_IN_5_8
 #include "CoreMinimal.h"
 #include "UObject/SoftObjectPath.h"
 #include "UObject/TopLevelAssetPath.h"
 #include "Misc/Optional.h"
 #include "WorldPartition/WorldPartitionActorDesc.h"
+#endif // UE_ENABLE_INCLUDE_ORDER_DEPRECATED_IN_5_8
 
-#if WITH_EDITOR
-class AActor;
-class IStreamingGenerationErrorHandler;
-class UActorDescContainer;
-class UWorldPartition;
-struct FWorldPartitionActorFilter;
-enum class EWorldPartitionActorFilterType : uint8;
-
-/**
- * A view on top of an actor desc, used to cache information that can be (potentially) different than the actor desc
- * itself due to streaming generation logic, etc.
- */
-class UE_DEPRECATED(5.4, "Class is deprecated in favor of IWorldPartitionActorDescInstanceView interface") FWorldPartitionActorDescView
-{
-	friend class UWorldPartitionRuntimeHash;
-
-public:
-	ENGINE_API FWorldPartitionActorDescView();
-	ENGINE_API FWorldPartitionActorDescView(const FWorldPartitionActorDesc * InActorDesc);
-
-	ENGINE_API const FGuid& GetGuid() const;
-	ENGINE_API FTopLevelAssetPath GetBaseClass() const;
-	ENGINE_API FTopLevelAssetPath GetNativeClass() const;
-	ENGINE_API UClass* GetActorNativeClass() const;
-	UE_DEPRECATED(5.2, "GetOrigin is deprecated.")
-		ENGINE_API FVector GetOrigin() const;
-	ENGINE_API FName GetRuntimeGrid() const;
-	ENGINE_API bool GetIsSpatiallyLoaded() const;
-	ENGINE_API bool GetActorIsEditorOnly() const;
-	ENGINE_API bool GetActorIsRuntimeOnly() const;
-	ENGINE_API bool GetActorIsHLODRelevant() const;
-	ENGINE_API FSoftObjectPath GetHLODLayer() const;
-	ENGINE_API const TArray<FName>& GetDataLayerInstanceNames() const;
-	ENGINE_API const TArray<FName>& GetRuntimeDataLayerInstanceNames() const;
-	ENGINE_API const TArray<FName>& GetTags() const;
-	ENGINE_API FName GetActorPackage() const;
-
-	ENGINE_API FSoftObjectPath GetActorSoftPath() const;
-	ENGINE_API FName GetActorLabel() const;
-
-	UE_DEPRECATED(5.2, "GetBounds is deprecated, GetEditorBounds or GetRuntimeBounds should be used instead.")
-		ENGINE_API FBox GetBounds() const;
-
-	ENGINE_API FBox GetEditorBounds() const;
-	ENGINE_API FBox GetRuntimeBounds() const;
-
-	ENGINE_API const TArray<FGuid>& GetReferences() const;
-	ENGINE_API const TArray<FGuid>& GetEditorReferences() const;
-	ENGINE_API FString ToString() const;
-	ENGINE_API const FGuid& GetParentActor() const;
-	ENGINE_API FName GetActorName() const;
-	ENGINE_API const FGuid& GetFolderGuid() const;
-
-	ENGINE_API FGuid GetContentBundleGuid() const;
-	ENGINE_API FName GetContainerPackage() const;
-	ENGINE_API bool IsContainerInstance() const;
-	ENGINE_API bool GetContainerInstance(FWorldPartitionActorDesc::FContainerInstance & OutContainerInstance) const;
-	ENGINE_API EWorldPartitionActorFilterType GetContainerFilterType() const;
-	ENGINE_API const FWorldPartitionActorFilter* GetContainerFilter() const;
-
-	UE_DEPRECATED(5.3, "GetLevelPackage is deprecated use GetContainerPackage instead.")
-		FName GetLevelPackage() const { return GetContainerPackage(); }
-
-	ENGINE_API void CheckForErrors(IStreamingGenerationErrorHandler * ErrorHandler) const;
-
-	ENGINE_API FName GetActorLabelOrName() const;
-
-	ENGINE_API void SetForcedNonSpatiallyLoaded();
-	ENGINE_API void SetForcedNoRuntimeGrid();
-	ENGINE_API void SetForcedNoDataLayers();
-	ENGINE_API void SetRuntimeDataLayerInstanceNames(const TArray<FName>&InRuntimeDataLayerInstanceNames);
-	ENGINE_API void SetRuntimeReferences(const TArray<FGuid>&InRuntimeReferences);
-	ENGINE_API void SetEditorReferences(const TArray<FGuid>&InEditorReferences);
-	ENGINE_API void SetDataLayerInstanceNames(const TArray<FName>&InDataLayerInstanceNames);
-
-	ENGINE_API void SetForcedNoHLODLayer();
-	ENGINE_API void SetRuntimeHLODLayer(const FSoftObjectPath & InHLODLayer);
-
-	ENGINE_API AActor* GetActor() const;
-
-	bool operator==(const FWorldPartitionActorDescView & Other) const
-	{
-		return GetGuid() == Other.GetGuid();
-	}
-
-	friend uint32 GetTypeHash(const FWorldPartitionActorDescView & Key)
-	{
-		return GetTypeHash(Key.GetGuid());
-	}
-
-	const FWorldPartitionActorDesc* GetActorDesc() const { return ActorDesc; }
-
-	ENGINE_API bool IsEditorOnlyReference(const FGuid & ReferenceGuid) const;
-
-	ENGINE_API bool GetProperty(FName PropertyName, FName * PropertyValue) const;
-	ENGINE_API bool HasProperty(FName PropertyName) const;
-
-	ENGINE_API void SetParentView(const FWorldPartitionActorDescView * InParentView);
-
-protected:
-	const FWorldPartitionActorDesc* ActorDesc;
-	const FWorldPartitionActorDescView* ParentView;
-	bool bIsForcedNonSpatiallyLoaded;
-	bool bIsForcedNoRuntimeGrid;
-	bool bIsForcedNoDataLayers;
-	bool bIsForceNoHLODLayer;
-	TOptional<TArray<FName>> ResolvedDataLayerInstanceNames;
-	TOptional<TArray<FName>> RuntimeDataLayerInstanceNames;
-	TOptional<TArray<FGuid>> RuntimeReferences;
-	TOptional<FSoftObjectPath> RuntimeHLODLayer;
-	TArray<FGuid> EditorReferences;
-};
-#endif
\ No newline at end of file
+UE_DEPRECATED_HEADER(5.8, "The contents of this header have been fully deprecated..") 
\ No newline at end of file
diff --git a/Engine/Source/Runtime/Engine/Public/WorldPartition/WorldPartitionActorDescViewProxy.h b/Engine/Source/Runtime/Engine/Public/WorldPartition/WorldPartitionActorDescViewProxy.h
--- a/Engine/Source/Runtime/Engine/Public/WorldPartition/WorldPartitionActorDescViewProxy.h
+++ b/Engine/Source/Runtime/Engine/Public/WorldPartition/WorldPartitionActorDescViewProxy.h
@@ -2,21 +2,9 @@
 #pragma once
 
 #if WITH_EDITOR
-
+#if UE_ENABLE_INCLUDE_ORDER_DEPRECATED_IN_5_8
 #include "WorldPartitionActorDescView.h"
+#endif // UE_ENABLE_INCLUDE_ORDER_DEPRECATED_IN_5_8
+#endif // WITH_EDITOR
 
-PRAGMA_DISABLE_DEPRECATION_WARNINGS
-/**
- * A view on top of an actor desc, which reverts to the actor if its loaded and dirty.
- */
-class ENGINE_API UE_DEPRECATED(5.4, "Class is deprecated") FWorldPartitionActorViewProxy : public FWorldPartitionActorDescView
-{
-public:
-	FWorldPartitionActorViewProxy(const FWorldPartitionActorDesc* InActorDesc);
-
-private:
-	TUniquePtr<FWorldPartitionActorDesc> CachedActorDesc;
-};
-PRAGMA_ENABLE_DEPRECATION_WARNINGS
-
-#endif
+UE_DEPRECATED_HEADER(5.8, "The contents of this header have been fully deprecated..") 
\ No newline at end of file
