Fixes for attachment and migrated objects.
- Clear AttachParent and AttachChildren when sending in USceneComponent::PostMigrate, so that they're not accessed later when the instance is destroyed. Similar to a move operation.
- Skip various component unregister logic in during PostMigrate if the component is not locally resident.
- Skip non-locally-resident attached children in UPrimitiveComponent::UnWeldChildren.

#Synced-CL 51599538

diff --git a/Engine/Source/Runtime/Engine/Classes/Components/SceneComponent.h b/Engine/Source/Runtime/Engine/Classes/Components/SceneComponent.h
--- a/Engine/Source/Runtime/Engine/Classes/Components/SceneComponent.h
+++ b/Engine/Source/Runtime/Engine/Classes/Components/SceneComponent.h
@@ -987,6 +987,10 @@ public:
 	ENGINE_API virtual void PostEditChangeChainProperty(FPropertyChangedChainEvent& PropertyChangedEvent) override;
 	ENGINE_API virtual bool CanEditChange(const FProperty* Property) const override;
 #endif
+
+#if UE_WITH_REMOTE_OBJECT_HANDLE
+	ENGINE_API virtual void PostMigrate(const struct FUObjectMigrationContext& MigrationContext) override;
+#endif
 	//~ End UObject Interface
 
 protected:
diff --git a/Engine/Source/Runtime/Engine/Private/Components/ActorComponent.cpp b/Engine/Source/Runtime/Engine/Private/Components/ActorComponent.cpp
--- a/Engine/Source/Runtime/Engine/Private/Components/ActorComponent.cpp
+++ b/Engine/Source/Runtime/Engine/Private/Components/ActorComponent.cpp
@@ -2130,6 +2130,14 @@ void UActorComponent::UnregisterComponent()
 		return;
 	}
 
+#if UE_WITH_REMOTE_OBJECT_HANDLE
+	if (UE::RemoteObject::Handle::GetResidence(this) == EResidence::LocalNotReady)
+	{
+		UE_LOGF(LogActorComponent, Log, "UnregisterComponent: (%ls) has LocalNotReady residence. Skipping unregistration.", *GetPathName());
+		return;
+	}
+#endif
+
 	// If registered, should have a world
 	checkf(WorldPrivate != nullptr, TEXT("%s"), *GetFullName());
 
diff --git a/Engine/Source/Runtime/Engine/Private/Components/SceneComponent.cpp b/Engine/Source/Runtime/Engine/Private/Components/SceneComponent.cpp
--- a/Engine/Source/Runtime/Engine/Private/Components/SceneComponent.cpp
+++ b/Engine/Source/Runtime/Engine/Private/Components/SceneComponent.cpp
@@ -4108,6 +4108,19 @@ bool USceneComponent::CanEditChange( const FProperty* Property ) const
 }
 #endif
 
+#if UE_WITH_REMOTE_OBJECT_HANDLE
+void USceneComponent::PostMigrate(const FUObjectMigrationContext& MigrationContext)
+{
+	Super::PostMigrate(MigrationContext);
+
+	if (MigrationContext.MigrationSide == EObjectMigrationSide::Send)
+	{
+		// Clear references to attachments so they don't get cleared on the other objects when this instance is destroyed.
+		AttachParent = nullptr;
+		AttachChildren.Reset();
+	}
+}
+#endif
 
 /**
  * FScopedPreventAttachedComponentMove implementation
diff --git a/Engine/Source/Runtime/Engine/Private/PrimitiveComponentPhysics.cpp b/Engine/Source/Runtime/Engine/Private/PrimitiveComponentPhysics.cpp
--- a/Engine/Source/Runtime/Engine/Private/PrimitiveComponentPhysics.cpp
+++ b/Engine/Source/Runtime/Engine/Private/PrimitiveComponentPhysics.cpp
@@ -1239,11 +1239,14 @@ void UPrimitiveComponent::UnWeldFromParent()
 
 void UPrimitiveComponent::UnWeldChildren()
 {
-	for (USceneComponent* ChildComponent : GetAttachChildren())
+	for (TObjectPtr<USceneComponent> ChildComponent : GetAttachChildren())
 	{
-		if (UPrimitiveComponent* PrimComp = Cast<UPrimitiveComponent>(ChildComponent))
+		if (ChildComponent.GetResidence() == EResidence::Local)
 		{
-			PrimComp->UnWeldFromParent();
+			if (UPrimitiveComponent* PrimComp = Cast<UPrimitiveComponent>(ChildComponent))
+			{
+				PrimComp->UnWeldFromParent();
+			}
 		}
 	}
 }
