Propagate client-side instantiation failure reasons to server for "Client disabled replication" ensures

When Iris replication fails to instantiate a net object on the client, the client disables replication for that handle and notifies the server. Previously, the server-side ensure only reported which object failed, but not why the client couldn't create it -- requiring access to client logs to diagnose the root cause.
This change:
- Adds a FailureDiagnosticMessage string to FInstantiateResult, populated by all factory failure paths with a description of the specific failure.
- Adds ErrorDiagnostic storage to FNetErrorContext / FNetSerializationContext so the diagnostic can flow from the factory through the replication reader.
- Threads the diagnostic string through SendErrorWithNetRefHandle -> OnErrorWithNetRefHandleReported -> OnErrorCannotReplicateObject.
- Introduces a new control channel message NMT_IrisNetRefHandleErrorWithDiagnosticMessage to carry the diagnostic from client to server. The original NMT_IrisNetRefHandleError message is still used when no diagnostic is available.
- Updates the server-side OnErrorCannotReplicateObject ensure to include the client's diagnostic reason when present.

This allows gameplay engineers to diagnose broken NetRefHandle issues from server logs alone, without needing client logs or a consistent repro.

#tests Verified with local server + clients that error message propagated to server ensure.
#Synced-CL 51332317

diff --git a/Engine/Source/Runtime/Engine/Private/DataChannel.cpp b/Engine/Source/Runtime/Engine/Private/DataChannel.cpp
--- a/Engine/Source/Runtime/Engine/Private/DataChannel.cpp
+++ b/Engine/Source/Runtime/Engine/Private/DataChannel.cpp
@@ -1757,6 +1757,7 @@ IMPLEMENT_CONTROL_CHANNEL_MESSAGE(BeaconAssignGUID);
 IMPLEMENT_CONTROL_CHANNEL_MESSAGE(BeaconNetGUIDAck);
 IMPLEMENT_CONTROL_CHANNEL_MESSAGE(IrisProtocolMismatch);
 IMPLEMENT_CONTROL_CHANNEL_MESSAGE(IrisNetRefHandleError);
+IMPLEMENT_CONTROL_CHANNEL_MESSAGE(IrisNetRefHandleErrorWithDiagnosticMessage);
 
 void UControlChannel::Init( UNetConnection* InConnection, int32 InChannelIndex, EChannelCreateFlags CreateFlags )
 {
@@ -1992,6 +1993,20 @@ void UControlChannel::ReceivedBunch( FInBunch& Bunch )
 				}
 			}
 		}
+		else if (MessageType == NMT_IrisNetRefHandleErrorWithDiagnosticMessage)
+		{
+			UE::Net::ENetRefHandleError NetRefHandleError = UE::Net::ENetRefHandleError::None;
+			uint64 NetRefHandleId = 0;
+			TArray<uint64> ExtraNetRefs;
+			FString DiagnosticMessage;
+			if (FNetControlMessage<NMT_IrisNetRefHandleErrorWithDiagnosticMessage>::Receive(Bunch, NetRefHandleError, NetRefHandleId, ExtraNetRefs, DiagnosticMessage))
+			{
+				if (UReplicationSystem* IrisRepSystem = Connection->Driver->GetReplicationSystem())
+				{
+					IrisRepSystem->ReportErrorWithNetRefHandle(NetRefHandleError, NetRefHandleId, Connection->GetConnectionHandle().GetParentConnectionId(), ExtraNetRefs, DiagnosticMessage);
+				}
+			}
+		}
 		else if (Connection->Driver->Notify != nullptr)
 		{
 			// Process control message on client/server connection
@@ -2097,6 +2112,9 @@ void UControlChannel::ReceivedBunch( FInBunch& Bunch )
 				case NMT_CloseChildConnection:
 					FNetControlMessage<NMT_CloseChildConnection>::Discard(Bunch);
 					break;
+				case NMT_IrisNetRefHandleErrorWithDiagnosticMessage:
+					FNetControlMessage<NMT_IrisNetRefHandleErrorWithDiagnosticMessage>::Discard(Bunch);
+					break;
 				default:
 					// if this fails, a case is missing above for an implemented message type
 					// or the connection is being sent potentially malformed packets
diff --git a/Engine/Source/Runtime/Engine/Private/Net/Iris/ReplicationSystem/EngineReplicationBridge.cpp b/Engine/Source/Runtime/Engine/Private/Net/Iris/ReplicationSystem/EngineReplicationBridge.cpp
--- a/Engine/Source/Runtime/Engine/Private/Net/Iris/ReplicationSystem/EngineReplicationBridge.cpp
+++ b/Engine/Source/Runtime/Engine/Private/Net/Iris/ReplicationSystem/EngineReplicationBridge.cpp
@@ -1214,7 +1214,7 @@ void UEngineReplicationBridge::OnProtocolMismatchReported(FNetRefHandle RefHandl
 	}
 }
 
-void UEngineReplicationBridge::SendErrorWithNetRefHandle(UE::Net::ENetRefHandleError ErrorType, FNetRefHandle RefHandle, uint32 ConnectionId, TConstArrayView<const FNetRefHandle> ExtraNetRefHandle)
+void UEngineReplicationBridge::SendErrorWithNetRefHandle(UE::Net::ENetRefHandleError ErrorType, FNetRefHandle RefHandle, uint32 ConnectionId, TConstArrayView<const FNetRefHandle> ExtraNetRefHandle, const FString& DiagnosticMessage)
 {
 	if (NetDriver)
 	{
@@ -1229,7 +1229,14 @@ void UEngineReplicationBridge::SendErrorWithNetRefHandle(UE::Net::ENetRefHandleE
 				ArrayIds.Emplace(ExtraHandle.GetId());
 			}
 
-			FNetControlMessage<NMT_IrisNetRefHandleError>::Send(ClientConnection, ErrorType, HandleId, ArrayIds);
+			if (DiagnosticMessage.IsEmpty())
+			{
+				FNetControlMessage<NMT_IrisNetRefHandleError>::Send(ClientConnection, ErrorType, HandleId, ArrayIds);
+			}
+			else
+			{
+				FNetControlMessage<NMT_IrisNetRefHandleErrorWithDiagnosticMessage>::Send(ClientConnection, ErrorType, HandleId, ArrayIds, const_cast<FString&>(DiagnosticMessage));
+			}
 		}
 		else
 		{
diff --git a/Engine/Source/Runtime/Engine/Private/Net/Iris/ReplicationSystem/NetActorFactory.cpp b/Engine/Source/Runtime/Engine/Private/Net/Iris/ReplicationSystem/NetActorFactory.cpp
--- a/Engine/Source/Runtime/Engine/Private/Net/Iris/ReplicationSystem/NetActorFactory.cpp
+++ b/Engine/Source/Runtime/Engine/Private/Net/Iris/ReplicationSystem/NetActorFactory.cpp
@@ -308,8 +308,14 @@ UNetObjectFactory::FInstantiateResult UNetActorFactory::InstantiateReplicatedObj
 		AActor* Actor = Cast<AActor>(ActorBridge->ResolveObjectReference(StaticHeader->ObjectReference, Context.ResolveContext));
 		if (!Actor)
 		{
-			UE_LOG(LogIris, Error, TEXT("UNetActorFactory::InstantiateNetObjectFromHeader Failed to resolve ObjectReference: %s . Could not find static actor."), *ActorBridge->DescribeObjectReference(StaticHeader->ObjectReference, Context.ResolveContext));
-			return FInstantiateResult();
+			FInstantiateResult InstantiateResult;
+			if (UE_LOG_ACTIVE(LogIris, Error))
+			{
+				const FString FailedObjectReferenceStr = ActorBridge->DescribeObjectReference(StaticHeader->ObjectReference, Context.ResolveContext);
+				UE_LOG(LogIris, Error, TEXT("UNetActorFactory::InstantiateNetObjectFromHeader Failed to resolve ObjectReference: %s . Could not find static actor."), *FailedObjectReferenceStr);
+				InstantiateResult.FailureDiagnosticMessage = FString::Printf(TEXT("Static actor not found. ObjectReference: %s"), *FailedObjectReferenceStr);
+			}
+			return InstantiateResult;
 		}
 
 		UE_LOG(LogIris, Verbose, TEXT("UNetActorFactory::InstantiateNetObjectFromHeader Found static Actor: %s using ObjectReference: %s"), ToCStr(Actor->GetPathName()), *ActorBridge->DescribeObjectReference(StaticHeader->ObjectReference, Context.ResolveContext));
@@ -332,8 +338,13 @@ UNetObjectFactory::FInstantiateResult UNetActorFactory::InstantiateReplicatedObj
 		AActor* FoundActor = Cast<AActor>(Bridge->GetPreRegisteredObject(Context.Handle));
 		if (!FoundActor)
 		{
-			UE_LOG(LogIris, Error, TEXT("UNetActorFactory::InstantiateReplicatedObjectFromHeader Unable to find pre-registered actor: %s"), *Context.Handle.ToString());
-			return FInstantiateResult();
+			FInstantiateResult InstantiateResult;
+			if (UE_LOG_ACTIVE(LogIris, Error))
+			{
+				UE_LOG(LogIris, Error, TEXT("UNetActorFactory::InstantiateReplicatedObjectFromHeader Unable to find pre-registered actor: %s"), *Context.Handle.ToString());
+				InstantiateResult.FailureDiagnosticMessage = FString::Printf(TEXT("Pre-registered actor not found. Handle: %s"), *Context.Handle.ToString());
+			}
+			return InstantiateResult;
 		}
 
 		FInstantiateResult Result{ .Instance = FoundActor };
@@ -353,8 +364,14 @@ UNetObjectFactory::FInstantiateResult UNetActorFactory::InstantiateReplicatedObj
 	AActor* Archetype = Cast<AActor>(Bridge->ResolveObjectReference(DynamicHeader->ArchetypeReference, Context.ResolveContext));
 	if (!Archetype)
 	{
-		UE_LOG(LogIris, Error, TEXT("UNetActorFactory::InstantiateNetObjectFromHeader Unable to instantiate actor, failed to resolve archetype: %s"), *Bridge->DescribeObjectReference(DynamicHeader->ArchetypeReference, Context.ResolveContext));
-		return FInstantiateResult();
+		FInstantiateResult InstantiateResult;
+		if (UE_LOG_ACTIVE(LogIris, Error))
+		{
+			const FString FailedArchetypeReferenceStr = Bridge->DescribeObjectReference(DynamicHeader->ArchetypeReference, Context.ResolveContext);
+			UE_LOG(LogIris, Error, TEXT("UNetActorFactory::InstantiateNetObjectFromHeader Unable to instantiate actor, failed to resolve archetype: %s"), *FailedArchetypeReferenceStr);
+			InstantiateResult.FailureDiagnosticMessage = FString::Printf(TEXT("Actor archetype not found. Archetype: %s"), *FailedArchetypeReferenceStr);
+		}
+		return InstantiateResult;
 	}
 
 	// Find level
@@ -390,8 +407,20 @@ UNetObjectFactory::FInstantiateResult UNetActorFactory::InstantiateReplicatedObj
 	// For Iris we expect that we will be able to spawn the actor as streaming always is controlled from server
 	if (!Actor)
 	{
-		ensureMsgf(Actor, TEXT("UNetActorFactory::InstantiateNetObjectFromHeader SpawnActor failed. Used Archetype: %s ObjectReference: %s"), *GetNameSafe(Archetype), *ActorBridge->DescribeObjectReference(DynamicHeader->ArchetypeReference, Context.ResolveContext));
-		return FInstantiateResult();
+		FInstantiateResult InstantiateResult;
+		if (UE_LOG_ACTIVE(LogIris, Error))
+		{
+			const FString ArchetypeNameStr = GetNameSafe(Archetype);
+			const FString ObjectReferenceStr = Bridge->DescribeObjectReference(DynamicHeader->ArchetypeReference, Context.ResolveContext);
+			ensureMsgf(Actor, TEXT("UNetActorFactory::InstantiateNetObjectFromHeader SpawnActor failed. Used Archetype: %s ObjectReference: %s"), *ArchetypeNameStr, *ObjectReferenceStr);
+			InstantiateResult.FailureDiagnosticMessage = FString::Printf(TEXT("SpawnActor failed. Archetype: %s, ObjectReference: %s"), *ArchetypeNameStr, *ObjectReferenceStr);
+		}
+		else
+		{
+			// We still want to ensure if the LogIris Error is disabled, but only construct the strings if ensures are enabled.
+			ensureMsgf(Actor, TEXT("UNetActorFactory::InstantiateNetObjectFromHeader SpawnActor failed. Used Archetype: %s ObjectReference: %s"), *GetNameSafe(Archetype), *Bridge->DescribeObjectReference(DynamicHeader->ArchetypeReference, Context.ResolveContext));
+		}
+		return InstantiateResult;
 	}
 
 	static constexpr float Epsilon = UE_KINDA_SMALL_NUMBER;
diff --git a/Engine/Source/Runtime/Engine/Private/Net/Iris/ReplicationSystem/NetRootObjectFactory.cpp b/Engine/Source/Runtime/Engine/Private/Net/Iris/ReplicationSystem/NetRootObjectFactory.cpp
--- a/Engine/Source/Runtime/Engine/Private/Net/Iris/ReplicationSystem/NetRootObjectFactory.cpp
+++ b/Engine/Source/Runtime/Engine/Private/Net/Iris/ReplicationSystem/NetRootObjectFactory.cpp
@@ -54,8 +54,14 @@ UNetObjectFactory::FInstantiateResult UNetRootObjectFactory::InstantiateReplicat
 
 		if (!RootObjectPtr)
 		{
-			UE_LOG(LogIris, Error, TEXT("%hs: Failed to resolve stable reference: %s . Could not find stable object to bind to."), __FUNCTION__, *Bridge->DescribeObjectReference(RootObjectHeader->Reference, Context.ResolveContext));
-			return FInstantiateResult();
+			FInstantiateResult InstantiateResult;
+			if (UE_LOG_ACTIVE(LogIris, Error))
+			{
+				const FString RootObjectReferenceStr = Bridge->DescribeObjectReference(RootObjectHeader->Reference, Context.ResolveContext);
+				UE_LOG(LogIris, Error, TEXT("%hs: Failed to resolve stable reference: %s . Could not find stable object to bind to."), __FUNCTION__, *RootObjectReferenceStr);
+				InstantiateResult.FailureDiagnosticMessage = FString::Printf(TEXT("Stable object not found. Reference: %s"), *RootObjectReferenceStr);
+			}
+			return InstantiateResult;
 		}
 
 		UE_LOG(LogIris, Verbose, TEXT("%hs: Found stable object: %s with Handle: %s via: %s"), __FUNCTION__, *GetNameSafe(RootObjectPtr), *Context.Handle.ToString(), *Bridge->DescribeObjectReference(RootObjectHeader->Reference, Context.ResolveContext));
@@ -71,8 +77,13 @@ UNetObjectFactory::FInstantiateResult UNetRootObjectFactory::InstantiateReplicat
 		UObject* RegisteredObject = Bridge->GetPreRegisteredObject(Context.Handle);
 		if (!RegisteredObject)
 		{
-			UE_LOG(LogIris, Error, TEXT("%hs Unable to find pre-registered root object: %s"), __FUNCTION__, *Context.Handle.ToString());
-			return FInstantiateResult();
+			FInstantiateResult InstantiateResult;
+			if (UE_LOG_ACTIVE(LogIris, Error))
+			{
+				UE_LOG(LogIris, Error, TEXT("%hs Unable to find pre-registered root object: %s"), __FUNCTION__, *Context.Handle.ToString());
+				InstantiateResult.FailureDiagnosticMessage = FString::Printf(TEXT("Pre-registered object not found. Handle %s"), *Context.Handle.ToString());
+			}
+			return InstantiateResult;
 		}
 
 		UE_LOG(LogIris, Verbose, TEXT("%hs: Found pre-registered object: %s with Handle: %s"), __FUNCTION__, *GetNameSafe(RegisteredObject), *Context.Handle.ToString());
@@ -86,8 +97,14 @@ UNetObjectFactory::FInstantiateResult UNetRootObjectFactory::InstantiateReplicat
 	UObject* Template = Bridge->ResolveObjectReference(RootObjectHeader->Reference, Context.ResolveContext);
 	if (!Template)
 	{
-		UE_LOG(LogIris, Error, TEXT("%hs Unable to instantiate root object, failed to resolve template: %s"), __FUNCTION__, *Bridge->DescribeObjectReference(RootObjectHeader->Reference, Context.ResolveContext));
-		return FInstantiateResult();
+		FInstantiateResult InstantiateResult;
+		if (UE_LOG_ACTIVE(LogIris, Error))
+		{
+			const FString TemplateReferenceStr = Bridge->DescribeObjectReference(RootObjectHeader->Reference, Context.ResolveContext);
+			UE_LOG(LogIris, Error, TEXT("%hs Unable to instantiate root object, failed to resolve template: %s"), __FUNCTION__, *TemplateReferenceStr);
+			InstantiateResult.FailureDiagnosticMessage = FString::Printf(TEXT("Dynamic object template not found. TemplateReference: %s"), *TemplateReferenceStr);
+		}
+		return InstantiateResult;
 	}
 
 	UObject* DynamicRootObject = NewObject<UObject>(GetTransientPackage(), Template->GetClass(), NAME_None, RF_NoFlags, Template);
diff --git a/Engine/Source/Runtime/Engine/Private/Net/Iris/ReplicationSystem/NetSubObjectFactory.cpp b/Engine/Source/Runtime/Engine/Private/Net/Iris/ReplicationSystem/NetSubObjectFactory.cpp
--- a/Engine/Source/Runtime/Engine/Private/Net/Iris/ReplicationSystem/NetSubObjectFactory.cpp
+++ b/Engine/Source/Runtime/Engine/Private/Net/Iris/ReplicationSystem/NetSubObjectFactory.cpp
@@ -167,9 +167,16 @@ UNetSubObjectFactory::FInstantiateResult UNetSubObjectFactory::InstantiateReplic
 
 		if (!SubObject)
 		{
-			UE_LOG(LogIris, Error, TEXT("UNetSubObjectFactory::InstantiateReplicatedObjectFromHeader %s: Failed to find static or stable name object referenced by SubObject: %s, Owner: %s, RootObject: %s"), 
-			*Context.Handle.ToString(), *Bridge->DescribeObjectReference(SubObjectHeader->ObjectReference, Context.ResolveContext), *Bridge->PrintObjectFromNetRefHandle(Context.RootObjectOfSubObject), *GetPathNameSafe(Bridge->GetReplicatedObject(Context.RootObjectOfSubObject)));
-			return FInstantiateResult();
+			FInstantiateResult InstantiateResult;
+			if (UE_LOG_ACTIVE(LogIris, Error))
+			{
+				const FString SubObjectReferenceStr = Bridge->DescribeObjectReference(SubObjectHeader->ObjectReference, Context.ResolveContext);
+				const FString OwnerStr = Bridge->PrintObjectFromNetRefHandle(Context.RootObjectOfSubObject);
+				const FString RootObjectStr = GetPathNameSafe(Bridge->GetReplicatedObject(Context.RootObjectOfSubObject));
+				UE_LOG(LogIris, Error, TEXT("UNetSubObjectFactory::InstantiateReplicatedObjectFromHeader %s: Failed to find static or stable name object referenced by SubObject: %s, Owner: %s, RootObject: %s"), *Context.Handle.ToString(), *SubObjectReferenceStr, *OwnerStr, *RootObjectStr);
+				InstantiateResult.FailureDiagnosticMessage = FString::Printf(TEXT("Static subobject not found. SubObjectReference: %s, Owner: %s, RootObject: %s"), *SubObjectReferenceStr, *OwnerStr, *RootObjectStr);
+			}
+			return InstantiateResult;
 		}
 
 		UE_LOG(LogIris, Verbose, TEXT("UNetSubObjectFactory::InstantiateReplicatedObjectFromHeader %s: Found static or stable name SubObject using path %s"), *Context.Handle.ToString(), ToCStr(SubObject->GetPathName()));
@@ -187,9 +194,17 @@ UNetSubObjectFactory::FInstantiateResult UNetSubObjectFactory::InstantiateReplic
 	UObject* Template = Bridge->ResolveObjectReference(SubObjectHeader->TemplateReference, Context.ResolveContext);
 	if (!Template)
 	{
-		UE_LOG(LogIris, Error, TEXT("UNetSubObjectFactory::InstantiateReplicatedObjectFromHeader Unable to resolve template: %s, cannot instantiate dynamic subobject %s"), *Bridge->DescribeObjectReference(SubObjectHeader->TemplateReference, Context.ResolveContext), *Context.Handle.ToString());
+		FInstantiateResult InstantiateResult;
+		if (UE_LOG_ACTIVE(LogIris, Error))
+		{
+			const FString TemplateReferenceStr = Bridge->DescribeObjectReference(SubObjectHeader->TemplateReference, Context.ResolveContext);
+			UE_LOG(LogIris, Error, TEXT("UNetSubObjectFactory::InstantiateReplicatedObjectFromHeader Unable to resolve template: %s, cannot instantiate dynamic subobject %s"), *TemplateReferenceStr, *Context.Handle.ToString());
+			InstantiateResult.FailureDiagnosticMessage = FString::Printf(TEXT("Dynamic subobject template not found. TemplateReference: %s"), *TemplateReferenceStr);
+		}
+
 		ensure(Template);
-		return FInstantiateResult();
+
+		return InstantiateResult;
 	}
 			
 	// Find the proper Outer
diff --git a/Engine/Source/Runtime/Engine/Public/Net/DataChannel.h b/Engine/Source/Runtime/Engine/Public/Net/DataChannel.h
--- a/Engine/Source/Runtime/Engine/Public/Net/DataChannel.h
+++ b/Engine/Source/Runtime/Engine/Public/Net/DataChannel.h
@@ -239,4 +239,7 @@ DEFINE_CONTROL_CHANNEL_MESSAGE(JoinNoPawnSplit, 32, FString, FUniqueNetIdRepl);
 //
 // It is expected that the local child connection will be cleaned up manually and does not depend on this message being 
 // successfully processed by the remote server.
-DEFINE_CONTROL_CHANNEL_MESSAGE(CloseChildConnection, 34, int32);
\ No newline at end of file
+DEFINE_CONTROL_CHANNEL_MESSAGE(CloseChildConnection, 34, int32);
+
+// a specific handle caused an error and we want the remote connection to log all information it has on it
+DEFINE_CONTROL_CHANNEL_MESSAGE(IrisNetRefHandleErrorWithDiagnosticMessage, 35, UE::Net::ENetRefHandleError, uint64, TArray<uint64>, FString);
\ No newline at end of file
diff --git a/Engine/Source/Runtime/Engine/Public/Net/Iris/ReplicationSystem/EngineReplicationBridge.h b/Engine/Source/Runtime/Engine/Public/Net/Iris/ReplicationSystem/EngineReplicationBridge.h
--- a/Engine/Source/Runtime/Engine/Public/Net/Iris/ReplicationSystem/EngineReplicationBridge.h
+++ b/Engine/Source/Runtime/Engine/Public/Net/Iris/ReplicationSystem/EngineReplicationBridge.h
@@ -128,7 +128,7 @@ public:
 	UIrisObjectReferencePackageMap* GetObjectReferencePackageMap() const { return ObjectReferencePackageMap; }
 
 	/** Tell the remote connection that we detected a reading error with a specific replicated object */
-	ENGINE_API virtual void SendErrorWithNetRefHandle(UE::Net::ENetRefHandleError ErrorType, FNetRefHandle RefHandle, uint32 ConnectionId, TConstArrayView<const FNetRefHandle> ExtraNetRefHandle = {}) override;
+	ENGINE_API virtual void SendErrorWithNetRefHandle(UE::Net::ENetRefHandleError ErrorType, FNetRefHandle RefHandle, uint32 ConnectionId, TConstArrayView<const FNetRefHandle> ExtraNetRefHandle = {}, const FString& DiagnosticMessage = FString()) override;
 
 	/** Add the rootobject to the level's filter group so it will only be relevant if the connection has that level streamed in. */
 	UE_DEPRECATED(5.8, "Call AddRootObjectToContainerGroup instead.")
diff --git a/Engine/Source/Runtime/Net/Iris/Private/Iris/ReplicationSystem/ObjectReplicationBridge.cpp b/Engine/Source/Runtime/Net/Iris/Private/Iris/ReplicationSystem/ObjectReplicationBridge.cpp
--- a/Engine/Source/Runtime/Net/Iris/Private/Iris/ReplicationSystem/ObjectReplicationBridge.cpp
+++ b/Engine/Source/Runtime/Net/Iris/Private/Iris/ReplicationSystem/ObjectReplicationBridge.cpp
@@ -1842,6 +1842,15 @@ FReplicationBridgeCreateNetRefHandleResult UObjectReplicationBridge::CreateNetRe
 				UE_LOG_BRIDGEID(LogIris, Warning, TEXT("CreateNetRefHandleFromRemote: Failed to instantiate RootObject NetHandle: %s using header: %s"), *WantedNetHandle.ToString(), ToCStr(Header->ToString()));
 			}
 		}
+
+		if (Result.FailureDiagnosticMessage.IsEmpty())
+		{
+			Context.SerializationContext.SetErrorDiagnostic(FString::Printf(TEXT("Factory %s failed to instantiate. No diagnostic message returned."), *Factory->GetName()));
+		}
+		else
+		{
+			Context.SerializationContext.SetErrorDiagnostic(FString::Printf(TEXT("Factory %s failed to instantiate. Error: %s"), *Factory->GetName(), *Result.FailureDiagnosticMessage));
+		}
 		
 		ensureMsgf(bSuppressCreateInstanceFailedEnsure, TEXT("Failed to instantiate Handle: %s"), *WantedNetHandle.ToString());
 		return FReplicationBridgeCreateNetRefHandleResult();
@@ -3524,7 +3533,7 @@ void UObjectReplicationBridge::OnProtocolMismatchReported(FNetRefHandle RefHandl
 	//TODO: Set object as broken in RepWriter
 }
 
-void UObjectReplicationBridge::OnErrorWithNetRefHandleReported(UE::Net::ENetRefHandleError ErrorType, FNetRefHandle RefHandle, uint32 ConnectionId, const TArray<FNetRefHandle>& ExtraHandles)
+void UObjectReplicationBridge::OnErrorWithNetRefHandleReported(UE::Net::ENetRefHandleError ErrorType, FNetRefHandle RefHandle, uint32 ConnectionId, const TArray<FNetRefHandle>& ExtraHandles, const FString& DiagnosticMessage)
 {
 	using namespace UE::Net;
 	using namespace UE::Net::Private;
@@ -3538,7 +3547,7 @@ void UObjectReplicationBridge::OnErrorWithNetRefHandleReported(UE::Net::ENetRefH
 		
 		case ENetRefHandleError::ReplicationDisabled:
 		{
-			OnErrorCannotReplicateObject(RefHandle, ConnectionId);
+			OnErrorCannotReplicateObject(RefHandle, ConnectionId, DiagnosticMessage);
 		} break;
 		
 		case ENetRefHandleError::BlockedByMustBeMapped:
@@ -3579,7 +3588,7 @@ void UObjectReplicationBridge::OnErrorBitstreamCorrupted(FNetRefHandle RefHandle
 	ensureMsgf(false, TEXT("Bitstream corruption reported with %s. Look at the log for important information on the object tied to the handle."), *NetRefHandleManager->PrintObjectFromNetRefHandle(RefHandle));
 }
 
-void UObjectReplicationBridge::OnErrorCannotReplicateObject(FNetRefHandle RefHandle, uint32 ConnectionId)
+void UObjectReplicationBridge::OnErrorCannotReplicateObject(FNetRefHandle RefHandle, uint32 ConnectionId, const FString& ClientDiagnosticMessage)
 {
 	using namespace UE::Net;
 	using namespace UE::Net::Private;
@@ -3596,7 +3605,7 @@ void UObjectReplicationBridge::OnErrorCannotReplicateObject(FNetRefHandle RefHan
 		// TODO: Tell RepWriter the creation failed
 	}
 
-	ensureMsgf(false, TEXT("Client disabled replication for %s. Client was unable to read the creation data and did not spawn the netobject."), *NetRefHandleManager->PrintObjectFromNetRefHandle(RefHandle));
+	ensureMsgf(false, TEXT("Client could not instantiate %s, replication disabled due to %s"), *NetRefHandleManager->PrintObjectFromNetRefHandle(RefHandle), ClientDiagnosticMessage.IsEmpty() ? TEXT("unknown") : *ClientDiagnosticMessage);
 }
 
 void UObjectReplicationBridge::OnErrorBlockedByAsyncLoading(FNetRefHandle RefHandle, uint32 ConnectionId, const TArray<FNetRefHandle>& MustBeMappedRefs)
diff --git a/Engine/Source/Runtime/Net/Iris/Private/Iris/ReplicationSystem/ReplicationReader.cpp b/Engine/Source/Runtime/Net/Iris/Private/Iris/ReplicationSystem/ReplicationReader.cpp
--- a/Engine/Source/Runtime/Net/Iris/Private/Iris/ReplicationSystem/ReplicationReader.cpp
+++ b/Engine/Source/Runtime/Net/Iris/Private/Iris/ReplicationSystem/ReplicationReader.cpp
@@ -1165,7 +1165,7 @@ uint32 FReplicationReader::ReadObjectBatch(FNetSerializationContext& Context, ui
 		{
 			if (Context.GetError() == GNetError_BrokenNetHandle)
 			{
-				ReplicationBridge->SendErrorWithNetRefHandle(UE::Net::ENetRefHandleError::ReplicationDisabled, IncompleteHandle, Parameters.ConnectionId);
+				ReplicationBridge->SendErrorWithNetRefHandle(UE::Net::ENetRefHandleError::ReplicationDisabled, IncompleteHandle, Parameters.ConnectionId, {}, Context.GetErrorDiagnostic());
  
 				// Log error and try to recover, if get more incoming data for an object in the broken state we will skip it.
 				UE_LOG(LogIris, Error, TEXT("FReplicationReader::ReadObject Failed to read object batch handle: %s skipping batch data"), ToCStr(IncompleteHandle.ToString()));
@@ -2733,7 +2733,7 @@ void FReplicationReader::ProcessQueuedBatches()
 
 				if (Context.GetError() == GNetError_BrokenNetHandle)
 				{
-					ReplicationBridge->SendErrorWithNetRefHandle(UE::Net::ENetRefHandleError::ReplicationDisabled, OwnerHandle, Parameters.ConnectionId);
+					ReplicationBridge->SendErrorWithNetRefHandle(UE::Net::ENetRefHandleError::ReplicationDisabled, OwnerHandle, Parameters.ConnectionId, {}, Context.GetErrorDiagnostic());
 
 					Context.ResetErrorContext();
 				}
diff --git a/Engine/Source/Runtime/Net/Iris/Private/Iris/ReplicationSystem/ReplicationSystem.cpp b/Engine/Source/Runtime/Net/Iris/Private/Iris/ReplicationSystem/ReplicationSystem.cpp
--- a/Engine/Source/Runtime/Net/Iris/Private/Iris/ReplicationSystem/ReplicationSystem.cpp
+++ b/Engine/Source/Runtime/Net/Iris/Private/Iris/ReplicationSystem/ReplicationSystem.cpp
@@ -2301,7 +2301,7 @@ void UReplicationSystem::ReportProtocolMismatch(uint64 NetRefHandleId, uint32 Co
 	Impl->ReplicationSystemInternal.GetReplicationBridge()->OnProtocolMismatchReported(NetRefHandle, ConnectionId);
 }
 
-void UReplicationSystem::ReportErrorWithNetRefHandle(UE::Net::ENetRefHandleError ErrorType, uint64 NetRefHandleId, uint32 ConnectionId, const TArray<uint64>& ExtraNetRefHandles)
+void UReplicationSystem::ReportErrorWithNetRefHandle(UE::Net::ENetRefHandleError ErrorType, uint64 NetRefHandleId, uint32 ConnectionId, const TArray<uint64>& ExtraNetRefHandles, const FString& DiagnosticMessage)
 {
 	using namespace UE::Net::Private;
 	const FNetRefHandle NetRefHandle = FNetRefHandleManager::MakeNetRefHandle(NetRefHandleId, GetId());
@@ -2312,7 +2312,7 @@ void UReplicationSystem::ReportErrorWithNetRefHandle(UE::Net::ENetRefHandleError
 		ExtraHandles.Emplace(FNetRefHandleManager::MakeNetRefHandle(HandleId, GetId()));
 	}
 
-	Impl->ReplicationSystemInternal.GetReplicationBridge()->OnErrorWithNetRefHandleReported(ErrorType, NetRefHandle, ConnectionId, ExtraHandles);
+	Impl->ReplicationSystemInternal.GetReplicationBridge()->OnErrorWithNetRefHandleReported(ErrorType, NetRefHandle, ConnectionId, ExtraHandles, DiagnosticMessage);
 }
 
 void UReplicationSystem::CollectNetMetrics(UE::Net::FNetMetrics& OutNetMetrics) const
diff --git a/Engine/Source/Runtime/Net/Iris/Public/Iris/ReplicationSystem/NetObjectFactory.h b/Engine/Source/Runtime/Net/Iris/Public/Iris/ReplicationSystem/NetObjectFactory.h
--- a/Engine/Source/Runtime/Net/Iris/Public/Iris/ReplicationSystem/NetObjectFactory.h
+++ b/Engine/Source/Runtime/Net/Iris/Public/Iris/ReplicationSystem/NetObjectFactory.h
@@ -218,6 +218,8 @@ public:
 		UObject* Template = nullptr;
 		/** Flags to pass back to the bridge */
 		EReplicationBridgeCreateNetRefHandleResultFlags Flags = EReplicationBridgeCreateNetRefHandleResultFlags::None;
+		/** String containing reason if instantiate request fails */
+		FString FailureDiagnosticMessage;
 	};
 
 	/** Contextual information to use during instantiation */
diff --git a/Engine/Source/Runtime/Net/Iris/Public/Iris/ReplicationSystem/ObjectReplicationBridge.h b/Engine/Source/Runtime/Net/Iris/Public/Iris/ReplicationSystem/ObjectReplicationBridge.h
--- a/Engine/Source/Runtime/Net/Iris/Public/Iris/ReplicationSystem/ObjectReplicationBridge.h
+++ b/Engine/Source/Runtime/Net/Iris/Public/Iris/ReplicationSystem/ObjectReplicationBridge.h
@@ -266,7 +266,7 @@ public:
 	IRISCORE_API UObject* GetPreRegisteredObject(FNetRefHandle Handle) const;
 
 	/** Tell the remote connection that we detected a reading error with a specific replicated object */
-	virtual void SendErrorWithNetRefHandle(UE::Net::ENetRefHandleError ErrorType, FNetRefHandle RefHandle, uint32 ConnectionId, TConstArrayView<const FNetRefHandle> ExtraNetRefHandle = {}) {}
+	virtual void SendErrorWithNetRefHandle(UE::Net::ENetRefHandleError ErrorType, FNetRefHandle RefHandle, uint32 ConnectionId, TConstArrayView<const FNetRefHandle> ExtraNetRefHandle = {}, const FString& DiagnosticMessage = FString()) {}
 
 	/**
 	 * Creates a NetRefHandle explicitly for object Instance and flags it as pre-registered in the ObjectReferenceCache.
@@ -400,14 +400,14 @@ protected:
 	IRISCORE_API virtual bool CacheNetRefHandleCreationInfo(FNetRefHandle Handle) override;
 	IRISCORE_API virtual void DetachInstance(FNetRefHandle Handle) override;
 	IRISCORE_API virtual void OnProtocolMismatchReported(FNetRefHandle RefHandle, uint32 ConnectionId) override;
-	IRISCORE_API virtual void OnErrorWithNetRefHandleReported(UE::Net::ENetRefHandleError ErrorType, FNetRefHandle RefHandle, uint32 ConnectionId, const TArray<FNetRefHandle>& ExtraHandles) override;
+	IRISCORE_API virtual void OnErrorWithNetRefHandleReported(UE::Net::ENetRefHandleError ErrorType, FNetRefHandle RefHandle, uint32 ConnectionId, const TArray<FNetRefHandle>& ExtraHandles, const FString& DiagnosticMessage) override;
 	IRISCORE_API virtual void RemoveDestructionInfosForGroup(UE::Net::FNetObjectGroupHandle GroupHandle) override;
 	IRISCORE_API virtual void RemoveDestructionInfo(UE::Net::Private::FInternalNetRefIndex InternalObjectIndex) override;
 
 private:
 
 	void OnErrorBitstreamCorrupted(FNetRefHandle RefHandle, uint32 ConnectionId);
-	void OnErrorCannotReplicateObject(FNetRefHandle RefHandle, uint32 ConnectionId);
+	void OnErrorCannotReplicateObject(FNetRefHandle RefHandle, uint32 ConnectionId, const FString& ClientDiagnosticMessage);
 	void OnErrorBlockedByAsyncLoading(FNetRefHandle RefHandle, uint32 ConnectionId, const TArray<FNetRefHandle>& MustBeMappedRefs);
 	void OnErrorBlockedByParents(FNetRefHandle RefHandle, uint32 ConnectionId, const TArray<FNetRefHandle>& BlockedParents);
 	
diff --git a/Engine/Source/Runtime/Net/Iris/Public/Iris/ReplicationSystem/ReplicationBridge.h b/Engine/Source/Runtime/Net/Iris/Public/Iris/ReplicationSystem/ReplicationBridge.h
--- a/Engine/Source/Runtime/Net/Iris/Public/Iris/ReplicationSystem/ReplicationBridge.h
+++ b/Engine/Source/Runtime/Net/Iris/Public/Iris/ReplicationSystem/ReplicationBridge.h
@@ -267,7 +267,7 @@ protected:
 	virtual void OnProtocolMismatchReported(FNetRefHandle RefHandle, uint32 ConnectionId) {}
 
 	/** Called when a remote connection has a critical error caused by a specific NetRefHandle */
-	virtual void OnErrorWithNetRefHandleReported(UE::Net::ENetRefHandleError ErrorType, FNetRefHandle RefHandle, uint32 ConnectionId, const TArray<FNetRefHandle>& ExtraNetRefHandle) {}
+	virtual void OnErrorWithNetRefHandleReported(UE::Net::ENetRefHandleError ErrorType, FNetRefHandle RefHandle, uint32 ConnectionId, const TArray<FNetRefHandle>& ExtraNetRefHandle, const FString& DiagnosticMessage) {}
 
 	/** Called from PreSeamlessTravelGarbageCollect. */
 	IRISCORE_API virtual void OnPreSeamlessTravelGarbageCollect();
diff --git a/Engine/Source/Runtime/Net/Iris/Public/Iris/ReplicationSystem/ReplicationSystem.h b/Engine/Source/Runtime/Net/Iris/Public/Iris/ReplicationSystem/ReplicationSystem.h
--- a/Engine/Source/Runtime/Net/Iris/Public/Iris/ReplicationSystem/ReplicationSystem.h
+++ b/Engine/Source/Runtime/Net/Iris/Public/Iris/ReplicationSystem/ReplicationSystem.h
@@ -823,8 +823,8 @@ public:
 	/** Called when a connection finds a protocol divergence when instantiating a replicated object. */
 	IRISCORE_API void ReportProtocolMismatch(uint64 NetRefHandleId, uint32 ConnectionId);
 
-		/** Called when a connection reports a critical error with a netrefhandle object */
-	IRISCORE_API void ReportErrorWithNetRefHandle(UE::Net::ENetRefHandleError ErrorType, uint64 NetRefHandleId, uint32 ConnectionId, const TArray<uint64>& ExtraNetRefHandles);
+	/** Called when a connection reports a critical error with a netrefhandle object */
+	IRISCORE_API void ReportErrorWithNetRefHandle(UE::Net::ENetRefHandleError ErrorType, uint64 NetRefHandleId, uint32 ConnectionId, const TArray<uint64>& ExtraNetRefHandles, const FString& DiagnosticMessage = FString());
 
 	/** Collect relevant metrics that could be sent to analytics */
 	IRISCORE_API void CollectNetMetrics(UE::Net::FNetMetrics& OutNetMetrics) const;
diff --git a/Engine/Source/Runtime/Net/Iris/Public/Iris/Serialization/NetErrorContext.h b/Engine/Source/Runtime/Net/Iris/Public/Iris/Serialization/NetErrorContext.h
--- a/Engine/Source/Runtime/Net/Iris/Public/Iris/Serialization/NetErrorContext.h
+++ b/Engine/Source/Runtime/Net/Iris/Public/Iris/Serialization/NetErrorContext.h
@@ -29,6 +29,10 @@ public:
 
 	FName GetError() const { return Error; }
 
+	void SetErrorDiagnostic(FString InDiagnostic) { ErrorDiagnostic = MoveTemp(InDiagnostic); }
+
+	const FString& GetErrorDiagnostic() const { return ErrorDiagnostic; }
+
 	void SetObjectHandle(const FNetRefHandle& InObjectHandle) { ObjectHandle = InObjectHandle; }
 
 	const FNetRefHandle& GetObjectHandle() const { return ObjectHandle; }
@@ -37,6 +41,7 @@ private:
 
 	FNetRefHandle ObjectHandle;
 	FName Error;
+	FString ErrorDiagnostic;
 };
 
 inline bool FNetErrorContext::HasError() const
diff --git a/Engine/Source/Runtime/Net/Iris/Public/Iris/Serialization/NetSerializationContext.h b/Engine/Source/Runtime/Net/Iris/Public/Iris/Serialization/NetSerializationContext.h
--- a/Engine/Source/Runtime/Net/Iris/Public/Iris/Serialization/NetSerializationContext.h
+++ b/Engine/Source/Runtime/Net/Iris/Public/Iris/Serialization/NetSerializationContext.h
@@ -51,6 +51,9 @@ public:
 	void SetError(const FName Error, bool bDoOverFlow = true);
 	FName GetError() const { return ErrorContext.GetError(); }
 
+	void SetErrorDiagnostic(FString InDiagnostic) { ErrorContext.SetErrorDiagnostic(MoveTemp(InDiagnostic)); }
+	const FString& GetErrorDiagnostic() const { return ErrorContext.GetErrorDiagnostic(); }
+
 	/** Store extra information regarding the object that triggered an error. */
 	void SetErrorHandleContext(const FNetRefHandle& HandleContext);
 	const FNetRefHandle& GetErrorHandleContext() const { return ErrorContext.GetObjectHandle(); }
