Add a couple more curve table helper functions to DataTable BPFL

#Synced-CL 51567995

diff --git a/Engine/Source/Runtime/Engine/Classes/Kismet/DataTableFunctionLibrary.h b/Engine/Source/Runtime/Engine/Classes/Kismet/DataTableFunctionLibrary.h
--- a/Engine/Source/Runtime/Engine/Classes/Kismet/DataTableFunctionLibrary.h
+++ b/Engine/Source/Runtime/Engine/Classes/Kismet/DataTableFunctionLibrary.h
@@ -59,6 +59,14 @@ class UDataTableFunctionLibrary : public UBlueprintFunctionLibrary
 	UFUNCTION(BlueprintCallable, Category = "DataTable")
 	static ENGINE_API bool AddSimpleCurveToTable(UCurveTable* CurveTable, FName RowName);
 
+	/**
+	 * Returns the names of the rows in the curve table.
+	 * @param CurveTable The Curve Table to query
+	 * @return A list of names of rows.
+	 */
+	UFUNCTION(BlueprintCallable, Category = "DataTable")
+	static ENGINE_API TArray<FName> GetCurveTableRowNames(UCurveTable* CurveTable);
+
 	/**
 	 * Adds a simple curve to a curve table.
 	 * @param CurveTable The Curve Table to modify
@@ -78,7 +86,16 @@ class UDataTableFunctionLibrary : public UBlueprintFunctionLibrary
 	 * @return True if the operation was successful.
 	 */
 	UFUNCTION(BlueprintCallable, Category = "DataTable")
-	static ENGINE_API bool AddKey(UCurveTable* CurveTable, FName RowName, float InTime, float InValue);
+	static ENGINE_API bool AddCurveTableKey(UCurveTable* CurveTable, FName RowName, float InTime, float InValue);
+
+	/**
+	 * Returns a list of keys in a row of a SimpleCurve table.
+	 * @param CurveTable The Curve Table to query.
+	 * @param RowName The name of the row to query.
+	 * @return The list of keys in the curve, or empty if the curve doesn't exist.
+	 */
+	UFUNCTION(BlueprintCallable, Category = "DataTable")
+	static ENGINE_API TArray<FSimpleCurveKey> GetCurveTableKeys(UCurveTable* CurveTable, FName RowName);
 
 	/**
 	 * Removes a curve from a curve table.
diff --git a/Engine/Source/Runtime/Engine/Private/DataTableFunctionLibrary.cpp b/Engine/Source/Runtime/Engine/Private/DataTableFunctionLibrary.cpp
--- a/Engine/Source/Runtime/Engine/Private/DataTableFunctionLibrary.cpp
+++ b/Engine/Source/Runtime/Engine/Private/DataTableFunctionLibrary.cpp
@@ -64,6 +64,16 @@ bool UDataTableFunctionLibrary::AddSimpleCurveToTable(UCurveTable* CurveTable, F
 	return true;
 }
 
+TArray<FName> UDataTableFunctionLibrary::GetCurveTableRowNames(UCurveTable* CurveTable)
+{
+	TArray<FName> RowNames;
+	if (CurveTable)
+	{
+		CurveTable->GetRowMap().GetKeys(RowNames);
+	}
+	return RowNames;
+}
+
 bool UDataTableFunctionLibrary::SetCurveTableRowDefault(UCurveTable* CurveTable, FName RowName, float DefaultValue)
 {
 	if (!CurveTable)
@@ -80,7 +90,7 @@ bool UDataTableFunctionLibrary::SetCurveTableRowDefault(UCurveTable* CurveTable,
 	return true;
 }
 
-bool UDataTableFunctionLibrary::AddKey(UCurveTable* CurveTable, FName RowName, float InTime, float InValue)
+bool UDataTableFunctionLibrary::AddCurveTableKey(UCurveTable* CurveTable, FName RowName, float InTime, float InValue)
 {
 	if (!CurveTable)
 	{
@@ -96,6 +106,21 @@ bool UDataTableFunctionLibrary::AddKey(UCurveTable* CurveTable, FName RowName, f
 	return true;
 }
 
+TArray<FSimpleCurveKey> UDataTableFunctionLibrary::GetCurveTableKeys(UCurveTable* CurveTable, FName RowName)
+{
+	TArray<FSimpleCurveKey> Keys;
+	if (CurveTable && CurveTable->GetCurveTableMode() == ECurveTableMode::SimpleCurves)
+	{
+		FString Context;
+		FSimpleCurve* Curve = CurveTable->FindSimpleCurve(RowName, Context);
+		if (Curve)
+		{
+			Keys = Curve->Keys;
+		}
+	}
+	return Keys;
+}
+
 bool UDataTableFunctionLibrary::RemoveCurveTableRow(UCurveTable* CurveTable, FName RowName)
 {
 	if (!CurveTable)
