Optimize away an include

#rnx

#Synced-CL 51585574

diff --git a/Engine/Source/Runtime/Engine/Private/Net/Iris/ReplicationSystem/NetRootObjectFactory.cpp b/Engine/Source/Runtime/Engine/Private/Net/Iris/ReplicationSystem/NetRootObjectFactory.cpp
--- a/Engine/Source/Runtime/Engine/Private/Net/Iris/ReplicationSystem/NetRootObjectFactory.cpp
+++ b/Engine/Source/Runtime/Engine/Private/Net/Iris/ReplicationSystem/NetRootObjectFactory.cpp
@@ -14,8 +14,6 @@
 #include "Iris/Serialization/NetSerializationContext.h"
 #include "Iris/Serialization/ObjectNetSerializer.h"
 
-#include "UObject/Package.h"
-
 #include "Net/DataBunch.h"
 
 #include UE_INLINE_GENERATED_CPP_BY_NAME(NetRootObjectFactory)
@@ -107,7 +105,7 @@ UNetObjectFactory::FInstantiateResult UNetRootObjectFactory::InstantiateReplicat
 		return InstantiateResult;
 	}
 
-	UObject* DynamicRootObject = NewObject<UObject>(GetTransientPackage(), Template->GetClass(), NAME_None, RF_NoFlags, Template);
+	UObject* DynamicRootObject = NewObject<UObject>(GetTransientPackageAsObject(), Template->GetClass(), NAME_None, RF_NoFlags, Template);
 
 	FInstantiateResult Result
 	{ 
