- Make sure boost to highest priority works for other types using FPrimitiveSceneProxyDesc instead of just primitive component
- Process boost on AddToScene when proxy is always visible because it doesn't go through relevance computation (nanite meshes)

#[robomerge][partner] megalights

#Synced-CL 51609817

diff --git a/Engine/Plugins/Experimental/FastGeoStreaming/Source/FastGeoStreaming/Private/FastGeoPrimitiveComponent.cpp b/Engine/Plugins/Experimental/FastGeoStreaming/Source/FastGeoStreaming/Private/FastGeoPrimitiveComponent.cpp
--- a/Engine/Plugins/Experimental/FastGeoStreaming/Source/FastGeoStreaming/Private/FastGeoPrimitiveComponent.cpp
+++ b/Engine/Plugins/Experimental/FastGeoStreaming/Source/FastGeoStreaming/Private/FastGeoPrimitiveComponent.cpp
@@ -275,7 +275,11 @@ void FFastGeoPrimitiveComponent::InitializeSceneProxyDescDynamicProperties()
 	SceneProxyDesc.FeatureLevel = SceneProxyDesc.Scene->GetFeatureLevel();
 	TArray<URuntimeVirtualTexture*> const& VirtualTextures = GetRuntimeVirtualTextures();
 	SceneProxyDesc.RuntimeVirtualTextures = MakeArrayView(const_cast<URuntimeVirtualTexture**>(VirtualTextures.GetData()), VirtualTextures.Num());
-	SceneProxyDesc.bUsePSOPrecacheFallbackMaterial = UsePSOPrecacheFallbackMaterial();
+
+#if !WITH_EDITOR && UE_WITH_PSO_PRECACHING
+	SceneProxyDesc.SetPSOPrecacheComponentData(PSOPrecacheComponentData);
+#endif // !WITH_EDITOR && UE_WITH_PSO_PRECACHING
+
 #if MESH_DRAW_COMMAND_STATS
 	static FName NAME_FastGeoPrimitiveComponent(TEXT("FastGeoPrimitiveComponent"));
 	SceneProxyDesc.MeshDrawCommandStatsCategory = NAME_FastGeoPrimitiveComponent;
diff --git a/Engine/Source/Runtime/Engine/Private/PrimitiveSceneProxy.cpp b/Engine/Source/Runtime/Engine/Private/PrimitiveSceneProxy.cpp
--- a/Engine/Source/Runtime/Engine/Private/PrimitiveSceneProxy.cpp
+++ b/Engine/Source/Runtime/Engine/Private/PrimitiveSceneProxy.cpp
@@ -612,15 +612,7 @@ PRAGMA_ENABLE_DEPRECATION_WARNINGS
 	}
 
 #if UE_WITH_PSO_PRECACHING
-	if (UPrimitiveComponent* PrimitiveComponent = Cast<UPrimitiveComponent>(InProxyDesc.Component))
-	{
-		FPSOPrecacheComponentData& PSOPrecacheComponentData = PrimitiveComponent->GetPSOPrecacheComponentData();
-		check(!PSOPrecacheComponentData.UsePSOPrecacheFallbackMaterial() || !PSOPrecacheComponentData.GetMaterialPSOPrecacheRequestIDs().IsEmpty());
-		if (ShouldBoostPSOPrecachePriorityOnDraw() && PSOPrecacheComponentData.IsPSOPrecaching())
-		{
-			SetPSORequestsToBoostOnDraw(PSOPrecacheComponentData.GetMaterialPSOPrecacheRequestIDs());
-		}
-	}
+	PSOPrecacheRequestsToBoostOnDraw = InProxyDesc.MaterialPSOPrecacheRequestIDsToBoostOnDraw;
 #endif
 
 	// FirstPerson primitives won't properly work with most shadow solutions, so just disable them here to give a smoother user experience out of the box.
@@ -1852,11 +1844,6 @@ void FPrimitiveSceneProxy::BoostPrecachedPSORequestsOnDraw()
 		BoostPSOPriorityOnDraw(LocalPSOPrecacheRequestsToBoostOnDraw);
 	}
 }
-
-void FPrimitiveSceneProxy::SetPSORequestsToBoostOnDraw(const TArray<FMaterialPSOPrecacheRequestID>& PSORequestsToBoostOnDrawIN)
-{
-	PSOPrecacheRequestsToBoostOnDraw = PSORequestsToBoostOnDrawIN;
-}
 #endif // UE_WITH_PSO_PRECACHING
 
 #if RHI_RAYTRACING
diff --git a/Engine/Source/Runtime/Engine/Private/PrimitiveSceneProxyDesc.cpp b/Engine/Source/Runtime/Engine/Private/PrimitiveSceneProxyDesc.cpp
--- a/Engine/Source/Runtime/Engine/Private/PrimitiveSceneProxyDesc.cpp
+++ b/Engine/Source/Runtime/Engine/Private/PrimitiveSceneProxyDesc.cpp
@@ -216,8 +216,7 @@ void FPrimitiveSceneProxyDesc::InitializeFromPrimitiveComponent(const UPrimitive
 	HiddenEditorViews = InComponent->GetHiddenEditorViews();
 	OverlayColor = InComponent->OverlayColor;
 #endif
-	bUsePSOPrecacheFallbackMaterial = InComponent->GetPSOPrecacheComponentData().UsePSOPrecacheFallbackMaterial();
-
+	
 	AdditionalStatObjectPtr = InComponent->AdditionalStatObject();
 	StatId = AdditionalStatObjectPtr ? AdditionalStatObjectPtr->GetStatID(true) : InComponent->GetStatID(true);
 
@@ -226,6 +225,8 @@ void FPrimitiveSceneProxyDesc::InitializeFromPrimitiveComponent(const UPrimitive
 	VirtualTextureRenderPassType = InComponent->GetVirtualTextureRenderPassType();
 	VirtualTextureMainPassMaxDrawDistance = InComponent->GetVirtualTextureMainPassMaxDrawDistance();
 
+	SetPSOPrecacheComponentData(InComponent->GetPSOPrecacheComponentData());
+
 #if MESH_DRAW_COMMAND_STATS
 	MeshDrawCommandStatsCategory = InComponent->GetMeshDrawCommandStatsCategory();
 #endif
@@ -239,3 +240,15 @@ void FPrimitiveSceneProxyDesc::GetUsedMaterials(TArray<UMaterialInterface*>& Out
 
 	return AsComponent->GetUsedMaterials(OutMaterials, bGetDebugMaterials);
 }
+
+void FPrimitiveSceneProxyDesc::SetPSOPrecacheComponentData(const FPSOPrecacheComponentData& PSOPrecacheComponentData)
+{
+#if UE_WITH_PSO_PRECACHING	
+	bUsePSOPrecacheFallbackMaterial = PSOPrecacheComponentData.UsePSOPrecacheFallbackMaterial();
+	check(!PSOPrecacheComponentData.UsePSOPrecacheFallbackMaterial() || !PSOPrecacheComponentData.GetMaterialPSOPrecacheRequestIDs().IsEmpty());
+	if (ShouldBoostPSOPrecachePriorityOnDraw() && PSOPrecacheComponentData.IsPSOPrecaching())
+	{
+		MaterialPSOPrecacheRequestIDsToBoostOnDraw = PSOPrecacheComponentData.GetMaterialPSOPrecacheRequestIDs();
+	}
+#endif // UE_WITH_PSO_PRECACHING
+}
\ No newline at end of file
diff --git a/Engine/Source/Runtime/Engine/Public/PrimitiveSceneProxy.h b/Engine/Source/Runtime/Engine/Public/PrimitiveSceneProxy.h
--- a/Engine/Source/Runtime/Engine/Public/PrimitiveSceneProxy.h
+++ b/Engine/Source/Runtime/Engine/Public/PrimitiveSceneProxy.h
@@ -1235,7 +1235,6 @@ PRAGMA_ENABLE_DEPRECATION_WARNINGS
 
 #if UE_WITH_PSO_PRECACHING
 	ENGINE_API void BoostPrecachedPSORequestsOnDraw();
-	ENGINE_API void SetPSORequestsToBoostOnDraw(const TArray<FMaterialPSOPrecacheRequestID>& PSORequestsToBoostOnDrawIN);
 #endif
 
 protected:
diff --git a/Engine/Source/Runtime/Engine/Public/PrimitiveSceneProxyDesc.h b/Engine/Source/Runtime/Engine/Public/PrimitiveSceneProxyDesc.h
--- a/Engine/Source/Runtime/Engine/Public/PrimitiveSceneProxyDesc.h
+++ b/Engine/Source/Runtime/Engine/Public/PrimitiveSceneProxyDesc.h
@@ -132,6 +132,10 @@ struct FPrimitiveSceneProxyDesc
 
 	ERHIFeatureLevel::Type FeatureLevel;
 
+#if UE_WITH_PSO_PRECACHING
+	TArray<FMaterialPSOPrecacheRequestID> MaterialPSOPrecacheRequestIDsToBoostOnDraw;
+#endif // UE_WITH_PSO_PRECACHING
+
 #if MESH_DRAW_COMMAND_STATS
 	FName MeshDrawCommandStatsCategory;
 #endif
@@ -324,6 +328,8 @@ struct FPrimitiveSceneProxyDesc
 
 	ENGINE_API virtual void GetUsedMaterials(TArray<UMaterialInterface*>& OutMaterials, bool bGetDebugMaterials = false) const;
 
+	ENGINE_API void SetPSOPrecacheComponentData(const FPSOPrecacheComponentData& PrecacheComponentData);
+
 	UE_DEPRECATED(5.6, "Use FPrimitiveSceneProxyDesc::IsFirstPersonRelevant() instead.")
 	bool IsStaticMeshFirstPersonRelevant() const
 	{
diff --git a/Engine/Source/Runtime/Renderer/Private/PrimitiveSceneInfo.cpp b/Engine/Source/Runtime/Renderer/Private/PrimitiveSceneInfo.cpp
--- a/Engine/Source/Runtime/Renderer/Private/PrimitiveSceneInfo.cpp
+++ b/Engine/Source/Runtime/Renderer/Private/PrimitiveSceneInfo.cpp
@@ -1057,10 +1057,7 @@ void FPrimitiveSceneInfo::RemoveCachedNaniteMaterialBins()
 	}
 
 	// Primitive doesn't care about pending static mesh updates anymore
-	if (auto MaterialsExtension = Scene->GetExtensionPtr<Nanite::FMaterialsSceneExtension>())
-	{
-		MaterialsExtension->ClearPendingRequestStaticMeshUpdate(this);
-	}
+	Scene->ClearPendingRequestStaticMeshUpdate(this);
 
 	bNaniteRasterBinsRenderCustomDepth = false;
 }
@@ -1967,6 +1964,20 @@ void FPrimitiveSceneInfo::AddToScene(FScene* Scene, TArrayView<FPrimitiveSceneIn
 			}
 		}
 	}
+
+#if UE_WITH_PSO_PRECACHING
+	{
+		SCOPED_NAMED_EVENT(FPrimitiveSceneInfo_AddToScene_BoostAlwaysVisibleProxyPSOs, FColor::Silver);
+		for (FPrimitiveSceneInfo* SceneInfo : SceneInfos)
+		{
+			// If proxy is always visible then boost right now, because won't go over relevance path and won't be boosted otherwise
+			if (SceneInfo->Proxy->IsAlwaysVisible())
+			{
+				SceneInfo->Proxy->BoostPrecachedPSORequestsOnDraw();
+			}
+		}
+	}
+#endif // UE_WITH_PSO_PRECACHING
 }
 
 void FPrimitiveSceneInfo::RemoveStaticMeshes()
