Make nanite lumen card and raster not required during PSO precache collection

#Synced-CL 51609817

diff --git a/Engine/Source/Runtime/Engine/Public/PSOPrecacheSettings.h b/Engine/Source/Runtime/Engine/Public/PSOPrecacheSettings.h
--- a/Engine/Source/Runtime/Engine/Public/PSOPrecacheSettings.h
+++ b/Engine/Source/Runtime/Engine/Public/PSOPrecacheSettings.h
@@ -130,11 +130,11 @@ struct FPSOPrecacheSettings
 
 	/** Make nanite raster programmable raster PSOs required before the component can be rendered with desired material or fallback to fixed function untill programmable raster PSOs are ready **/
 	UPROPERTY()
-	bool bNaniteProgrammableRasterRequired = true;
+	bool bNaniteProgrammableRasterRequired = false;
 	
 	/** Make nanite lumen card PSOs required before the component can be rendered with desired material or otherwise the mesh is not rendered to the LumenScene yet **/
 	UPROPERTY()
-	bool bNaniteLumenCardRequired = true;
+	bool bNaniteLumenCardRequired = false;
 };
 
 
