Substrate - NMT - Added support for Subsurface and Specular Profiles.
Also added support for Toon BSDF + Toon Profile.

#rnx
#Synced-CL 51624139

diff --git a/Engine/Source/Runtime/Engine/Private/Materials/MaterialExpressionsToMIR.cpp b/Engine/Source/Runtime/Engine/Private/Materials/MaterialExpressionsToMIR.cpp
--- a/Engine/Source/Runtime/Engine/Private/Materials/MaterialExpressionsToMIR.cpp
+++ b/Engine/Source/Runtime/Engine/Private/Materials/MaterialExpressionsToMIR.cpp
@@ -271,6 +271,8 @@
 #include "ColorManagement/ColorSpace.h"
 #include "SubstrateTranslatorCommon.h"
 #include "Engine/SubsurfaceProfile.h"
+#include "Engine/SpecularProfile.h"
+#include "Engine/ToonProfile.h"
 
 #include "PostProcess/PostProcessMaterialInputs.h"
 
@@ -5932,6 +5934,20 @@ static FValueRef FinalizeBSDFAndPromoteToOperatorIfNeeded(MIR::FEmitter& Em, FSu
 	return BSDF;
 }
 
+static FValueRef CreateSubsurfaceProfileParameter(MIR::FEmitter& Em, USubsurfaceProfile* InProfile)
+{
+	check(InProfile);
+	const FName SubsurfaceProfileParameterName = SubsurfaceProfile::CreateSubsurfaceProfileParameterName(InProfile);
+	const FValueRef SSSProfileCodeArg = Em.CastToFloat(Em.NamedPrimitiveUniform(SubsurfaceProfileParameterName, Em.ConstantFloat(0.0f)), 1);
+	return SSSProfileCodeArg;
+}
+
+static FValueRef CreateDefaultSubsurfaceProfileParameter(MIR::FEmitter& Em)
+{
+	const FValueRef SSSProfileCodeArg = Em.CastToFloat(Em.NamedPrimitiveUniform(SubsurfaceProfile::GetSubsurfaceProfileParameterName(), Em.ConstantFloat(0.0f)), 1);
+	return SSSProfileCodeArg;
+}
+
 void UMaterialExpressionSubstrateSlabBSDF::Build(MIR::FEmitter& Em)
 {
 	// SUBSTRATE_TODO make this similar to UMaterialExpressionSubstrateSlabBSDF::Compile
@@ -5942,61 +5958,26 @@ void UMaterialExpressionSubstrateSlabBSDF::Build(MIR::FEmitter& Em)
 	FGuid PathUniqueId = SubstrateTranslatorData->SubstrateTreeStackGetPathUniqueId();
 	FSubstrateOperator& SubstrateOperator = SubstrateTranslatorData->SubstrateCompilationGetOperator(PathUniqueId);
 
-#if 0
-	int32 SSSProfileCodeChunk = INDEX_NONE;
-	FName SubsurfaceProfileParameterName;
+	FValueRef SubsurfaceProfileArg;
 	if (SubstrateOperator.Has(ESubstrateBsdfFeature::SSS) && SubsurfaceProfile)
 	{
-		SubsurfaceProfileParameterName = SubsurfaceProfile::CreateSubsurfaceProfileParameterName(SubsurfaceProfile);
+		SubsurfaceProfileArg = CreateSubsurfaceProfileParameter(Em, SubsurfaceProfile);
 	}
 	else
 	{
-		SubsurfaceProfileParameterName = SubsurfaceProfile::GetSubsurfaceProfileParameterName();
+		SubsurfaceProfileArg = CreateDefaultSubsurfaceProfileParameter(Em);
 	}
-	Em.PreshaderParameter(MIR::FType::MakeScalar(MIR::EScalarKind::Float), EPreshader2Opcode::Abs, MIR::FValueRef::, )
-
-	//	FPreshaderParameter
-	//	PRESHADER2_OP(SparseVolumeTextureUniform)
-
-		// FValueRef FEmitter::PreshaderParameter(FType Type, EPreshader2Opcode Opcode, FValueRef SourceParameter, FPreshaderParameterPayload Payload)
 
-		/*
-		const UE::Shader::EValueType UniformType = URuntimeVirtualTexture::GetUniformParameterType(UniformIndex);
-		Uniforms[UniformIndex] = Em.PreshaderParameter(
-			MIR::FType::FromShaderType(UniformType), EPreshader2Opcode::RuntimeVirtualTextureUniform,
-			TextureObjects[0], MIR::FPreshaderParameterPayload{ .UniformIndex = UniformIndex });
-		*/
-	// SUBSTRATE_TODO implement
-	//SubstrateTranslatorData->SubstrateThicknessStackGetThicknessCode()
-	//SubstrateTranslatorData->GetSubstrateSharedLocalBasisIndexMacro(NewRegisteredSharedLocalBasis)
-
-/*
-	int32 SSSProfileCodeChunk = INDEX_NONE;
-	if (SubstrateOperator.Has(ESubstrateBsdfFeature::SSS) && SubsurfaceProfile)
+	FValueRef SpecularProfileArg = Em.ConstantFloat(0.0f);
+	if (SubstrateOperator.Has(ESubstrateBsdfFeature::SpecularProfile))
 	{
-		SSSProfileCodeChunk = CreateSubsurfaceProfileParameter(Compiler, SubsurfaceProfile);
+		const FName SpecularProfileParameterName = SpecularProfile::CreateSpecularProfileParameterName(SpecularProfile);
+		SpecularProfileArg = Em.CastToFloat(Em.NamedPrimitiveUniform(SpecularProfileParameterName, Em.ConstantFloat(0.0f)), 1);
 	}
 	else
 	{
-		SSSProfileCodeChunk = CreateDefaultSubsurfaceProfileParameter(Compiler);
-	}
-*/
-#else
-	FValueRef SubsurfaceProfileArg = Em.ConstantFloat(0.0f); // SUBSTRATE_TODO
-#endif
-
-#if 0
-/*
-	int32 SpecularProfileCodeChunk = INDEX_NONE;
-	if (SubstrateOperator.Has(ESubstrateBsdfFeature::SpecularProfile))
-	{
-		const FName SpecularProfileParameterName = SpecularProfile::CreateSpecularProfileParameterName(SpecularProfile);
-		SpecularProfileCodeChunk = Compiler->ForceCast(Compiler->ScalarParameter(SpecularProfileParameterName, 1.0f), MCT_Float1);
+		SpecularProfileArg = Em.ConstantFloat(0.0f);
 	}
-*/
-#else
-	FValueRef SpecularProfileArg = Em.ConstantFloat(0.0f); // SUBSTRATE_TODO
-#endif
 
 	const float DefaultSpecular = 0.5f;
 	const float DefaultF0 = DielectricSpecularToF0(DefaultSpecular);
@@ -6089,7 +6070,7 @@ void UMaterialExpressionSubstrateSlabBSDF::Build(MIR::FEmitter& Em)
 	}
 	if (!SubstrateOperator.Has(ESubstrateBsdfFeature::SpecularProfile))
 	{
-		SpecularProfileArg = Em.ConstantFloat(float(INDEX_NONE));
+		SpecularProfileArg = Em.ConstantFloat(0.0f);
 	}
 
 	const bool bIsAtTheBottomOfTopology = SubstrateOperator.bIsBottom > 0;
@@ -6149,16 +6130,15 @@ MIR::FValueRef UMaterialExpressionSubstrateShadingModels::BuildCommon(MIR::FEmit
 	FGuid PathUniqueId = SubstrateTranslatorData->SubstrateTreeStackGetPathUniqueId();
 	FSubstrateOperator& SubstrateOperator = SubstrateTranslatorData->SubstrateCompilationGetOperator(PathUniqueId);
 
+	FValueRef SubsurfaceProfileIdArg;
 	if (bHasSSS && SSSProfile != nullptr)
 	{
-		check(false);	// SUBSTRATE TODO implemnent
-		//	SSSProfileCodeChunk = CreateSubsurfaceProfileParameter(Compiler, SSSProfile);
+		SubsurfaceProfileIdArg = CreateSubsurfaceProfileParameter(Em, SSSProfile);
+	}
+	else
+	{
+		SubsurfaceProfileIdArg = CreateDefaultSubsurfaceProfileParameter(Em);
 	}
-	//else
-	//{
-	//	SSSProfileCodeChunk = CreateDefaultSubsurfaceProfileParameter(Compiler);SUBSTRATE_TODO
-	//}
-	FValueRef SubsurfaceProfileIdArg = Em.ConstantFloat(0.0f);
 
 
 	// We also cannot ignore the tangent when using the default Tangent because GetTangentBasis
@@ -6331,7 +6311,7 @@ void UMaterialExpressionSubstrateSimpleClearCoatBSDF::Build(MIR::FEmitter& Em)
 	FValueRef FuzzRoughnessArg = Em.ConstantFloat(0.5f);
 	FValueRef GlintValueArg = Em.ConstantFloat(1.0f);
 	FValueRef GlintUVArg = Em.ConstantFloat2(FVector2f(0.0f, 0.0f));
-	FValueRef SpecularProfileArg = Em.ConstantFloat(float(INDEX_NONE));
+	FValueRef SpecularProfileArg = Em.ConstantFloat(0.0);
 	FValueRef SubSurfaceTypeArg = Em.ConstantFloat(float(MSS_None));
 
 	const bool bIsAtTheBottomOfTopology = SubstrateOperator.bIsBottom > 0;
@@ -6578,6 +6558,55 @@ void UMaterialExpressionSubstrateUnlitBSDF::Build(MIR::FEmitter& Em)
 	FinalizeBSDFAndPromoteToOperatorIfNeeded(Em, SubstrateOperator, UnlitBSDF);
 }
 
+void UMaterialExpressionSubstrateToonBSDF::Build(MIR::FEmitter& Em)
+{
+	FSubstrateTranslatorData* SubstrateTranslatorData = Em.GetSubstrateTranslatorData();
+
+	FGuid PathUniqueId = SubstrateTranslatorData->SubstrateTreeStackGetPathUniqueId();
+	FSubstrateOperator& SubstrateOperator = SubstrateTranslatorData->SubstrateCompilationGetOperator(PathUniqueId);
+
+	FValueRef NormalArg = CompileWithDefaultNormalWS(Em, Material, Normal);
+	const FSubstrateRegisteredSharedLocalBasis NewRegisteredSharedLocalBasis = MIRSubstrateCompilationInfoCreateSharedLocalBasis2(SubstrateTranslatorData, &NormalArg, nullptr);
+	SubstrateOperator.BSDFRegisteredSharedLocalBasis = NewRegisteredSharedLocalBasis;
+	FValueRef SharedLocalBasisIndexMacroArg = EmitSubstrateSharedLocalBasis(Em, NewRegisteredSharedLocalBasis, SubstrateTranslatorData->CurrentSubstrateCompilationContext);
+
+	FValueRef ToonProfileIdArg;
+	if (ToonProfile != nullptr)
+	{
+ 		const FName ToonProfileParameterName = ToonProfile::CreateToonProfileParameterName(ToonProfile);
+		ToonProfileIdArg = Em.CastToFloat(Em.NamedPrimitiveUniform(ToonProfileParameterName, Em.ConstantFloat(0.0f)), 1);
+	}
+	else
+	{
+		ToonProfileIdArg = Em.ConstantFloat(0.0f);
+	}
+
+	FValueRef BaseColorArg = CompileWithDefaultFloat3(Em, BaseColor, 0.18f, 0.18f, 0.18f);
+	FValueRef MetallicArg = CompileWithDefaultFloat1(Em, Metallic, 0.0f);
+	FValueRef SpecularArg = CompileWithDefaultFloat1(Em, Specular, 0.5f);
+	FValueRef RoughnessArg = CompileWithDefaultFloat1(Em, Roughness, 0.5f);
+	FValueRef EmissiveColorArg = CompileWithDefaultFloat3(Em, EmissiveColor, 0.0f, 0.0f, 0.0f);
+	FValueRef PatternUVsArg = CompileWithDefaultFloat2(Em, PatternUVs, 0.0f, 0.0f);
+
+	const bool bIsAtTheBottomOfTopology = SubstrateOperator.bIsBottom > 0;
+	MIR::FSubstrateToonDesc SubstrateToonDesc = {
+			NormalArg,
+			ToonProfileIdArg,
+			BaseColorArg,
+			MetallicArg,
+			SpecularArg,
+			RoughnessArg,
+			EmissiveColorArg,
+			PatternUVsArg,
+			bIsAtTheBottomOfTopology ? Em.ConstantTrue() : Em.ConstantFalse(),
+			SharedLocalBasisIndexMacroArg };
+
+	FValueRef ToonBSDF = Em.SubstrateToon(SubstrateToonDesc);
+
+	FinalizeBSDFAndPromoteToOperatorIfNeeded(Em, SubstrateOperator, ToonBSDF);
+
+}
+
 void UMaterialExpressionSubstrateLightFunction::Build(MIR::FEmitter& Em)
 {
 	FSubstrateTranslatorData* SubstrateTranslatorData = Em.GetSubstrateTranslatorData();
diff --git a/Engine/Source/Runtime/Engine/Private/Materials/MaterialIR.cpp b/Engine/Source/Runtime/Engine/Private/Materials/MaterialIR.cpp
--- a/Engine/Source/Runtime/Engine/Private/Materials/MaterialIR.cpp
+++ b/Engine/Source/Runtime/Engine/Private/Materials/MaterialIR.cpp
@@ -42,6 +42,7 @@ const TCHAR* LexToString(EValueKind Kind)
 		case VK_SubstrateDefaultSlab: return TEXT("SubstrateDefaultSlab");
 		case VK_SubstrateSlab: return TEXT("SubstrateSlab");
 		case VK_SubstrateShadingModels: return TEXT("SubstrateShadingModel");
+		case VK_SubstrateToon: return TEXT("SubstrateToon");
 		case VK_SubstrateHorizontalMixing: return TEXT("SubstrateHorizontalMixing");
 		case VK_SubstrateVerticalLayering: return TEXT("SubstrateVerticalLayering");
 		case VK_SubstrateCoverageWeight: return TEXT("SubstrateCoverageWeight");
@@ -119,6 +120,7 @@ uint32 FValue::GetSizeInBytes() const
 		case VK_SubstrateDefaultSlab: return sizeof(FSubstrateDefaultSlab);
 		case VK_SubstrateSlab: return sizeof(FSubstrateSlab);
 		case VK_SubstrateShadingModels: return sizeof(FSubstrateShadingModels);
+		case VK_SubstrateToon: return sizeof(FSubstrateToon);
 		case VK_SubstrateHorizontalMixing: return sizeof(FSubstrateHorizontalMixing);
 		case VK_SubstrateVerticalLayering: return sizeof(FSubstrateVerticalLayering);
 		case VK_SubstrateCoverageWeight: return sizeof(FSubstrateCoverageWeight);
@@ -236,6 +238,12 @@ TConstArrayView<FValue*> FValue::GetUses() const
 			return { &This->BaseColor, FSubstrateShadingModels::NumStaticUses };
 		}
 
+		case VK_SubstrateToon:
+		{
+			auto This = static_cast<const FSubstrateToon*>(this);
+			return { &This->Normal, FSubstrateToon::NumStaticUses };
+		}
+
 		case VK_SubstrateHorizontalMixing:
 		{
 			auto This = static_cast<const FSubstrateHorizontalMixing*>(this);
diff --git a/Engine/Source/Runtime/Engine/Private/Materials/MaterialIREmitter.cpp b/Engine/Source/Runtime/Engine/Private/Materials/MaterialIREmitter.cpp
--- a/Engine/Source/Runtime/Engine/Private/Materials/MaterialIREmitter.cpp
+++ b/Engine/Source/Runtime/Engine/Private/Materials/MaterialIREmitter.cpp
@@ -4367,6 +4367,25 @@ FValueRef FEmitter::SubstrateSlab(const FSubstrateSlabDesc& Desc)
 	return EmitPrototype(*this, Prototype);
 }
 
+FValueRef FEmitter::SubstrateToon(const FSubstrateToonDesc& Desc)
+{
+	FSubstrateToon Prototype = MakePrototype<FSubstrateToon>(FType::MakeSubstrateData());
+
+	Prototype.Normal = Desc.Normal;
+	Prototype.ToonProfileId = Desc.ToonProfileId;
+	Prototype.BaseColor = Desc.BaseColor;
+	Prototype.Metallic = Desc.Metallic;
+	Prototype.Specular = Desc.Specular;
+	Prototype.Roughness = Desc.Roughness;
+	Prototype.EmissiveColor = Desc.EmissiveColor;
+	Prototype.PatternUVs = Desc.PatternUVs;
+	Prototype.IsAtBottom = Desc.IsAtBottom;
+	Prototype.LocalBasisIndex = Desc.LocalBasisIndex;
+
+	Prototype.bFullySimplifiedInstruction = SubstrateTranslatorData->CurrentSubstrateCompilationContext == Substrate::ESubstrateCompilationContext::SCC_FullySimplified;
+	return EmitPrototype(*this, Prototype);
+}
+
 FValueRef FEmitter::SubstrateHorizontalMixing(FValueRef Background, FValueRef Foreground, FValueRef Mix, FValueRef OperatorIndex, FValueRef MaxDistanceFromLeaves)
 {
 	FSubstrateHorizontalMixing Prototype = MakePrototype<FSubstrateHorizontalMixing>(FType::MakeSubstrateData());
diff --git a/Engine/Source/Runtime/Engine/Private/Materials/MaterialIRToHLSLTranslator.cpp b/Engine/Source/Runtime/Engine/Private/Materials/MaterialIRToHLSLTranslator.cpp
--- a/Engine/Source/Runtime/Engine/Private/Materials/MaterialIRToHLSLTranslator.cpp
+++ b/Engine/Source/Runtime/Engine/Private/Materials/MaterialIRToHLSLTranslator.cpp
@@ -463,6 +463,7 @@ struct FEntryPointCodeGen
 			case MIR::VK_SubstrateDefaultSlab: LowerSubstrateDefaultSlab(static_cast<const MIR::FSubstrateDefaultSlab*>(Instr)); break;
 			case MIR::VK_SubstrateSlab: LowerSubstrateSlab(static_cast<const MIR::FSubstrateSlab*>(Instr)); break;
 			case MIR::VK_SubstrateShadingModels: LowerSubstrateShadingModel(static_cast<const MIR::FSubstrateShadingModels*>(Instr)); break;
+			case MIR::VK_SubstrateToon: LowerSubstrateToon(static_cast<const MIR::FSubstrateToon*>(Instr)); break;
 			case MIR::VK_SubstrateHorizontalMixing: LowerSubstrateHorizontalMixing(static_cast<const MIR::FSubstrateHorizontalMixing*>(Instr)); break;
 			case MIR::VK_SubstrateVerticalLayering: LowerSubstrateVerticalLayering(static_cast<const MIR::FSubstrateVerticalLayering*>(Instr)); break;
 			case MIR::VK_SubstrateCoverageWeight: LowerSubstrateCoverageWeight(static_cast<const MIR::FSubstrateCoverageWeight*>(Instr)); break;
@@ -1831,7 +1832,7 @@ struct FEntryPointCodeGen
 		Printer << LowerValue(SubstrateParameter->F90) << TEXTVIEW(",");
 		Printer << LowerValue(SubstrateParameter->Roughness) << TEXTVIEW(",");
 		Printer << LowerValue(SubstrateParameter->Anisotropy) << TEXTVIEW(",");
-		Printer << LowerValue(SubstrateParameter->SSSProfileId) << TEXTVIEW(",");
+		Printer << TEXT("ExtractSubsurfaceProfileInt(") << LowerValue(SubstrateParameter->SSSProfileId) << TEXTVIEW("),");
 		Printer << TEXTVIEW("false") << TEXTVIEW(","); // bSupportDefaultSSSProfile
 		Printer << LowerValue(SubstrateParameter->SSSMFP) << TEXTVIEW(",");
 		Printer << LowerValue(SubstrateParameter->SSSMFPScale) << TEXTVIEW(",");
@@ -1868,7 +1869,7 @@ struct FEntryPointCodeGen
 		Printer << LowerValue(SubstrateParameter->Roughness) << TEXTVIEW(",");
 		Printer << LowerValue(SubstrateParameter->Anisotropy) << TEXTVIEW(",");
 		Printer << LowerValue(SubstrateParameter->SubSurfaceColor) << TEXTVIEW(",");
-		Printer << LowerValue(SubstrateParameter->SubSurfaceProfileId) << TEXTVIEW(",");
+		Printer << TEXT("ExtractSubsurfaceProfileInt(") << LowerValue(SubstrateParameter->SubSurfaceProfileId) << TEXTVIEW("),");
 		Printer << LowerValue(SubstrateParameter->ClearCoat) << TEXTVIEW(",");
 		Printer << LowerValue(SubstrateParameter->ClearCoatRoughness) << TEXTVIEW(",");
 		Printer << LowerValue(SubstrateParameter->EmissiveColor) << TEXTVIEW(",");
@@ -1890,6 +1891,22 @@ struct FEntryPointCodeGen
 		Printer << EndArgs;
 	}
 
+	void LowerSubstrateToon(const MIR::FSubstrateToon* SubstrateParameter)
+	{
+		Printer << TEXTVIEW("GetSubstrateToonBSDF");
+		Printer << BeginArgs;
+		Printer << LowerValue(SubstrateParameter->ToonProfileId) << TEXTVIEW(",");
+		Printer << LowerValue(SubstrateParameter->BaseColor) << TEXTVIEW(",");
+		Printer << LowerValue(SubstrateParameter->Metallic) << TEXTVIEW(",");
+		Printer << LowerValue(SubstrateParameter->Specular) << TEXTVIEW(",");
+		Printer << LowerValue(SubstrateParameter->Roughness) << TEXTVIEW(",");
+		Printer << LowerValue(SubstrateParameter->EmissiveColor) << TEXTVIEW(",");
+		Printer << LowerValue(SubstrateParameter->PatternUVs) << TEXTVIEW(",");
+		Printer << LowerValue(SubstrateParameter->LocalBasisIndex) << TEXTVIEW(",");
+		Printer << GetSharedLocalBasesTypes(SubstrateParameter->bFullySimplifiedInstruction);
+		Printer << EndArgs;
+	}
+
 	void LowerSubstrateHorizontalMixing(const MIR::FSubstrateHorizontalMixing* SubstrateParameter)
 	{
 		if (SubstrateParameter->bParameterBlendingEnabled)
@@ -2106,6 +2123,7 @@ struct FEntryPointCodeGen
 			case MIR::VK_SubstratePromoteToOperator:
 			case MIR::VK_SubstrateSlab:
 			case MIR::VK_SubstrateShadingModels:
+			case MIR::VK_SubstrateToon:
 				return false;
 			case MIR::VK_Extern:
 				return (static_cast<const MIR::FExtern*>(Instr)->GetInfo().Flags & MIR::EExternFlags::Inline) != MIR::EExternFlags::None;
@@ -2474,6 +2492,7 @@ static void GetShaderCompilerEnvironment(const FMaterial* Material, const FMater
 	OutEnvironment.SetDefine(TEXT("MATERIAL_COMPUTE_FOG_PER_PIXEL"), Material->ComputeFogPerPixel());
 	OutEnvironment.SetDefine(TEXT("MATERIAL_FULLY_ROUGH"), false);
 	OutEnvironment.SetDefine(TEXT("MATERIAL_USES_ANISOTROPY"), CompilationOutput.bUsesAnisotropy);
+	OutEnvironment.SetDefine(TEXT("MATERIAL_USES_SPECULAR_PROFILE"), EnumHasAnyFlags(CompilationOutput.SubstrateMaterialCompilationOutput.SubstrateMaterialBsdfFeatures, ESubstrateBsdfFeature::SpecularProfile) && Substrate::IsSpecularProfileEnabled(ShaderPlatform));
 	OutEnvironment.SetDefine(TEXT("MATERIAL_NEURAL_POST_PROCESS"), (CompilationOutput.bUsedWithNeuralNetworks || Material->IsUsedWithNeuralNetworks()) && Material->IsPostProcessMaterial());
 	OutEnvironment.SetDefine(TEXT("NUM_VIRTUALTEXTURE_SAMPLES"), CompilationOutput.UniformExpressionSet.GetVTStacks().Num());
 	OutEnvironment.SetDefine(TEXT("NUM_VIRTUALTEXTURE_FEEDBACK_REQUESTS"), CompilationOutput.NumVirtualTextureFeedbackRequests);
diff --git a/Engine/Source/Runtime/Engine/Public/Materials/MaterialExpressionSubstrate.h b/Engine/Source/Runtime/Engine/Public/Materials/MaterialExpressionSubstrate.h
--- a/Engine/Source/Runtime/Engine/Public/Materials/MaterialExpressionSubstrate.h
+++ b/Engine/Source/Runtime/Engine/Public/Materials/MaterialExpressionSubstrate.h
@@ -911,7 +911,7 @@ class UMaterialExpressionSubstrateToonBSDF : public UMaterialExpressionSubstrate
 	//~ Begin UMaterialExpression Interface
 #if WITH_EDITOR
 	virtual int32 Compile(class FMaterialCompiler* Compiler, int32 OutputIndex) override;
-	// SUBSTRATE_TODO virtual void Build(MIR::FEmitter& Emitter) override;
+	virtual void Build(MIR::FEmitter& Emitter) override;
 	virtual void GetCaption(TArray<FString>& OutCaptions) const override;
 	virtual EMaterialValueType GetOutputValueType(int32 OutputIndex) override;
 	virtual EMaterialValueType GetInputValueType(int32 InputIndex) override;
diff --git a/Engine/Source/Runtime/Engine/Public/Materials/MaterialIR.h b/Engine/Source/Runtime/Engine/Public/Materials/MaterialIR.h
--- a/Engine/Source/Runtime/Engine/Public/Materials/MaterialIR.h
+++ b/Engine/Source/Runtime/Engine/Public/Materials/MaterialIR.h
@@ -73,6 +73,7 @@ enum EValueKind
 	VK_SubstrateDefaultSlab,
 	VK_SubstrateSlab,
 	VK_SubstrateShadingModels,
+	VK_SubstrateToon,
 	VK_SubstrateHorizontalMixing,
 	VK_SubstrateVerticalLayering,
 	VK_SubstrateCoverageWeight,
@@ -1182,6 +1183,22 @@ struct FSubstrateShadingModels : TInstruction<VK_SubstrateShadingModels, 24>
 	bool bFullySimplifiedInstruction;
 };
 
+struct FSubstrateToon : TInstruction<VK_SubstrateToon, 10> 
+{
+	FValue* Normal;
+	FValue* ToonProfileId;
+	FValue* BaseColor;
+	FValue* Metallic;
+	FValue* Specular;
+	FValue* Roughness;
+	FValue* EmissiveColor;
+	FValue* PatternUVs;
+	FValue* IsAtBottom;
+	FValue* LocalBasisIndex;
+
+	bool bFullySimplifiedInstruction;
+};
+
 struct FSubstrateHorizontalMixing : TInstruction<VK_SubstrateHorizontalMixing, 9>
 {
 	FValue* Background;
diff --git a/Engine/Source/Runtime/Engine/Public/Materials/MaterialIREmitter.h b/Engine/Source/Runtime/Engine/Public/Materials/MaterialIREmitter.h
--- a/Engine/Source/Runtime/Engine/Public/Materials/MaterialIREmitter.h
+++ b/Engine/Source/Runtime/Engine/Public/Materials/MaterialIREmitter.h
@@ -259,6 +259,20 @@ struct FSubstrateSlabDesc
 	FValueRef LocalBasisIndex;
 };
 
+struct FSubstrateToonDesc
+{
+	FValueRef Normal;
+	FValueRef ToonProfileId;
+	FValueRef BaseColor;
+	FValueRef Metallic;
+	FValueRef Specular;
+	FValueRef Roughness;
+	FValueRef EmissiveColor;
+	FValueRef PatternUVs;
+	FValueRef IsAtBottom;
+	FValueRef LocalBasisIndex;
+};
+
 // The FEmitter is a helper object responsible for emitting MIR values and instructions.
 // It is primarily used within UMaterialExpression::Build() implementations to lower the
 // high-level semantics of a material expression into corresponding lower-level MIR nodes.
@@ -829,6 +843,9 @@ public:
 	// Create a SubstrateData representing a Slab BSDF.
 	FValueRef SubstrateSlab(const FSubstrateSlabDesc& Desc);
 
+	// Create a SubstrateData representing a Toon BSDF.
+	FValueRef SubstrateToon(const FSubstrateToonDesc& Desc);
+
 	// Get the SubstrateData resulting from the horizontal mixing of the Background and Foreground SubstrateData.
 	FValueRef SubstrateHorizontalMixing(FValueRef Background, FValueRef Foreground, FValueRef Mix, FValueRef OperatorIndex, FValueRef MaxDistanceFromLeaves);
 
