Fix CI error for configurations that doesn't have stats enabled.

#Synced-CL 51621904

diff --git a/Engine/Source/Runtime/Engine/Private/ProfilingDebugging/TraceHitches.cpp b/Engine/Source/Runtime/Engine/Private/ProfilingDebugging/TraceHitches.cpp
--- a/Engine/Source/Runtime/Engine/Private/ProfilingDebugging/TraceHitches.cpp
+++ b/Engine/Source/Runtime/Engine/Private/ProfilingDebugging/TraceHitches.cpp
@@ -15,7 +15,7 @@
 	#define UE_SNAPSHOTHITCHES_ENABLED !UE_BUILD_SHIPPING
 #endif
 
-#if UE_SNAPSHOTHITCHES_ENABLED
+#if UE_SNAPSHOTHITCHES_ENABLED && STATS
 
 static class FTraceHitches : private FSelfRegisteringExec
 {
@@ -106,4 +106,4 @@ private:
 	
 } GTraceHitechesCmd;
 
-#endif //UE_SNAPSHOTHITCHES_ENABLED
+#endif //UE_SNAPSHOTHITCHES_ENABLED && STATS
