Remove deprecated functions in Actor

#Synced-CL 51625217

diff --git a/Engine/Source/Developer/FunctionalTesting/Classes/FunctionalTest.h b/Engine/Source/Developer/FunctionalTesting/Classes/FunctionalTest.h
--- a/Engine/Source/Developer/FunctionalTesting/Classes/FunctionalTest.h
+++ b/Engine/Source/Developer/FunctionalTesting/Classes/FunctionalTest.h
@@ -744,8 +744,6 @@ public:
 	UE_API virtual void PostLoad() override;
 	UE_API void PostEditChangeProperty(struct FPropertyChangedEvent& PropertyChangedEvent) override;
 	UE_API virtual void GetAssetRegistryTags(FAssetRegistryTagsContext Context) const override;
-	UE_DEPRECATED(5.4, "Implement the version that takes FAssetRegistryTagsContext instead.")
-	UE_API virtual void GetAssetRegistryTags(TArray<FAssetRegistryTag>& OutTags) const override;
 
 	static UE_API void OnSelectObject(UObject* NewSelection);
 #endif // WITH_EDITOR
diff --git a/Engine/Source/Developer/FunctionalTesting/Private/FunctionalTest.cpp b/Engine/Source/Developer/FunctionalTesting/Private/FunctionalTest.cpp
--- a/Engine/Source/Developer/FunctionalTesting/Private/FunctionalTest.cpp
+++ b/Engine/Source/Developer/FunctionalTesting/Private/FunctionalTest.cpp
@@ -738,13 +738,6 @@ FString AFunctionalTest::GetBeautifiedName()
 	return TEXT("");
 }
 
-void AFunctionalTest::GetAssetRegistryTags(TArray<FAssetRegistryTag>& OutTags) const
-{
-	PRAGMA_DISABLE_DEPRECATION_WARNINGS;
-	Super::GetAssetRegistryTags(OutTags);
-	PRAGMA_ENABLE_DEPRECATION_WARNINGS;
-}
-
 void AFunctionalTest::GetAssetRegistryTags(FAssetRegistryTagsContext Context) const
 {
 	Super::GetAssetRegistryTags(Context);
diff --git a/Engine/Source/Runtime/Engine/Classes/GameFramework/Actor.h b/Engine/Source/Runtime/Engine/Classes/GameFramework/Actor.h
--- a/Engine/Source/Runtime/Engine/Classes/GameFramework/Actor.h
+++ b/Engine/Source/Runtime/Engine/Classes/GameFramework/Actor.h
@@ -393,12 +393,6 @@ public:
 	}
 
 #if WITH_EDITOR
-	/** Deprecated for a non virtual version. */
-	UE_DEPRECATED(5.4, "Call IsInEditLevelInstanceHierarchy/IsInEditLevelInstance instead.")
-	virtual bool IsInEditingLevelInstance() const
-	{
-		return IsInEditLevelInstance();
-	}
 
 	/** If true, the actor belongs to a level instance which is currently in an edit mode: edit or property override edit */
 	bool IsInAnyEditLevelInstance() const
@@ -847,14 +841,6 @@ public:
 	UPROPERTY(BlueprintReadOnly, EditDefaultsOnly, Category=Replication)
 	TEnumAsByte<enum ENetDormancy> NetDormancy;
 
-	/** Gives the actor a chance to pause replication to a player represented by the passed in actor - only called on server */
-	UE_DEPRECATED(5.3, "Replication pausing is deprecated.")
-	ENGINE_API virtual bool IsReplicationPausedForConnection(const FNetViewer& ConnectionOwnerNetViewer);
-
-	/** Called on the client when the replication paused value is changed */
-	UE_DEPRECATED(5.3, "Replication pausing is deprecated.")
-	ENGINE_API virtual void OnReplicationPausedChanged(bool bIsReplicationPaused);
-
 	/** Controls how to handle spawning this actor in a situation where it's colliding with something else. "Default" means AlwaysSpawn here. */
 	UPROPERTY(EditAnywhere, BlueprintReadWrite, Category=Actor)
 	ESpawnActorCollisionHandlingMethod SpawnCollisionHandlingMethod;
@@ -2366,8 +2352,6 @@ public:
 #if WITH_EDITOR
 	ENGINE_API virtual bool Modify(bool bAlwaysMarkDirty = true) override;
 	ENGINE_API virtual void GetAssetRegistryTags(FAssetRegistryTagsContext Context) const override;
-	UE_DEPRECATED(5.4, "Implement the version that takes FAssetRegistryTagsContext instead.")
-	ENGINE_API virtual void GetAssetRegistryTags(TArray<FAssetRegistryTag>& OutTags) const override;
 	ENGINE_API virtual bool NeedsLoadForTargetPlatform(const ITargetPlatform* TargetPlatform) const;
 	ENGINE_API virtual void PreEditChange(FProperty* PropertyThatWillChange) override;
 	ENGINE_API virtual bool CanEditChange(const FProperty* InProperty) const;
@@ -2673,8 +2657,10 @@ public:
 	ENGINE_API virtual AActor* GetSceneOutlinerParent() const;
 
 	/** Returns the actor package that should be used in the scene outliner */
+#if !UE_STRIP_DEPRECATED_PROPERTIES 
 	UE_DEPRECATED(5.6, "Please use GetSceneOutlinerItemPackage or GetSceneOutlinerTopParentPackage instead.")
 	ENGINE_API virtual UPackage* GetSceneOutlinerPackage() const;
+#endif
 
 	/** Returns the actor item package that should be used in the scene outliner */
 	ENGINE_API virtual UPackage* GetSceneOutlinerItemPackage() const;
@@ -2685,10 +2671,6 @@ public:
 	/** Called before editor copy, true allow export */
 	virtual bool ShouldExport() { return true; }
 
-	/** Called before editor paste, true allow import */
-	UE_DEPRECATED(5.2, "Use the override that takes a StringView instead")
-	virtual bool ShouldImport(FString* ActorPropString, bool IsMovingLevel) { return ActorPropString ? ShouldImport(FStringView(*ActorPropString), IsMovingLevel) : ShouldImport(FStringView(), IsMovingLevel); }
-
 	/** Called before editor paste, true allow import */
 	virtual bool ShouldImport(FStringView ActorPropString, bool IsMovingLevel) { return true; }
 
@@ -3266,10 +3248,6 @@ public:
 	/** Debug rendering to visualize the component tree for this actor. */
 	ENGINE_API void DrawDebugComponents(FColor const& BaseColor=FColor::White) const;
 
-	
-	UE_DEPRECATED(5.3, "Use MarkComponentsAsGarbage instead.")
-	virtual void MarkComponentsAsPendingKill() { MarkComponentsAsGarbage(); }
-
 	/** Called to mark all components as garbage when the actor is being destroyed
 	 *
 	 *  @param bModify if True, Modify will be called on actor before marking components
@@ -3472,6 +3450,7 @@ protected:
 	 */
 	ENGINE_API virtual void OnStopReplicationForIris(const FOnStopReplicationParams&);
 
+#if !UE_STRIP_DEPRECATED_PROPERTIES 
 	/**
 	* Called for all Actors set to replicate during BeginPlay, It will also be called if SetReplicates(true) is called and the object is not already replicating
 	*/
@@ -3490,6 +3469,8 @@ protected:
 	*/
 	UE_DEPRECATED(5.7, "This is being deprecated as part of iris beta. Do not call directly override FillReplicationParams instead to provide custom parameters.")
 	void BeginReplication(const FActorReplicationParams& Params) {}
+	
+#endif
 
 	/**
 	 * Pushes the owning NetConnection for the actor and all of its children to the replication system.
diff --git a/Engine/Source/Runtime/Engine/Private/Actor.cpp b/Engine/Source/Runtime/Engine/Private/Actor.cpp
--- a/Engine/Source/Runtime/Engine/Private/Actor.cpp
+++ b/Engine/Source/Runtime/Engine/Private/Actor.cpp
@@ -639,13 +639,6 @@ void AActor::PreSave(FObjectPreSaveContext ObjectSaveContext)
 }
 
 #if WITH_EDITOR
-void AActor::GetAssetRegistryTags(TArray<FAssetRegistryTag>& OutTags) const
-{
-	PRAGMA_DISABLE_DEPRECATION_WARNINGS;
-	Super::GetAssetRegistryTags(OutTags);
-	PRAGMA_ENABLE_DEPRECATION_WARNINGS;
-}
-
 void AActor::GetAssetRegistryTags(FAssetRegistryTagsContext Context) const
 {
 	Super::GetAssetRegistryTags(Context);
@@ -3034,15 +3027,6 @@ void AActor::ForceNetUpdate()
 	}
 }
 
-bool AActor::IsReplicationPausedForConnection(const FNetViewer& ConnectionOwnerNetViewer)
-{
-	return false;
-}
-
-void AActor::OnReplicationPausedChanged(bool bIsReplicationPaused)
-{
-}
-
 void AActor::SetNetDormancy(ENetDormancy NewDormancy)
 {
 	if (IsNetMode(NM_Client))
diff --git a/Engine/Source/Runtime/Engine/Private/ActorEditor.cpp b/Engine/Source/Runtime/Engine/Private/ActorEditor.cpp
--- a/Engine/Source/Runtime/Engine/Private/ActorEditor.cpp
+++ b/Engine/Source/Runtime/Engine/Private/ActorEditor.cpp
@@ -1092,12 +1092,14 @@ AActor* AActor::GetSceneOutlinerParent() const
 	return GetAttachParentActor();
 }
 
+#if !UE_STRIP_DEPRECATED_PROPERTIES 
 PRAGMA_DISABLE_DEPRECATION_WARNINGS
 UPackage* AActor::GetSceneOutlinerPackage() const
 {
 	return GetSceneOutlinerItemPackage();
 }
 PRAGMA_ENABLE_DEPRECATION_WARNINGS
+#endif
 
 UPackage* AActor::GetSceneOutlinerItemPackage() const
 {
diff --git a/Engine/Source/Runtime/Engine/Private/DataChannel.cpp b/Engine/Source/Runtime/Engine/Private/DataChannel.cpp
--- a/Engine/Source/Runtime/Engine/Private/DataChannel.cpp
+++ b/Engine/Source/Runtime/Engine/Private/DataChannel.cpp
@@ -3367,7 +3367,6 @@ void UActorChannel::ProcessBunchInternal(FInBunch & Bunch)
 	bool bLatestIsReplicationPaused = Bunch.bIsReplicationPaused != 0;
 	if (bLatestIsReplicationPaused != IsReplicationPaused())
 	{
-		Actor->OnReplicationPausedChanged(bLatestIsReplicationPaused);
 		SetReplicationPaused(bLatestIsReplicationPaused);
 	}
 	PRAGMA_ENABLE_DEPRECATION_WARNINGS
@@ -3675,16 +3674,7 @@ int64 UActorChannel::ReplicateActor()
 
 	if (OpenPacketId.First != INDEX_NONE && NetViewers.Num() > 0)
 	{
-		bool bNewPaused = true;
-
-		for (const FNetViewer& NetViewer : NetViewers)
-		{
-			if (!Actor->IsReplicationPausedForConnection(NetViewer))
-			{
-				bNewPaused = false;
-				break;
-			}
-		}
+		bool bNewPaused = false;
 
 		const bool bOldPaused = IsReplicationPaused();
 
