Iris
* Fix RootObjectAdapter defaulting to NetUpdates every frame until RenitPollFrequency would be called where it switched to 30/secs

#rnx

#Synced-CL 51624330

diff --git a/Engine/Source/Runtime/Engine/Private/Net/Iris/ReplicationSystem/NetRootObjectAdapter.cpp b/Engine/Source/Runtime/Engine/Private/Net/Iris/ReplicationSystem/NetRootObjectAdapter.cpp
--- a/Engine/Source/Runtime/Engine/Private/Net/Iris/ReplicationSystem/NetRootObjectAdapter.cpp
+++ b/Engine/Source/Runtime/Engine/Private/Net/Iris/ReplicationSystem/NetRootObjectAdapter.cpp
@@ -177,6 +177,7 @@ void FNetRootObjectAdapter::StartReplication(ULevel* Level)
 		UObjectReplicationBridge::FRootObjectReplicationParams Params 
 		{
 			.bUseClassConfigDynamicFilter = false,
+			.PollFrequency = UNetRootObjectFactory::GetDefaultNetUpdateFrequency(),
 		};
 
 		if (bIsNotRouted)
diff --git a/Engine/Source/Runtime/Engine/Private/Net/Iris/ReplicationSystem/NetRootObjectFactory.cpp b/Engine/Source/Runtime/Engine/Private/Net/Iris/ReplicationSystem/NetRootObjectFactory.cpp
--- a/Engine/Source/Runtime/Engine/Private/Net/Iris/ReplicationSystem/NetRootObjectFactory.cpp
+++ b/Engine/Source/Runtime/Engine/Private/Net/Iris/ReplicationSystem/NetRootObjectFactory.cpp
@@ -29,6 +29,11 @@ namespace UE::Net::Private
 		TEXT("The default net update frequency of objects using the NetRootObjectFactory."));
 }
 
+float UNetRootObjectFactory::GetDefaultNetUpdateFrequency()
+{
+	return UE::Net::Private::DefaultRootObjectNetUpdateFrequency;
+}
+
 FName UNetRootObjectFactory::GetFactoryName()
 {
 	return UE::Net::Private::RootObjectFactoryName;
diff --git a/Engine/Source/Runtime/Engine/Public/Net/Iris/ReplicationSystem/NetRootObjectFactory.h b/Engine/Source/Runtime/Engine/Public/Net/Iris/ReplicationSystem/NetRootObjectFactory.h
--- a/Engine/Source/Runtime/Engine/Public/Net/Iris/ReplicationSystem/NetRootObjectFactory.h
+++ b/Engine/Source/Runtime/Engine/Public/Net/Iris/ReplicationSystem/NetRootObjectFactory.h
@@ -55,6 +55,7 @@ class UNetRootObjectFactory : public UNetObjectFactory
 public:
 
 	ENGINE_API static FName GetFactoryName();
+	ENGINE_API static float GetDefaultNetUpdateFrequency();
 
 public:
 
