Fix async physics thread safety: unsafe UObject access on worker threads.
Destruction:
- TermBody worker-safe (IsInGameThread guard, PhysScene from actor handle off-GT)
- StartAsyncTermBody_GameThread delegates cleanup to TermBody (no duplication)
- AsyncTermBody uses actor handle only (no UObject access)
- AsyncTermBody always calls ReleaseActor for valid handles
- FAsyncTermBodyPayload stripped to ActorHandle only
- ISM: UnRegister once on GT, move bodies to worker, call TermBody directly
- RemoveFromComponentMaps: Pin(true) instead of Get(true)
Creation:
- CachePhysicsCreationInputs pre-computes collision data, velocity, OwnerComponentMapKey
- BuildBodyFilterData: private static overload, instance method uses cache off-GT
- InitDynamicProperties: cached velocity off-GT, GetBodySetup() instead of direct access
- InitBodies post-add: OwnerComponentMapKey for maps, GetBodySetup(), PhysScene from scope
- GetAsyncPhysicsCreationInputs() centralizes off-GT cache access
#pf 69b08daa0120ab9865a4814a

#Synced-CL 51624647

diff --git a/Engine/Plugins/Experimental/FastGeoStreaming/Source/FastGeoStreaming/Private/FastGeoInstancedStaticMeshComponent.cpp b/Engine/Plugins/Experimental/FastGeoStreaming/Source/FastGeoStreaming/Private/FastGeoInstancedStaticMeshComponent.cpp
--- a/Engine/Plugins/Experimental/FastGeoStreaming/Source/FastGeoStreaming/Private/FastGeoInstancedStaticMeshComponent.cpp
+++ b/Engine/Plugins/Experimental/FastGeoStreaming/Source/FastGeoStreaming/Private/FastGeoInstancedStaticMeshComponent.cpp
@@ -17,6 +17,7 @@
 #include "NaniteSceneProxy.h"
 #include "PhysicsEngine/BodySetup.h"
 #include "SceneInterface.h"
+#include "Physics/Experimental/PhysScene_Chaos.h"
 #include "PhysicsEngine/PhysicsObjectExternalInterface.h"
 #include "PhysicsProxy/SingleParticlePhysicsProxy.h"
 
@@ -300,6 +301,19 @@ void FFastGeoInstancedStaticMeshComponent::OnAsyncDestroyPhysicsStateBegin_GameT
 {
 	FFastGeoComponent::OnAsyncDestroyPhysicsStateBegin_GameThread();
 
+	// UnRegisterForCollisionEvents on game thread before going async.
+	// TermBody skips UnRegister when called from a worker thread (see FBodyInstance::TermBody contract).
+	if (BodyInstance.bNotifyRigidBodyCollision)
+	{
+		if (UFastGeoSurrogateComponent* SurrogateComponent = GetSurrogateComponent())
+		{
+			if (FPhysScene* PhysScene = GetPhysicsScene())
+			{
+				PhysScene->UnRegisterForCollisionEvents(SurrogateComponent);
+			}
+		}
+	}
+
 	// Move InstanceBodies in AsyncDestroyPhysicsStatePayload
 	check(AsyncDestroyPhysicsStatePayload.IsEmpty());
 	AsyncDestroyPhysicsStatePayload = MoveTemp(InstanceBodies);
diff --git a/Engine/Plugins/Experimental/FastGeoStreaming/Source/FastGeoStreaming/Private/FastGeoPrimitiveComponent.cpp b/Engine/Plugins/Experimental/FastGeoStreaming/Source/FastGeoStreaming/Private/FastGeoPrimitiveComponent.cpp
--- a/Engine/Plugins/Experimental/FastGeoStreaming/Source/FastGeoStreaming/Private/FastGeoPrimitiveComponent.cpp
+++ b/Engine/Plugins/Experimental/FastGeoStreaming/Source/FastGeoStreaming/Private/FastGeoPrimitiveComponent.cpp
@@ -932,7 +932,7 @@ FBodyInstance* FFastGeoPrimitiveComponent::GetBodyInstance(FName BoneName, bool
 void FFastGeoPrimitiveComponent::OnAsyncCreatePhysicsStateBegin_GameThread()
 {
 	BodyInstanceOwner.Initialize(this);
-	BodyInstance.CachePhysicsCreationInputs(&BodyInstanceOwner);
+	BodyInstance.CachePhysicsCreationInputs(GetSurrogateComponent(), &BodyInstanceOwner);
 	BodyInstanceOwner.Uninitialize();
 
 	Super::OnAsyncCreatePhysicsStateBegin_GameThread();
diff --git a/Engine/Source/Runtime/Engine/Classes/PhysicsEngine/BodyInstance.h b/Engine/Source/Runtime/Engine/Classes/PhysicsEngine/BodyInstance.h
--- a/Engine/Source/Runtime/Engine/Classes/PhysicsEngine/BodyInstance.h
+++ b/Engine/Source/Runtime/Engine/Classes/PhysicsEngine/BodyInstance.h
@@ -76,6 +76,17 @@ struct FBodyInstancePhysicsCreationInputs
 	// Cached complex materials and masks
 	TArray<UPhysicalMaterial*> ResolvedComplexMaterials;
 	TArray<FPhysicalMaterialMaskParams> ResolvedComplexMasks;
+
+	// Pre-computed collision filter data (avoids OwnerComponent/SourceObject weak ptr access on worker thread)
+	FBodyCollisionData ResolvedCollisionData;
+
+	// Pre-computed initial velocity for InitDynamicProperties (avoids OwnerComponent dereference on worker thread)
+	FVector ResolvedInitialLinearVelocity = FVector::ZeroVector;
+	bool bResolvedComponentAwake = false;
+
+	// Cached owner component pointer for use as a map key ONLY.
+	// Do NOT dereference - use only for AddToComponentMaps / RegisterForCollisionEvents.
+	UPrimitiveComponent* OwnerComponentMapKey = nullptr;
 };
 
 #ifndef CHAOS_DEBUG_NAME
@@ -825,7 +836,10 @@ public:
 	ENGINE_API int32 GetAllShapes_AssumesLocked(TArray<FPhysicsShapeHandle>& OutShapes) const;
 
 	/**
-	 * Terminates the body, releasing resources
+	 * Terminates the body, releasing resources.
+	 * Can be called from a worker thread IF:
+	 * - UnRegisterForCollisionEvents was already called on the game thread (or bNotifyRigidBodyCollision is false)
+	 * - The caller has exclusive ownership of this FBodyInstance (e.g., ISM bodies moved to async payload)
 	 * @param bNeverDeferRelease In some cases orphaned actors can have their internal release deferred. If this isn't desired this flag will override that behavior
 	 */
 	ENGINE_API void TermBody(bool bNeverDeferRelease = false);
@@ -838,8 +852,6 @@ public:
 
 	private:
 		FPhysicsActorHandle ActorHandle;
-		TWeakObjectPtr<class UPrimitiveComponent> OwnerComponent;
-		bool bNotifyRigidBodyCollision;
 
 		friend struct FBodyInstance;
 	};
@@ -1544,11 +1556,22 @@ private:
 
 public:
 	// Resolved Inputs used by physics state creation (currently only used by async mode)
-	ENGINE_API const FBodyInstancePhysicsCreationInputsPtr& CachePhysicsCreationInputs(IPhysicsBodyInstanceOwner* BodyInstanceOwner);
+	ENGINE_API const FBodyInstancePhysicsCreationInputsPtr& CachePhysicsCreationInputs(UPrimitiveComponent* InOwnerComponent, IPhysicsBodyInstanceOwner* InBodyInstanceOwner = nullptr);
 	void ClearCachedPhysicsCreationInputs() { PhysicsCreationInputs.Reset(); }
 	const FBodyInstancePhysicsCreationInputsPtr& GetPhysicsCreationInputs() const { return PhysicsCreationInputs; }
 
+	/**
+	 * Returns the cached physics creation inputs when called from a worker thread (not game thread, not parallel game thread).
+	 * Returns nullptr on the game thread or if no cache is available.
+	 * Use this to safely branch between cached data (off-GT) and live data (on-GT).
+	 */
+	const FBodyInstancePhysicsCreationInputs* GetAsyncPhysicsCreationInputs() const;
+
 private:
+
+	/** Static version of BuildBodyFilterData taking explicit parameters (avoids weak ptr dereference, used for pre-computation and off-GT calls) */
+	static ENGINE_API void BuildBodyFilterData(const FBodyInstance* InInstance, UPrimitiveComponent* InOwnerComp, UObject* InSourceObj, UBodySetup* InSetup, FBodyCollisionFilterData& OutFilterData, const int32 InShapeIndex = INDEX_NONE);
+
 	FBodyInstancePhysicsCreationInputsPtr PhysicsCreationInputs;
 };
 
diff --git a/Engine/Source/Runtime/Engine/Private/InstancedStaticMesh.cpp b/Engine/Source/Runtime/Engine/Private/InstancedStaticMesh.cpp
--- a/Engine/Source/Runtime/Engine/Private/InstancedStaticMesh.cpp
+++ b/Engine/Source/Runtime/Engine/Private/InstancedStaticMesh.cpp
@@ -40,6 +40,7 @@
 #include "CollisionDebugDrawingPublic.h"
 #include "DataDrivenShaderPlatformInfo.h"
 #include "Chaos/PhysicsObjectInterface.h"
+#include "Physics/Experimental/PhysScene_Chaos.h"
 #include "PhysicsProxy/SingleParticlePhysicsProxy.h"
 
 #include "Elements/Framework/EngineElementsLibrary.h"
@@ -3017,6 +3018,16 @@ void UInstancedStaticMeshComponent::OnAsyncDestroyPhysicsStateBegin_GameThread(T
 {
 	USceneComponent::OnAsyncDestroyPhysicsStateBegin_GameThread(OutRootedObjects);
 
+	// Call UnRegisterForCollisionEvents on game thread before going async.
+	// TermBody skips UnRegisterForCollisionEvents when called from a worker thread (see FBodyInstance::TermBody contract).
+	if (BodyInstance.bNotifyRigidBodyCollision)
+	{
+		if (FPhysScene* PhysScene = GetWorld()->GetPhysicsScene())
+		{
+			PhysScene->UnRegisterForCollisionEvents(this);
+		}
+	}
+
 	// Move InstanceBodies in AsyncDestroyPhysicsStatePayload
 	check(AsyncDestroyPhysicsStatePayload.IsEmpty());
 	AsyncDestroyPhysicsStatePayload = MoveTemp(InstanceBodies);
diff --git a/Engine/Source/Runtime/Engine/Private/PhysicsEngine/BodyInstance.cpp b/Engine/Source/Runtime/Engine/Private/PhysicsEngine/BodyInstance.cpp
--- a/Engine/Source/Runtime/Engine/Private/PhysicsEngine/BodyInstance.cpp
+++ b/Engine/Source/Runtime/Engine/Private/PhysicsEngine/BodyInstance.cpp
@@ -349,7 +349,6 @@ bool FCollisionResponse::operator==(const FCollisionResponse& Other) const
 
 FBodyInstance::FAsyncTermBodyPayload::FAsyncTermBodyPayload()
 : ActorHandle(DefaultPhysicsActorHandle())
-, bNotifyRigidBodyCollision(false)
 {
 }
 ////////////////////////////////////////////////////////////////////////////
@@ -1222,6 +1221,14 @@ static void GetSimulatingAndBlendWeight(const USkeletalMeshComponent* SkelMeshCo
 	}
 }
 
+namespace UE::Private
+{
+	static UObject* GetSourceObject(UPrimitiveComponent* InPrimitiveComp, IPhysicsBodyInstanceOwner* InBodyInstanceOwner)
+	{
+		return InPrimitiveComp ? InPrimitiveComp->GetOwner() : InBodyInstanceOwner ? InBodyInstanceOwner->GetSourceObject() : nullptr;
+	}
+}
+
 namespace Chaos::Private
 {
 	void FInitBodiesHelperBaseInternal::SetBodyInstanceExternalCollisionProfileBodySetup(FBodyInstance* BodyInstance, UBodySetup* ExternalCollisionProfileBodySetup)
@@ -1241,7 +1248,7 @@ FInitBodiesHelper<bCompileStatic, InAllocatorType>::FInitBodiesHelper(
 	UPrimitiveComponent* InPrimitiveComp, FPhysScene* InRBScene,
 	const FInitBodySpawnParams& InSpawnParams, FPhysicsAggregateHandle InAggregate, IPhysicsBodyInstanceOwner* InBodyInstanceOwner)
 	: FInitBodiesHelper<bCompileStatic, InAllocatorType>(InBodies, InTransforms, InBodySetup, InPrimitiveComp,
-		InPrimitiveComp ? InPrimitiveComp->GetOwner() : InBodyInstanceOwner ? InBodyInstanceOwner->GetSourceObject() : nullptr,
+		UE::Private::GetSourceObject(InPrimitiveComp, InBodyInstanceOwner),
 		InRBScene, InSpawnParams, InAggregate, InBodyInstanceOwner)
 {
 }
@@ -1285,9 +1292,8 @@ TInitBodiesHelperBase<InAllocatorType>::TInitBodiesHelperBase(TArray<FBodyInstan
 	class UBodySetup* InBodySetup, class UPrimitiveComponent* InPrimitiveComp, FPhysScene* InRBScene,
 	const FInitBodySpawnParams& InSpawnParams, FPhysicsAggregateHandle InAggregate, IPhysicsBodyInstanceOwner* InBodyInstanceOwner)
 	// Forward to the other constructor
-	: TInitBodiesHelperBase<InAllocatorType>(InBodies, InTransforms, InBodySetup,
-		InPrimitiveComp,
-		InPrimitiveComp ? InPrimitiveComp->GetOwner() : InBodyInstanceOwner ? InBodyInstanceOwner->GetSourceObject() : nullptr,
+	: TInitBodiesHelperBase<InAllocatorType>(InBodies, InTransforms, InBodySetup, InPrimitiveComp,
+		UE::Private::GetSourceObject(InPrimitiveComp, InBodyInstanceOwner),
 		InRBScene, InSpawnParams, InAggregate, InBodyInstanceOwner)
 {
 }
@@ -1355,8 +1361,15 @@ bool TInitBodiesHelperBase<InAllocatorType>::CreateShapes_AssumesLocked(FBodyIns
 	ComplexPhysMats = Instance->GetComplexPhysicalMaterials(ComplexPhysMatMasks);
 
 	FBodyCollisionData BodyCollisionData;
-	Instance->BuildBodyFilterData(BodyCollisionData.CollisionFilterData);
-	FBodyInstance::BuildBodyCollisionFlags(BodyCollisionData.CollisionFlags, Instance->GetCollisionEnabled(), BodySetup->GetCollisionTraceFlag() == CTF_UseComplexAsSimple);
+	if (const FBodyInstancePhysicsCreationInputs* Cached = Instance->GetAsyncPhysicsCreationInputs())
+	{
+		BodyCollisionData = Cached->ResolvedCollisionData;
+	}
+	else
+	{
+		Instance->BuildBodyFilterData(BodyCollisionData.CollisionFilterData);
+		FBodyInstance::BuildBodyCollisionFlags(BodyCollisionData.CollisionFlags, Instance->GetCollisionEnabled(), BodySetup->GetCollisionTraceFlag() == CTF_UseComplexAsSimple);
+	}
 
 	bool bInitFail = false;
 
@@ -1373,9 +1386,9 @@ bool TInitBodiesHelperBase<InAllocatorType>::CreateShapes_AssumesLocked(FBodyIns
 
 UBodySetup* FBodyInstance::GetBodySetup() const
 {
-	if (!IsInGameThread() && !IsInParallelGameThread() && PhysicsCreationInputs.IsValid())
+	if (const FBodyInstancePhysicsCreationInputs* Cached = GetAsyncPhysicsCreationInputs())
 	{
-		return PhysicsCreationInputs->ResolvedBodySetup;
+		return Cached->ResolvedBodySetup;
 	}
 
 	FGCScopeGuard Guard;
@@ -1651,7 +1664,7 @@ void TInitBodiesHelperBase<InAllocatorType>::InitBodies()
 								Body_External.SetShapeSimCollisionEnabled(ShapeIndex, false);
 							}
 						}
-						if (const UBodySetupCore* BS = BI->BodySetup.Get())
+						if (const UBodySetup* BS = BI->GetBodySetup())
 						{
 							for (int32 ShapeIndex = 0; ShapeIndex < NumShapes; ++ShapeIndex)
 							{
@@ -1668,14 +1681,16 @@ void TInitBodiesHelperBase<InAllocatorType>::InitBodies()
 					FPhysicsActorHandle ActorHandle = BI->GetPhysicsActor();
 					if (FPhysicsInterface::IsValid(ActorHandle))
 					{
-						PhysScene->AddToComponentMaps(BI->OwnerComponent.Get(), ActorHandle);
-					}
-					if (BI->bNotifyRigidBodyCollision)
-					{
-						if (UPrimitiveComponent* PrimComp = BI->OwnerComponent.Get())
+						// Resolve component pointer safely: cached map key on worker thread, weak ptr on game thread
+						const FBodyInstancePhysicsCreationInputs* Cached = BI->GetAsyncPhysicsCreationInputs();
+						if (UPrimitiveComponent* CompKey = Cached ? Cached->OwnerComponentMapKey : BI->OwnerComponent.Get())
 						{
-							FPhysScene_Chaos* LocalPhysScene = PrimComp->GetWorld()->GetPhysicsScene();
-							LocalPhysScene->RegisterForCollisionEvents(PrimComp);
+							PhysScene->AddToComponentMaps(CompKey, ActorHandle);
+
+							if (BI->bNotifyRigidBodyCollision)
+							{
+								PhysScene->RegisterForCollisionEvents(CompKey);
+							}
 						}
 					}
 				}
@@ -1752,16 +1767,7 @@ void FBodyInstance::InitBody(UBodySetup* Setup, const FTransform& Transform, UPr
 	Bodies.Add(this);
 	Transforms.Add(Transform);
 
-	UObject* SourceObjectToUse = InSourceObject;
-	if (!SourceObjectToUse && PrimComp)
-	{
-		SourceObjectToUse = PrimComp->GetOwner();
-	}
-	if (!SourceObjectToUse && InBodyInstanceOwner)
-	{
-		SourceObjectToUse = InBodyInstanceOwner->GetSourceObject();
-	}
-
+	UObject* SourceObjectToUse = InSourceObject ? InSourceObject : UE::Private::GetSourceObject(PrimComp, InBodyInstanceOwner);
 	bool bIsStatic = SpawnParams.bStaticPhysics;
 	if(bIsStatic)
 	{
@@ -1814,25 +1820,17 @@ FBodyInstance::FAsyncTermBodyPayload FBodyInstance::StartAsyncTermBody_GameThrea
 	TRACE_CPUPROFILER_EVENT_SCOPE(FBodyInstance::StartAsyncTermBody_GameThread);
 	SCOPE_CYCLE_COUNTER(STAT_TermBody);
 
-	// Prepare payload for async task before resetting the ActorHandle/OwnerComponent
+	// Capture the actor handle before nulling it - AsyncTermBody will release it on the worker thread.
 	FAsyncTermBodyPayload Payload;
 	Payload.ActorHandle = GetPhysicsActor();
-	Payload.OwnerComponent = OwnerComponent;
-	Payload.bNotifyRigidBodyCollision = bNotifyRigidBodyCollision;
+
+	// Null the actor so TermBody skips RemoveFromComponentMaps + ReleaseActor
+	// (those will be done async by AsyncTermBody using the captured handle).
 	SetPhysicsActor(nullptr);
-	OwnerComponent = nullptr;
-	BodyInstanceOwner = nullptr;
 
-	CurrentSceneState = BodyInstanceSceneState::NotAdded;
-	BodySetup = NULL;
-	ExternalCollisionProfileBodySetup = nullptr;
+	// TermBody handles everything else on GT: UnRegister, cleanup, DOFConstraint, per-shape overrides.
+	TermBody();
 
-	if (DOFConstraint)
-	{
-		DOFConstraint->TermConstraint();
-		FConstraintInstance::Free(DOFConstraint);
-		DOFConstraint = NULL;
-	}
 	return Payload;
 }
 
@@ -1840,24 +1838,14 @@ void FBodyInstance::AsyncTermBody(FAsyncTermBodyPayload& Payload)
 {
 	SCOPE_CYCLE_COUNTER(STAT_TermBody);
 
-	if (UPrimitiveComponent* PrimComp = Payload.OwnerComponent.Get())
-	{
-		if (FPhysScene_Chaos* PhysScene = PrimComp->GetWorld()->GetPhysicsScene())
-		{
-			if (FPhysicsInterface::IsValid(Payload.ActorHandle))
-			{
-				PhysScene->RemoveFromComponentMaps(Payload.ActorHandle);
-			}
-			if (Payload.bNotifyRigidBodyCollision)
-			{
-				PhysScene->UnRegisterForCollisionEvents(PrimComp);
-			}
-		}
-	}
-
 	if (FPhysicsInterface::IsValid(Payload.ActorHandle))
 	{
 		FPhysScene* PhysScene = static_cast<FPhysScene*>(FPhysicsInterface::GetCurrentScene(Payload.ActorHandle));
+		if (PhysScene)
+		{
+			PhysScene->RemoveFromComponentMaps(Payload.ActorHandle);
+		}
+		// ReleaseActor handles null PhysScene (falls back to direct delete)
 		FPhysicsInterface::ReleaseActor(Payload.ActorHandle, PhysScene);
 	}
 
@@ -1866,13 +1854,17 @@ void FBodyInstance::AsyncTermBody(FAsyncTermBodyPayload& Payload)
 
 /**
  *	Clean up the physics engine info for this instance.
+ *	Can be called from a worker thread (see header comment for contract).
  */
 void FBodyInstance::TermBody(bool bNeverDeferRelease)
 {
 	SCOPE_CYCLE_COUNTER(STAT_TermBody);
 
-	// It's bad if we're being terminated but there's still stuff welded to us
-	if (ShapeToBodiesMap != nullptr && !ShapeToBodiesMap->IsEmpty())
+	const bool bIsGameThread = IsInGameThread();
+
+	// Welded bodies check - only on game thread.
+	// ShapeToBodiesMap is a non-thread-safe TSharedPtr and the error message dereferences OwnerComponent.
+	if (bIsGameThread && ShapeToBodiesMap != nullptr && !ShapeToBodiesMap->IsEmpty())
 	{
 		TStringBuilder<256> Builder;
 		int32 BodyIndex = 0;
@@ -1883,18 +1875,18 @@ void FBodyInstance::TermBody(bool bNeverDeferRelease)
 			Builder.Appendf(TEXT("\tBody[%d]: \'%s'\n"), BodyIndex, *Name);
 			++BodyIndex;
 		}
-		ensureMsgf(false, TEXT("Body %s is being terminated but has %d welded bodies:\n%s"), *OwnerComponent->GetName(), ShapeToBodiesMap->Num(), Builder.ToString());
+		ensureMsgf(false, TEXT("Body %s is being terminated but has %d welded bodies:\n%s"), *GetNameSafe(OwnerComponent.Get()), ShapeToBodiesMap->Num(), Builder.ToString());
 	}
 
-	if (UPrimitiveComponent* PrimComp = OwnerComponent.Get())
+	// UnRegisterForCollisionEvents - game thread only.
+	// Worker-thread callers must UnRegister on GT before going async (see header comment for contract).
+	FPhysScene_Chaos* PhysScene = nullptr;
+	if (bIsGameThread)
 	{
-		if (FPhysScene_Chaos* PhysScene = PrimComp->GetWorld()->GetPhysicsScene())
+		if (UPrimitiveComponent* PrimComp = OwnerComponent.Get())
 		{
-			if (FPhysicsInterface::IsValid(GetPhysicsActor()))
-			{
-				PhysScene->RemoveFromComponentMaps(GetPhysicsActor());
-			}
-			if (bNotifyRigidBodyCollision)
+			PhysScene = PrimComp->GetWorld()->GetPhysicsScene();
+			if (bNotifyRigidBodyCollision && PhysScene)
 			{
 				PhysScene->UnRegisterForCollisionEvents(PrimComp);
 			}
@@ -1903,6 +1895,15 @@ void FBodyInstance::TermBody(bool bNeverDeferRelease)
 
 	if (IsValidBodyInstance())
 	{
+		// Off game thread: get PhysScene from the actor handle (no UObject access needed)
+		if (!PhysScene)
+		{
+			PhysScene = static_cast<FPhysScene_Chaos*>(FPhysicsInterface::GetCurrentScene(GetPhysicsActor()));
+		}
+		if (PhysScene)
+		{
+			PhysScene->RemoveFromComponentMaps(GetPhysicsActor());
+		}
 		ReleasePhysicsActor();
 	}
 
@@ -1911,6 +1912,7 @@ void FBodyInstance::TermBody(bool bNeverDeferRelease)
 	CurrentSceneState = BodyInstanceSceneState::NotAdded;
 	BodySetup = NULL;
 	OwnerComponent = NULL;
+	BodyInstanceOwner = nullptr;
 	ExternalCollisionProfileBodySetup = nullptr;
 
 	// Reset per-shape runtime overrides
@@ -4628,7 +4630,8 @@ bool FBodyInstance::ValidateTransform(const FTransform &Transform, const FString
 
 void FBodyInstance::InitDynamicProperties_AssumesLocked()
 {
-	if (!BodySetup.IsValid())
+	UBodySetup* ResolvedSetup = GetBodySetup();
+	if (!ResolvedSetup)
 	{
 		// This may be invalid following an undo if the BodySetup was a transient object (e.g. in Mesh Paint mode)
 		// Just exit gracefully if so.
@@ -4641,12 +4644,12 @@ void FBodyInstance::InitDynamicProperties_AssumesLocked()
 	{
 		return;
 	}
-	
+
 	if(FPhysicsInterface::IsDynamic(GetPhysicsActor()))
 	{
 		//A non simulated body may become simulated at runtime, so we need to compute its mass.
 		//However, this is not supported for complexAsSimple since a trimesh cannot itself be simulated, it can only be used for collision of other simple shapes.
-		if (BodySetup->GetCollisionTraceFlag() != ECollisionTraceFlag::CTF_UseComplexAsSimple)
+		if (ResolvedSetup->GetCollisionTraceFlag() != ECollisionTraceFlag::CTF_UseComplexAsSimple)
 		{
 			UpdateMassProperties();
 			UpdateDampingProperties();
@@ -4657,7 +4660,7 @@ void FBodyInstance::InitDynamicProperties_AssumesLocked()
 		else
 		{
 #if !(UE_BUILD_SHIPPING || UE_BUILD_TEST)
-			if (bSimulatePhysics)
+			if (bSimulatePhysics && IsInGameThread())
 			{
 				if(UPrimitiveComponent* OwnerComponentInst = OwnerComponent.Get())
 				{
@@ -4668,11 +4671,19 @@ void FBodyInstance::InitDynamicProperties_AssumesLocked()
 #endif
 		}
 
-		UPrimitiveComponent* OwnerComponentInst = OwnerComponent.Get();
-		const AActor* OwningActor = OwnerComponentInst ? OwnerComponentInst->GetOwner() : nullptr;
-
+		FVector InitialLinVel;
 		bool bComponentAwake = false;
-		FVector InitialLinVel = GetInitialLinearVelocity(OwningActor, bComponentAwake);
+		if (const FBodyInstancePhysicsCreationInputs* Cached = GetAsyncPhysicsCreationInputs())
+		{
+			InitialLinVel = Cached->ResolvedInitialLinearVelocity;
+			bComponentAwake = Cached->bResolvedComponentAwake;
+		}
+		else
+		{
+			UPrimitiveComponent* OwnerComponentInst = OwnerComponent.Get();
+			const AActor* OwningActor = OwnerComponentInst ? OwnerComponentInst->GetOwner() : nullptr;
+			InitialLinVel = GetInitialLinearVelocity(OwningActor, bComponentAwake);
+		}
 
 		if (ShouldInstanceSimulatingPhysics())
 		{
@@ -4715,61 +4726,72 @@ void FBodyInstance::InitDynamicProperties_AssumesLocked()
 }
 
 void FBodyInstance::BuildBodyFilterData(FBodyCollisionFilterData& OutFilterData, const int32 ShapeIndex) const
+{
+	if (const FBodyInstancePhysicsCreationInputs* Cached = GetAsyncPhysicsCreationInputs())
+	{
+		OutFilterData = Cached->ResolvedCollisionData.CollisionFilterData;
+	}
+	else
+	{
+		FBodyInstance::BuildBodyFilterData(this, OwnerComponent.Get(), SourceObject.Get(), GetBodySetup(), OutFilterData, ShapeIndex);
+	}
+}
+
+void FBodyInstance::BuildBodyFilterData(const FBodyInstance* InInstance, UPrimitiveComponent* InOwnerComp, UObject* InSourceObj, UBodySetup* InSetup, FBodyCollisionFilterData& OutFilterData, const int32 InShapeIndex)
 {
 	// this can happen in landscape height field collision component
-	if(!BodySetup.IsValid())
+	if(!InSetup)
 	{
 		return;
 	}
 
-	UPrimitiveComponent* OwnerComponentInst = OwnerComponent.Get();
-	AActor* ActorOwner = Cast<AActor>(SourceObject.Get());
-	if (!ActorOwner && OwnerComponentInst)
+	AActor* ActorOwner = Cast<AActor>(InSourceObj);
+	if (!ActorOwner && InOwnerComp)
 	{
-		ActorOwner = OwnerComponentInst->GetOwner();
+		ActorOwner = InOwnerComp->GetOwner();
 	}
 
-	const bool bPhysicsStatic = !OwnerComponentInst || OwnerComponentInst->IsWorldGeometry();
+	const bool bPhysicsStatic = !InOwnerComp || InOwnerComp->IsWorldGeometry();
 
 	// Grab collision setting from body instance.
 	// If we're looking at a specific shape, intersect it with the shape's specific collision setting.
-	ECollisionEnabled::Type UseCollisionEnabled // this checks actor/component override. 
-		= ShapeIndex == INDEX_NONE
-		? GetCollisionEnabled()
-		: CollisionEnabledIntersection(GetCollisionEnabled(), GetShapeCollisionEnabled(ShapeIndex));
-	FCollisionResponseContainer UseResponse = CollisionResponses.GetResponseContainer();
-	bool bUseNotifyRBCollision = bNotifyRigidBodyCollision;
-	ECollisionChannel UseChannel = ObjectType;
+	ECollisionEnabled::Type UseCollisionEnabled // this checks actor/component override.
+		= InShapeIndex == INDEX_NONE
+		? InInstance->GetCollisionEnabled()
+		: CollisionEnabledIntersection(InInstance->GetCollisionEnabled(), InInstance->GetShapeCollisionEnabled(InShapeIndex));
+	FCollisionResponseContainer UseResponse = InInstance->CollisionResponses.GetResponseContainer();
+	bool bUseNotifyRBCollision = InInstance->bNotifyRigidBodyCollision;
+	ECollisionChannel UseChannel = InInstance->ObjectType;
 
 
-	bool bUseContactModification = bContactModification;
+	bool bUseContactModification = InInstance->bContactModification;
 
 	// @TODO: The skel mesh really shouldn't be the (pseudo-)authority here on the body's collision.
-	//        This block should ultimately be removed, and outside of this (in the skel component) we 
+	//        This block should ultimately be removed, and outside of this (in the skel component) we
 	//        should configure the bodies to reflect this desired behavior.
-	if(USkeletalMeshComponent* SkelMeshComp = Cast<USkeletalMeshComponent>(OwnerComponentInst))
+	if(USkeletalMeshComponent* SkelMeshComp = Cast<USkeletalMeshComponent>(InOwnerComp))
 	{
-		if (CVarEnableDynamicPerBodyFilterHacks.GetValueOnAnyThread() && bHACK_DisableCollisionResponse)
+		if (CVarEnableDynamicPerBodyFilterHacks.GetValueOnAnyThread() && InInstance->bHACK_DisableCollisionResponse)
 		{
 			UseResponse.SetAllChannels(ECR_Ignore);
 			UseCollisionEnabled = ECollisionEnabled::PhysicsOnly;
 		}
-		else if(BodySetup->CollisionReponse == EBodyCollisionResponse::BodyCollision_Enabled)
+		else if(InSetup->CollisionReponse == EBodyCollisionResponse::BodyCollision_Enabled)
 		{
 			UseResponse.SetAllChannels(ECR_Block);
 		}
-		else if(BodySetup->CollisionReponse == EBodyCollisionResponse::BodyCollision_Disabled)
+		else if(InSetup->CollisionReponse == EBodyCollisionResponse::BodyCollision_Disabled)
 		{
 			UseResponse.SetAllChannels(ECR_Ignore);
 			UseCollisionEnabled = ECollisionEnabled::PhysicsOnly;		// this will prevent object traces hitting this as well
 		}
 
-		const bool bDisableSkelComponentOverride = CVarEnableDynamicPerBodyFilterHacks.GetValueOnAnyThread() && bHACK_DisableSkelComponentFilterOverriding;
+		const bool bDisableSkelComponentOverride = CVarEnableDynamicPerBodyFilterHacks.GetValueOnAnyThread() && InInstance->bHACK_DisableSkelComponentFilterOverriding;
 		if (bDisableSkelComponentOverride)
 		{
 			// if we are disabling the skeletal component override, we want the original body instance collision response
 			// this is to allow per body instance collision filtering instead of taking the data from the skeletal mesh
-			UseResponse = CollisionResponses.GetResponseContainer();
+			UseResponse = InInstance->CollisionResponses.GetResponseContainer();
 		}
 		else
 		{
@@ -4780,11 +4802,11 @@ void FBodyInstance::BuildBodyFilterData(FBodyCollisionFilterData& OutFilterData,
 		bUseNotifyRBCollision = bUseNotifyRBCollision && SkelMeshComp->BodyInstance.bNotifyRigidBodyCollision;
 	}
 
-	if (ShapeIndex != INDEX_NONE)
+	if (InShapeIndex != INDEX_NONE)
 	{
 		// If we're looking at a specific shape, take the intersection of its collision responses
 		// with the base primitive component's collision responses.
-		FCollisionResponseContainer ShapeResponse = GetShapeResponseToChannels(ShapeIndex, UseResponse);
+		FCollisionResponseContainer ShapeResponse = InInstance->GetShapeResponseToChannels(InShapeIndex, UseResponse);
 		UseResponse = FCollisionResponseContainer::CreateMinContainer(UseResponse, ShapeResponse);
 	}
 
@@ -4793,11 +4815,10 @@ void FBodyInstance::BuildBodyFilterData(FBodyCollisionFilterData& OutFilterData,
 	if((UseCollisionEnabled == ECollisionEnabled::NoCollision) && ActorOwner && (ActorOwner->IsA(AVolume::StaticClass()) == false))
 	{
 		UWorld* World = ActorOwner->GetWorld();
-		UPrimitiveComponent* PrimComp = OwnerComponentInst;
-		if(World && World->bEnableTraceCollision &&
-		   (PrimComp->IsA(UStaticMeshComponent::StaticClass()) || PrimComp->IsA(USkeletalMeshComponent::StaticClass()) || PrimComp->IsA(UBrushComponent::StaticClass())))
+		if(World && World->bEnableTraceCollision && InOwnerComp &&
+		   (InOwnerComp->IsA(UStaticMeshComponent::StaticClass()) || InOwnerComp->IsA(USkeletalMeshComponent::StaticClass()) || InOwnerComp->IsA(UBrushComponent::StaticClass())))
 		{
-			//UE_LOG(LogPhysics, Warning, TEXT("Enabling collision %s : %s"), *GetNameSafe(Owner), *GetNameSafe(OwnerComponent.Get()));
+			//UE_LOG(LogPhysics, Warning, TEXT("Enabling collision %s : %s"), *GetNameSafe(Owner), *GetNameSafe(InOwnerComp));
 			// clear all other channel just in case other people using those channels to do something
 			UseResponse.SetAllChannels(ECR_Ignore);
 			UseCollisionEnabled = ECollisionEnabled::QueryOnly;
@@ -4805,23 +4826,23 @@ void FBodyInstance::BuildBodyFilterData(FBodyCollisionFilterData& OutFilterData,
 	}
 #endif
 
-	const bool bUseComplexAsSimple = (BodySetup.Get()->GetCollisionTraceFlag() == CTF_UseComplexAsSimple);
-	const bool bUseSimpleAsComplex = (BodySetup.Get()->GetCollisionTraceFlag() == CTF_UseSimpleAsComplex);
+	const bool bUseComplexAsSimple = (InSetup->GetCollisionTraceFlag() == CTF_UseComplexAsSimple);
+	const bool bUseSimpleAsComplex = (InSetup->GetCollisionTraceFlag() == CTF_UseSimpleAsComplex);
 
 	if(UseCollisionEnabled != ECollisionEnabled::NoCollision)
 	{
 		// CCD is determined by root body in case of welding
-		bool bRootCCD = (WeldParent != nullptr) ? WeldParent->bUseCCD : bUseCCD;
+		bool bRootCCD = (InInstance->WeldParent != nullptr) ? InInstance->WeldParent->bUseCCD : InInstance->bUseCCD;
 
-		uint32 SourceObjectID = SourceObject.IsValid() ? SourceObject->GetUniqueID() : ActorOwner ? ActorOwner->GetUniqueID() : 0;
-		uint32 CompID = (OwnerComponentInst != nullptr) ? OwnerComponentInst->GetUniqueID() : 0;
+		uint32 SourceObjectID = InSourceObj ? InSourceObj->GetUniqueID() : ActorOwner ? ActorOwner->GetUniqueID() : 0;
+		uint32 CompID = (InOwnerComp != nullptr) ? InOwnerComp->GetUniqueID() : 0;
 
 		FPhysicsFilterBuilder Builder;
 		Builder.SetOwnerID(SourceObjectID);
 		Builder.SetComponentID(CompID);
 		Builder.SetCollisionChannelIndex(UseChannel);
 		Builder.SetResponses(UseResponse);
-		Builder.SetMaskFilter(MaskFilter);
+		Builder.SetMaskFilter(InInstance->MaskFilter);
 		Builder.SetFlags(Chaos::EFilterFlags::CCD, bRootCCD && !bPhysicsStatic);
 		Builder.SetFlags(Chaos::EFilterFlags::ContactNotify, bUseNotifyRBCollision);
 		Builder.SetFlags(Chaos::EFilterFlags::StaticShape, bPhysicsStatic);
@@ -5008,10 +5029,25 @@ FRecalculatedMassProperties& FBodyInstance::OnRecalculatedMassProperties()
 	return BodyInstanceDelegates->OnRecalculatedMassProperties;
 }
 
-const FBodyInstancePhysicsCreationInputsPtr& FBodyInstance::CachePhysicsCreationInputs(IPhysicsBodyInstanceOwner* InBodyInstanceOwner)
+const FBodyInstancePhysicsCreationInputs* FBodyInstance::GetAsyncPhysicsCreationInputs() const
+{
+	if (!IsInGameThread() && !IsInParallelGameThread() && PhysicsCreationInputs.IsValid())
+	{
+		return PhysicsCreationInputs.Get();
+	}
+	return nullptr;
+}
+
+const FBodyInstancePhysicsCreationInputsPtr& FBodyInstance::CachePhysicsCreationInputs(UPrimitiveComponent* InOwnerComponent, IPhysicsBodyInstanceOwner* InBodyInstanceOwner)
 {
 	check(IsInGameThread() || IsInParallelGameThread());
 	check(!PhysicsCreationInputs.IsValid());
+
+	// UPrimitiveComponent implements IPhysicsBodyInstanceOwner, so use it as fallback
+	if (!InBodyInstanceOwner)
+	{
+		InBodyInstanceOwner = InOwnerComponent;
+	}
 	check(InBodyInstanceOwner);
 
 	FBodyInstancePhysicsCreationInputsPtr NewPhysicsCreationInputs = MakeShared<FBodyInstancePhysicsCreationInputs, ESPMode::ThreadSafe>();
@@ -5020,6 +5056,20 @@ const FBodyInstancePhysicsCreationInputsPtr& FBodyInstance::CachePhysicsCreation
 	NewPhysicsCreationInputs->ResolvedSimpleMaterial = FBodyInstance::GetSimplePhysicalMaterial(this, TWeakObjectPtr<UPrimitiveComponent>(), OwnerBodySetup, InBodyInstanceOwner);
 	FBodyInstance::GetComplexPhysicalMaterials(this, TWeakObjectPtr<UPrimitiveComponent>(), NewPhysicsCreationInputs->ResolvedComplexMaterials, &NewPhysicsCreationInputs->ResolvedComplexMasks, InBodyInstanceOwner);
 
+	// Pre-compute collision filter data (avoids OwnerComponent/SourceObject weak ptr access on worker thread)
+	UObject* SourceObj = UE::Private::GetSourceObject(InOwnerComponent, InBodyInstanceOwner);
+	FBodyInstance::BuildBodyFilterData(this, InOwnerComponent, SourceObj, OwnerBodySetup, NewPhysicsCreationInputs->ResolvedCollisionData.CollisionFilterData);
+	FBodyInstance::BuildBodyCollisionFlags(NewPhysicsCreationInputs->ResolvedCollisionData.CollisionFlags, GetCollisionEnabled(), OwnerBodySetup->GetCollisionTraceFlag() == CTF_UseComplexAsSimple);
+
+	// Pre-compute initial velocity for InitDynamicProperties (avoids OwnerComponent dereference on worker thread)
+	const AActor* OwningActor = InOwnerComponent ? InOwnerComponent->GetOwner() : nullptr;
+	bool bComponentAwake = false;
+	NewPhysicsCreationInputs->ResolvedInitialLinearVelocity = GetInitialLinearVelocity(OwningActor, bComponentAwake);
+	NewPhysicsCreationInputs->bResolvedComponentAwake = bComponentAwake;
+
+	// Cache component pointer as a map key only (rooted by the async job queue, safe to use as key on worker thread)
+	NewPhysicsCreationInputs->OwnerComponentMapKey = InOwnerComponent;
+
 	PhysicsCreationInputs = NewPhysicsCreationInputs;
 	return PhysicsCreationInputs;
 }
diff --git a/Engine/Source/Runtime/Engine/Private/PhysicsEngine/Experimental/PhysScene_Chaos.cpp b/Engine/Source/Runtime/Engine/Private/PhysicsEngine/Experimental/PhysScene_Chaos.cpp
--- a/Engine/Source/Runtime/Engine/Private/PhysicsEngine/Experimental/PhysScene_Chaos.cpp
+++ b/Engine/Source/Runtime/Engine/Private/PhysicsEngine/Experimental/PhysScene_Chaos.cpp
@@ -1490,17 +1490,19 @@ void FPhysScene_Chaos::RemoveFromComponentMaps(IPhysicsProxyBase* InObject)
 	const TWeakObjectPtr<UPrimitiveComponent>* ComponentWeak = PhysicsProxyToComponentMap.Find(InObject);
 	if (ComponentWeak)
 	{
-		if (TObjectPtr<UPrimitiveComponent> Component = ComponentWeak->Get(true))
+		// Pin() safely resolves the weak pointer (FGCScopeGuard internally), unlike Get() which is not safe off game thread.
+		// The strong reference is not retained - ComponentKey is used as a map key only, never dereferenced.
+		if (UPrimitiveComponent* ComponentKey = Cast<UPrimitiveComponent>(ComponentWeak->Pin(true).Get()))
 		{
-			PendingMigratedPhysicsStateByComponent.Remove(Component.Get());
+			PendingMigratedPhysicsStateByComponent.Remove(ComponentKey);
 
-			TArray<IPhysicsProxyBase*>* ProxyArray = ComponentToPhysicsProxyMap.Find(Component.Get());
+			TArray<IPhysicsProxyBase*>* ProxyArray = ComponentToPhysicsProxyMap.Find(ComponentKey);
 			if (ProxyArray)
 			{
 				ProxyArray->RemoveSingleSwap(InObject, EAllowShrinking::No);
 				if (ProxyArray->Num() == 0)
 				{
-					ComponentToPhysicsProxyMap.Remove(Component.Get());
+					ComponentToPhysicsProxyMap.Remove(ComponentKey);
 				}
 			}
 		}
