[TextureCompiler]
- Add support for recovering from missing PreEditChange instead of crashing fatally when possible

#Synced-CL 51597565

diff --git a/Engine/Source/Runtime/Engine/Private/Tests/Texture2DTests.cpp b/Engine/Source/Runtime/Engine/Private/Tests/Texture2DTests.cpp
--- a/Engine/Source/Runtime/Engine/Private/Tests/Texture2DTests.cpp
+++ b/Engine/Source/Runtime/Engine/Private/Tests/Texture2DTests.cpp
@@ -194,6 +194,86 @@ bool FTexture2DAsyncCompileCancelation::RunTest(const FString& Parameters)
 	return true;
 }
 
+
+IMPLEMENT_SIMPLE_AUTOMATION_TEST(FTexture2DMissingPreEditChange, "System.Engine.Texture2D.MissingPreEditChange", TestFlags)
+bool FTexture2DMissingPreEditChange::RunTest(const FString& Parameters)
+{
+	constexpr int32 TextureSize = 2 * 1024;
+	constexpr int32 NumIteration = 3;
+
+	auto DoTest =
+		[this](bool bVirtualStreaming)
+		{
+			double CompilationTime = 0.0;
+
+			for (int32 Iteration = 0; Iteration < NumIteration; ++Iteration)
+			{
+				TRACE_CPUPROFILER_EVENT_SCOPE(FTexture2DMissingPreEditChange::Iteration);
+
+				UTexture2D* TestTexture = nullptr;
+
+				{
+					TRACE_CPUPROFILER_EVENT_SCOPE(FTexture2DAsyncCompileCancelation::CreateTransient);
+					TestTexture = UTexture2D::CreateTransient(TextureSize, TextureSize);
+				}
+
+				{
+					TRACE_CPUPROFILER_EVENT_SCOPE(FTexture2DMissingPreEditChange::PreEditChange);
+					TestTexture->PreEditChange(nullptr);
+				}
+
+				TestTexture->VirtualTextureStreaming = bVirtualStreaming;
+
+				{
+					TRACE_CPUPROFILER_EVENT_SCOPE(FTexture2DMissingPreEditChange::WriteSource);
+					TestTexture->Source.Init(TextureSize, TextureSize, 1, 1, TSF_BGRA8);
+					int32* Data = (int32*)TestTexture->Source.LockMip(0);
+
+					ParallelFor(TextureSize * TextureSize, [&](int32 Index) { Data[Index] = FMath::Rand(); });
+
+					TestTexture->Source.UnlockMip(0);
+				}
+
+				// This will start a build async
+				{
+					TRACE_CPUPROFILER_EVENT_SCOPE(FTexture2DMissingPreEditChange::PostEditChange);
+					TestTexture->PostEditChange();
+				}
+				
+				// Overwrite source again without calling PreEditChange, this should log some errors
+				{
+					TRACE_CPUPROFILER_EVENT_SCOPE(FTexture2DAsyncCompileCancelation::OverwriteWriteSource);
+					TestTexture->Source.Init(TextureSize, TextureSize, 1, 1, TSF_BGRA8);
+					int32* Data = (int32*)TestTexture->Source.LockMip(0);
+
+					ParallelFor(TextureSize * TextureSize, [&](int32 Index) { Data[Index] = FMath::Rand(); });
+
+					TestTexture->Source.UnlockMip(0);
+				}
+
+				FAssetCompilingManager::Get().FinishCompilationForObjects({ TestTexture });
+				
+				TWeakObjectPtr<UTexture> WeakPtr(TestTexture);
+
+				{
+					TRACE_CPUPROFILER_EVENT_SCOPE(FTexture2DMissingPreEditChange::GC);
+					CollectGarbage(GARBAGE_COLLECTION_KEEPFLAGS);
+				}
+
+				TestFalse(TEXT("There shouldn't be anything preventing the texture from being GCed during async compilation"), WeakPtr.IsValid());
+			}
+		};
+
+	AddExpectedMessage(TEXT("missing call to PreEditChange"), EAutomationExpectedErrorFlags::Contains, 0 /* Must be seen at least once for the test to pass*/);
+	{
+		// Normal and VT textures have different cancelation points
+		DoTest(false /*bVirtualStreaming*/);
+		DoTest(true  /*bVirtualStreaming*/);
+	}
+
+	return true;
+}
+
 #endif
 
 } // namespace Texture2DTest
diff --git a/Engine/Source/Runtime/Engine/Private/TextureCompiler.cpp b/Engine/Source/Runtime/Engine/Private/TextureCompiler.cpp
--- a/Engine/Source/Runtime/Engine/Private/TextureCompiler.cpp
+++ b/Engine/Source/Runtime/Engine/Private/TextureCompiler.cpp
@@ -207,9 +207,11 @@ void FTextureCompilingManager::UpdateCompilationNotification()
 	Notification->Update(GetNumRemainingTextures());
 }
 
-void FTextureCompilingManager::PostCompilation(UTexture* Texture)
+bool FTextureCompilingManager::PostCompilation(UTexture* Texture)
 {
 	TGuardValue PostCompilationGuard(bIsRoutingPostCompilation, true);
+	TGuardValue PostCompilationTextureGuard(PostCompilationTexture, Texture);
+	bIsRoutingPostCompilationError = false;
 
 	check(IsInGameThread());
 	TRACE_CPUPROFILER_EVENT_SCOPE(FTextureCompilingManager::PostCompilation);
@@ -220,7 +222,7 @@ void FTextureCompilingManager::PostCompilation(UTexture* Texture)
 
 	// Track the DDC key suffix of the texture we are done with so that if we re-enter we can
 	// log info and hopefully be able to do some post-mortem on it.
-	CurrentPostCompilationTexture = Texture;
+
 	CurrentPostCompilationDDCKey.Empty();
 	if (FTexturePlatformData** RunningPlatformData = Texture->GetRunningPlatformData(); RunningPlatformData && *RunningPlatformData)
 	{
@@ -233,10 +235,19 @@ void FTextureCompilingManager::PostCompilation(UTexture* Texture)
 
 	Texture->UpdateResource();
 
+	// Most likely a recursion happened during UpdateResource to add the texture again
+	// so we must return an error and take action on it.
+	if (bIsRoutingPostCompilationError)
+	{
+		return false;
+	}
+
 	// Generate an empty property changed event, to force the asset registry tag
 	// to be refreshed now that pixel format and alpha channels are available.
 	FPropertyChangedEvent EmptyPropertyChangedEvent(nullptr);
 	FCoreUObjectDelegates::OnObjectPropertyChanged.Broadcast(Texture, EmptyPropertyChangedEvent);
+
+	return true;
 }
 
 bool FTextureCompilingManager::IsAsyncCompilationAllowed(UTexture* Texture) const
@@ -291,42 +302,55 @@ void FTextureCompilingManager::AddTextures(TArrayView<UTexture* const> InTexture
 	// in question - only the compilation manager shoulid be calling that for editor resources.
 	if (bIsRoutingPostCompilation)
 	{
-		UE_LOG(LogTexture, Error, TEXT("PostCompilation Texture: %s"), CurrentPostCompilationTexture ? *CurrentPostCompilationTexture->GetPathName() : TEXT("<nullptr>"));
+		TStringBuilder<1024> StringBuilder;
+
+		StringBuilder.Appendf(
+			TEXT("Registering a texture to the compile manager from inside a texture postcompilation usually\n")
+			TEXT(" indicates that the previous async operation wasn't completed (i.e. missing call to PreEditChange) before modifying a texture property.\n")
+		);
+
+		StringBuilder.Appendf(TEXT("  PostCompilation Texture: %s\n"), PostCompilationTexture ? *PostCompilationTexture->GetPathName() : TEXT("<nullptr>"));
 
 		// Empty keys most likely means we are on ddc2.
-		UE_LOG(LogTexture, Error, TEXT("PostCompilation DDCKey: %s"), *CurrentPostCompilationDDCKey);
-		UE_LOG(LogTexture, Error, TEXT("AddTextures Count: %d"), InTextures.Num());
+		StringBuilder.Appendf(TEXT("  PostCompilation DDCKey: %s\n"), *CurrentPostCompilationDDCKey);
+		StringBuilder.Appendf(TEXT("  AddTextures Count: %d\n"), InTextures.Num());
 
 		for (const UTexture* ConstTexture : InTextures)
 		{
+			if (ConstTexture == PostCompilationTexture)
+			{
+				bIsRoutingPostCompilationError = true;
+			}
+
 			// We're about to crash anyway.
 			UTexture* Texture = const_cast<UTexture*>(ConstTexture);
 
-			UE_LOG(LogTexture, Error, TEXT("%s:"), *Texture->GetPathName());
+			StringBuilder.Appendf(TEXT("    %s:\n"), *Texture->GetPathName());
 
 			FTexturePlatformData** RunningPlatformData = Texture->GetRunningPlatformData();
 			if (!RunningPlatformData || !RunningPlatformData[0])
 			{
-				UE_LOG(LogTexture, Error, TEXT("   -> No RunningPlatformData!"));
+				StringBuilder.Appendf(TEXT("      -> No RunningPlatformData!\n"));
 			}
 			else
 			{
 				FString* Key = RunningPlatformData[0]->FetchFirstDerivedDataKey.TryGet<FString>();
-				UE_LOG(LogTexture, Error, TEXT("    FetchFirstKey: %s"), Key ? **Key : TEXT("<empty, likely new texture build flow?>"));
+				StringBuilder.Appendf(TEXT("       FetchFirstKey: %s\n"), Key ? **Key : TEXT("<empty, likely new texture build flow?>"));
 				Key = RunningPlatformData[0]->FetchOrBuildDerivedDataKey.TryGet<FString>();
-				UE_LOG(LogTexture, Error, TEXT("    FetchOrBuildKey: %s"), Key ? **Key : TEXT("<empty, likely new texture build flow?>"));
+				StringBuilder.Appendf(TEXT("       FetchOrBuildKey: %s\n"), Key ? **Key : TEXT("<empty, likely new texture build flow?>"));
 			}
 		}
 
-		// This has been updated to Fatal because it potentially modifies RegisteredTextureBuckets below which is iterated upon
-		// during PostCompilation routing. That modification can put us in an unstable state and crash in unexpected and rather
-		// undebuggable ways.
-		UE_LOG(LogTexture, Fatal, 
-			TEXT("Registering a texture to the compile manager from inside a texture postcompilation is not supported and usually")
-			TEXT(" indicates that the previous async operation wasn't completed (i.e. missing call to PreEditChange) before modifying a texture property.")
-			);
+		if (bIsRoutingPostCompilationError && InTextures.Num() == 1)
+		{
+			// Just log as error, will be catched and returned by PostCompilation, and acted on accordingly.
+			UE_LOG(LogTexture, Error, TEXT("%s"), StringBuilder.ToString());
+			return;
+		}
+
+		// Treat any other case that we didn't plan on as fatal...
+		UE_LOG(LogTexture, Fatal, TEXT("%s"), StringBuilder.ToString());
 	}
-		
 
 	TRACE_CPUPROFILER_EVENT_SCOPE(FTextureCompilingManager::AddTextures)
 
@@ -454,7 +478,6 @@ void FTextureCompilingManager::FinishCompilation(TArrayView<UTexture* const> InT
 		UE_LOG(LogTexture, Fatal, TEXT("Calling FinishCompilation is not allowed during PostCompilation. NumTextures = %d, Texture[0] = %s"), InTextures.Num(), *InTextures[0]->GetPathName());
 	}
 
-
 	using namespace TextureCompilingManagerImpl;
 	check(IsInGameThread());
 
@@ -473,51 +496,57 @@ void FTextureCompilingManager::FinishCompilation(TArrayView<UTexture* const> InT
 		}
 	}
 
-	if (PendingTextures.Num())
+	if (!PendingTextures.Num())
+	{
+		return;
+	}
+
+	class FCompilableTexture final : public AsyncCompilationHelpers::ICompilable
 	{
-		class FCompilableTexture final : public AsyncCompilationHelpers::ICompilable
+	public:
+		FCompilableTexture(UTexture* InTexture)
+			: Texture(InTexture)
 		{
-		public:
-			FCompilableTexture(UTexture* InTexture)
-				: Texture(InTexture)
-			{
-			}
+		}
 
-			FTextureAsyncCacheDerivedDataTask* GetAsyncTask()
+		FTextureAsyncCacheDerivedDataTask* GetAsyncTask()
+		{
+			if (FTexturePlatformData** PlatformDataPtr = Texture->GetRunningPlatformData())
 			{
-				if (FTexturePlatformData** PlatformDataPtr = Texture->GetRunningPlatformData())
+				if (FTexturePlatformData* PlatformData = *PlatformDataPtr)
 				{
-					if (FTexturePlatformData* PlatformData = *PlatformDataPtr)
-					{
-						return PlatformData->AsyncTask;
-					}
+					return PlatformData->AsyncTask;
 				}
-
-				return nullptr;
 			}
 
-			void Reschedule(FQueuedThreadPool* InThreadPool, EQueuedWorkPriority InPriority) final
+			return nullptr;
+		}
+
+		void Reschedule(FQueuedThreadPool* InThreadPool, EQueuedWorkPriority InPriority) final
+		{
+			if (FTextureAsyncCacheDerivedDataTask* AsyncTask = GetAsyncTask())
 			{
-				if (FTextureAsyncCacheDerivedDataTask* AsyncTask = GetAsyncTask())
-				{
-					AsyncTask->SetPriority(InPriority);
-				}
+				AsyncTask->SetPriority(InPriority);
 			}
-		
-			bool WaitCompletionWithTimeout(float TimeLimitSeconds) final
+		}
+
+		bool WaitCompletionWithTimeout(float TimeLimitSeconds) final
+		{
+			if (FTextureAsyncCacheDerivedDataTask* AsyncTask = GetAsyncTask())
 			{
-				if (FTextureAsyncCacheDerivedDataTask* AsyncTask = GetAsyncTask())
-				{
-					return AsyncTask->WaitWithTimeout(TimeLimitSeconds);
-				}
-				return true;
+				return AsyncTask->WaitWithTimeout(TimeLimitSeconds);
 			}
+			return true;
+		}
 
-			FName GetName() override { return Texture->GetOutermost()->GetFName(); }
+		FName GetName() override { return Texture->GetOutermost()->GetFName(); }
 
-			TStrongObjectPtr<UTexture> Texture;
-		};
+		TStrongObjectPtr<UTexture> Texture;
+	};
 
+	TArray<UTexture*> InitialTextures(PendingTextures.Array());
+	while (PendingTextures.Num())
+	{
 		TArray<UTexture*> UniqueTextures(PendingTextures.Array());
 		TArray<FCompilableTexture> CompilableTextures(UniqueTextures);
 		using namespace AsyncCompilationHelpers;
@@ -527,20 +556,26 @@ void FTextureCompilingManager::FinishCompilation(TArrayView<UTexture* const> InT
 			CompilableTextures.Num(),
 			LOCTEXT("Textures", "Textures"),
 			LogTexture,
-			[this](ICompilable* Object)
+			[this, &PendingTextures](ICompilable* Object)
 			{
 				UTexture* Texture = static_cast<FCompilableTexture*>(Object)->Texture.Get();
-				PostCompilation(Texture);
-
-				for (TSet<TWeakObjectPtr<UTexture>>& Bucket : RegisteredTextureBuckets)
+				// If the compilation doesn't succeed, it means we need to do another pass on it
+				// to wait and postcompile again, so leave it in the PendingTextures.
+				if (PostCompilation(Texture))
 				{
-					Bucket.Remove(Texture);
+					PendingTextures.Remove(Texture);
+
+					for (TSet<TWeakObjectPtr<UTexture>>& Bucket : RegisteredTextureBuckets)
+					{
+						Bucket.Remove(Texture);
+					}
 				}
 			}
 		);
-
-		PostCompilation(UniqueTextures);
 	}
+
+	// Do the group post compilation only once after all textures has successfully passed individual PostCompilation.
+	PostCompilation(InitialTextures);
 }
 
 void FTextureCompilingManager::PostCompilation(TArrayView<UTexture* const> InCompiledTextures)
@@ -737,9 +772,8 @@ void FTextureCompilingManager::ProcessTextures(bool bLimitExecutionTime, int32 M
 						if (Texture.IsValid())
 						{
 							const bool bHasTimeLeft = bLimitExecutionTime ? ((FPlatformTime::Seconds() - TickStartTime) < MaxSecondsPerFrame) : true;
-							if ((bIsHighestPrio || bHasTimeLeft) && Texture->IsAsyncCacheComplete())
+							if ((bIsHighestPrio || bHasTimeLeft) && Texture->IsAsyncCacheComplete() && PostCompilation(Texture.Get()))
 							{
-								PostCompilation(Texture.Get());
 								ProcessedTextures.Add(Texture.Get());
 							}
 							else
diff --git a/Engine/Source/Runtime/Engine/Public/TextureCompiler.h b/Engine/Source/Runtime/Engine/Public/TextureCompiler.h
--- a/Engine/Source/Runtime/Engine/Public/TextureCompiler.h
+++ b/Engine/Source/Runtime/Engine/Public/TextureCompiler.h
@@ -111,17 +111,19 @@ private:
 	void UpdateCompilationNotification();
 
 	bool GetCurrentPriority(UTexture* InTexture, EQueuedWorkPriority& OutPriority);
-	void PostCompilation(UTexture* Texture);
+	[[nodiscard]] bool PostCompilation(UTexture* Texture);
 	void PostCompilation(TArrayView<UTexture* const> InCompiledTextures);
 
 	void ProcessDeferredRequests();
 
 
 	FString CurrentPostCompilationDDCKey;
-	UTexture* CurrentPostCompilationTexture = nullptr;
 	double LastReschedule = 0.0f;
 	bool bHasShutdown = false;
 	bool bIsRoutingPostCompilation = false;
+	bool bIsRoutingPostCompilationError = false;
+	UTexture* PostCompilationTexture = nullptr;
+
 	UE::TConsumeAllMpmcQueue<TWeakObjectPtr<UTexture>> DeferredRebuildRequestQueue;
 	TArray<TSet<TWeakObjectPtr<UTexture>>> RegisteredTextureBuckets;
 	TUniquePtr<FAsyncCompilationNotification> Notification;
