[Texture2D]
- Proxied resource RHI ref is already initialized in InitRHI, removed an early access from the GT since it was detected as a race condition with the render thread.

#Synced-CL 51597565

diff --git a/Engine/Source/Runtime/Engine/Private/Rendering/Texture2DResource.cpp b/Engine/Source/Runtime/Engine/Private/Rendering/Texture2DResource.cpp
--- a/Engine/Source/Runtime/Engine/Private/Rendering/Texture2DResource.cpp
+++ b/Engine/Source/Runtime/Engine/Private/Rendering/Texture2DResource.cpp
@@ -50,7 +50,6 @@ FTexture2DResource::FTexture2DResource(UTexture2D* InOwner, const FTexture2DReso
 	, ResourceMem(InOwner->ResourceMem)
 	, ProxiedResource(InProxiedResource)
 {
-	TextureReferenceRHI = InOwner->TextureReference.TextureReferenceRHI;
 	CacheSamplerStateInitializer(InOwner);
 }
 
