[TextureCompiler]
- Add guard about potential nullptr being passed in parameter

#Synced-CL 51636719

diff --git a/Engine/Source/Runtime/Engine/Private/TextureCompiler.cpp b/Engine/Source/Runtime/Engine/Private/TextureCompiler.cpp
--- a/Engine/Source/Runtime/Engine/Private/TextureCompiler.cpp
+++ b/Engine/Source/Runtime/Engine/Private/TextureCompiler.cpp
@@ -317,27 +317,30 @@ void FTextureCompilingManager::AddTextures(TArrayView<UTexture* const> InTexture
 
 		for (const UTexture* ConstTexture : InTextures)
 		{
-			if (ConstTexture == PostCompilationTexture)
+			if (ConstTexture)
 			{
-				bIsRoutingPostCompilationError = true;
-			}
+				if (ConstTexture == PostCompilationTexture)
+				{
+					bIsRoutingPostCompilationError = true;
+				}
 
-			// We're about to crash anyway.
-			UTexture* Texture = const_cast<UTexture*>(ConstTexture);
+				// We're about to crash anyway.
+				UTexture* Texture = const_cast<UTexture*>(ConstTexture);
 
-			StringBuilder.Appendf(TEXT("    %s:\n"), *Texture->GetPathName());
+				StringBuilder.Appendf(TEXT("    %s:\n"), *Texture->GetPathName());
 
-			FTexturePlatformData** RunningPlatformData = Texture->GetRunningPlatformData();
-			if (!RunningPlatformData || !RunningPlatformData[0])
-			{
-				StringBuilder.Appendf(TEXT("      -> No RunningPlatformData!\n"));
-			}
-			else
-			{
-				FString* Key = RunningPlatformData[0]->FetchFirstDerivedDataKey.TryGet<FString>();
-				StringBuilder.Appendf(TEXT("       FetchFirstKey: %s\n"), Key ? **Key : TEXT("<empty, likely new texture build flow?>"));
-				Key = RunningPlatformData[0]->FetchOrBuildDerivedDataKey.TryGet<FString>();
-				StringBuilder.Appendf(TEXT("       FetchOrBuildKey: %s\n"), Key ? **Key : TEXT("<empty, likely new texture build flow?>"));
+				FTexturePlatformData** RunningPlatformData = Texture->GetRunningPlatformData();
+				if (!RunningPlatformData || !RunningPlatformData[0])
+				{
+					StringBuilder.Appendf(TEXT("      -> No RunningPlatformData!\n"));
+				}
+				else
+				{
+					FString* Key = RunningPlatformData[0]->FetchFirstDerivedDataKey.TryGet<FString>();
+					StringBuilder.Appendf(TEXT("       FetchFirstKey: %s\n"), Key ? **Key : TEXT("<empty, likely new texture build flow?>"));
+					Key = RunningPlatformData[0]->FetchOrBuildDerivedDataKey.TryGet<FString>();
+					StringBuilder.Appendf(TEXT("       FetchOrBuildKey: %s\n"), Key ? **Key : TEXT("<empty, likely new texture build flow?>"));
+				}
 			}
 		}
 
@@ -362,22 +365,25 @@ void FTextureCompilingManager::AddTextures(TArrayView<UTexture* const> InTexture
 	// in the RHI.
 	for (UTexture* Texture : InTextures)
 	{
-		int32 TexturePriority = 2;
-		switch (Texture->LODGroup)
+		if (Texture)
 		{
-			case TEXTUREGROUP_UI:
-				TexturePriority = 0;
-			break;
-			case TEXTUREGROUP_Terrain_Heightmap:
-				TexturePriority = 1;
-			break;
-		}
+			int32 TexturePriority = 2;
+			switch (Texture->LODGroup)
+			{
+				case TEXTUREGROUP_UI:
+					TexturePriority = 0;
+				break;
+				case TEXTUREGROUP_Terrain_Heightmap:
+					TexturePriority = 1;
+				break;
+			}
 
-		if (RegisteredTextureBuckets.Num() <= TexturePriority)
-		{
-			RegisteredTextureBuckets.SetNum(TexturePriority + 1);
+			if (RegisteredTextureBuckets.Num() <= TexturePriority)
+			{
+				RegisteredTextureBuckets.SetNum(TexturePriority + 1);
+			}
+			RegisteredTextureBuckets[TexturePriority].Emplace(Texture);
 		}
-		RegisteredTextureBuckets[TexturePriority].Emplace(Texture);
 	}
 
 	TRACE_COUNTER_SET(QueuedTextureCompilation, GetNumRemainingTextures());
