DeviceProfile::GetTextureLODSettings followup part 1

#Synced-CL 51641465

diff --git a/Engine/Source/Runtime/Engine/Classes/DeviceProfiles/DeviceProfile.h b/Engine/Source/Runtime/Engine/Classes/DeviceProfiles/DeviceProfile.h
--- a/Engine/Source/Runtime/Engine/Classes/DeviceProfiles/DeviceProfile.h
+++ b/Engine/Source/Runtime/Engine/Classes/DeviceProfiles/DeviceProfile.h
@@ -196,6 +196,13 @@ public:
 	/** Set the memory size bucket to be used when previewing this DP, changing this will reset the expanded cvars. */
 	ENGINE_API void SetPreviewMemorySizeBucket(EPlatformMemorySizeBucket PreviewMemorySizeBucketIn);
 	ENGINE_API EPlatformMemorySizeBucket GetPreviewMemorySizeBucket() const;
+	
+	ENGINE_API void SetPreviewTextureGroupRedirect(const UTextureLODSettings* InPreviewTextureGroupRedirect);
+
+	const UTextureLODSettings* GetPreviewTextureGroupRedirect() const
+	{
+		return PreviewTextureGroupRedirect;
+	}
 
 private:
 	
@@ -210,5 +217,8 @@ private:
 
 	/** The EPlatformMemorySizeBucket to use when processing the device profile for previewing. */
 	EPlatformMemorySizeBucket PreviewMemorySizeBucket = EPlatformMemorySizeBucket::Default;
+	
+	UTextureLODSettings const * PreviewTextureGroupRedirect = nullptr;  // @todo CBTLDP : why is this a member here and not a static or global ?
+
 #endif
 };
diff --git a/Engine/Source/Runtime/Engine/Classes/Engine/TextureLODSettings.h b/Engine/Source/Runtime/Engine/Classes/Engine/TextureLODSettings.h
--- a/Engine/Source/Runtime/Engine/Classes/Engine/TextureLODSettings.h
+++ b/Engine/Source/Runtime/Engine/Classes/Engine/TextureLODSettings.h
@@ -260,17 +260,4 @@ public:
 	/** Array of LOD settings with entries per group. */
 	UPROPERTY(EditAnywhere, config, Category="Texture LOD Settings")
 	TArray<FTextureLODGroup> TextureLODGroups;
-
-#if ALLOW_OTHER_PLATFORM_CONFIG
-public:
-	ENGINE_API void SetPreviewTextureGroupRedirect(UTextureLODSettings* InPreviewTextureGroupRedirect);
-
-	UTextureLODSettings* GetPreviewTextureGroupRedirect() const
-	{
-		return PreviewTextureGroupRedirect;
-	}
-
-private:
-	UTextureLODSettings* PreviewTextureGroupRedirect = nullptr;  // @todo CBTLDP : why is this a member here and not a static or global ?
-#endif
 };
diff --git a/Engine/Source/Runtime/Engine/Private/DeviceProfiles/DeviceProfile.cpp b/Engine/Source/Runtime/Engine/Private/DeviceProfiles/DeviceProfile.cpp
--- a/Engine/Source/Runtime/Engine/Private/DeviceProfiles/DeviceProfile.cpp
+++ b/Engine/Source/Runtime/Engine/Private/DeviceProfiles/DeviceProfile.cpp
@@ -69,18 +69,26 @@ void UDeviceProfile::GatherParentCVarInformationRecursively(OUT TMap<FString, FS
 	}
 }
 
-// @todo CBTLDP : fixme : this casts away const
+
+#if ALLOW_OTHER_PLATFORM_CONFIG
+ENGINE_API void UDeviceProfile::SetPreviewTextureGroupRedirect(const UTextureLODSettings* InPreviewTextureGroupRedirect)
+{
+	PreviewTextureGroupRedirect = InPreviewTextureGroupRedirect;
+}
+#endif
+
 UTextureLODSettings* UDeviceProfile::GetTextureLODSettings() const
 {
+	const UTextureLODSettings * Ret = this; // inheritance
+	
 #if ALLOW_OTHER_PLATFORM_CONFIG
-	UTextureLODSettings* PreviewTextureGroup = GetPreviewTextureGroupRedirect();
-	if (PreviewTextureGroup != nullptr)
+	if ( PreviewTextureGroupRedirect )
 	{
-		return PreviewTextureGroup;
+		Ret = PreviewTextureGroupRedirect;
 	}
 #endif
 
-	return const_cast<UDeviceProfile*>(this); // inheritance
+	return const_cast<UTextureLODSettings *>(Ret);
 }
 
 UDeviceProfile* UDeviceProfile::GetParentProfile(bool bIncludeDefaultObject) const
diff --git a/Engine/Source/Runtime/Engine/Private/DeviceProfiles/DeviceProfileManager.cpp b/Engine/Source/Runtime/Engine/Private/DeviceProfiles/DeviceProfileManager.cpp
--- a/Engine/Source/Runtime/Engine/Private/DeviceProfiles/DeviceProfileManager.cpp
+++ b/Engine/Source/Runtime/Engine/Private/DeviceProfiles/DeviceProfileManager.cpp
@@ -1281,14 +1281,16 @@ void UDeviceProfileManager::LoadProfiles()
 				ITargetPlatform* Platform = TargetPlatforms[PlatformIndex];
 
 				// Set TextureLODSettings
-				UTextureLODSettings* TextureLODSettingsObj = FindProfile(Platform->CookingDeviceProfileName(), false);
-				checkf(TextureLODSettingsObj, TEXT("No TextureLODSettings found for %s"), *Platform->CookingDeviceProfileName());
+				UDeviceProfile * DP = FindProfile(Platform->CookingDeviceProfileName(), false);
+				checkf(DP, TEXT("No Device Profile found for %s"), *Platform->CookingDeviceProfileName());
 
 #if ALLOW_OTHER_PLATFORM_CONFIG
 				// disable redirect if we are using this for cooking.
-				TextureLODSettingsObj->SetPreviewTextureGroupRedirect(nullptr);
+				DP->SetPreviewTextureGroupRedirect(nullptr);
 #endif
 
+				const UTextureLODSettings* TextureLODSettingsObj = DP->GetTextureLODSettings();
+
 				Platform->RegisterTextureLODSettings(TextureLODSettingsObj);
 			}
 
diff --git a/Engine/Source/Runtime/Engine/Private/TextureLODSettings.cpp b/Engine/Source/Runtime/Engine/Private/TextureLODSettings.cpp
--- a/Engine/Source/Runtime/Engine/Private/TextureLODSettings.cpp
+++ b/Engine/Source/Runtime/Engine/Private/TextureLODSettings.cpp
@@ -537,12 +537,3 @@ ETextureMipLoadOptions UTextureLODSettings::GetMipLoadOptions(const UTexture* Te
 		return TextureLODGroups[Texture->LODGroup].MipLoadOptions;
 	}
 }
-
-#if ALLOW_OTHER_PLATFORM_CONFIG
-void UTextureLODSettings::SetPreviewTextureGroupRedirect(UTextureLODSettings* InPreviewTextureGroupRedirect)
-{
-	//@todo CBTLDP : why is this here? move to DeviceProfile?
-
-	PreviewTextureGroupRedirect = InPreviewTextureGroupRedirect;
-}
-#endif
\ No newline at end of file
