Add more instrumentation to try to find UGameEngine::SceneViewport lifetime issue

#Synced-CL 51619035

diff --git a/Engine/Source/Runtime/Engine/Private/GameEngine.cpp b/Engine/Source/Runtime/Engine/Private/GameEngine.cpp
--- a/Engine/Source/Runtime/Engine/Private/GameEngine.cpp
+++ b/Engine/Source/Runtime/Engine/Private/GameEngine.cpp
@@ -810,6 +810,14 @@ void UGameEngine::OnGameWindowClosed(const TSharedRef<SWindow>& WindowBeingClose
 		SceneViewport->bCrashOnDestruction = false; // TEMP: Check for premature destruction
 		SceneViewport.Reset();
 	}
+	// TEMP - BEGIN: Check for premature destruction
+#if !UE_BUILD_SHIPPING
+	else
+	{
+		UE_LOGF(LogViewport, Fatal, "FSceneViewport is missing somehow.");
+	}
+#endif
+	// TEMP - END: Check for premature destruction
 }
 
 void UGameEngine::OnGameWindowMoved_SavePosition(const TSharedRef<SWindow>& WindowBeingMoved, int32 SaveWinPosIndex)
diff --git a/Engine/Source/Runtime/Engine/Private/Slate/SceneViewport.cpp b/Engine/Source/Runtime/Engine/Private/Slate/SceneViewport.cpp
--- a/Engine/Source/Runtime/Engine/Private/Slate/SceneViewport.cpp
+++ b/Engine/Source/Runtime/Engine/Private/Slate/SceneViewport.cpp
@@ -79,7 +79,11 @@ FSceneViewport::FSceneViewport( FViewportClient* InViewportClient, TSharedPtr<SV
 
 FSceneViewport::~FSceneViewport()
 {
-	ensureAlwaysMsgf(!bCrashOnDestruction, TEXT("Freeing FSceneViewport too early.")); // TEMP: Check for premature destruction
+	// TEMP - BEGIN: Check for premature destruction
+#if !UE_BUILD_SHIPPING
+	UE_CLOGF(bCrashOnDestruction, LogViewport, Fatal, "Freeing FSceneViewport too early.");
+#endif
+	// TEMP - END: Check for premature destruction
 	Destroy();
 	// Wait for resources to be deleted
 	FlushRenderingCommands();
