RemoteObjects:  Implementing AActor::GetPreferredRemoteServer to allow Actors to migrate together (e.g. Controller depends on Pawn's Location).

#Synced-CL 51654223

diff --git a/Engine/Source/Runtime/Engine/Classes/Camera/PlayerCameraManager.h b/Engine/Source/Runtime/Engine/Classes/Camera/PlayerCameraManager.h
--- a/Engine/Source/Runtime/Engine/Classes/Camera/PlayerCameraManager.h
+++ b/Engine/Source/Runtime/Engine/Classes/Camera/PlayerCameraManager.h
@@ -562,6 +562,9 @@ public:
 	ENGINE_API virtual void Destroyed() override;
 	ENGINE_API virtual void DisplayDebug(class UCanvas* Canvas, const FDebugDisplayInfo& DebugDisplay, float& YL, float& YPos) override;
 	ENGINE_API virtual void ApplyWorldOffset(const FVector& InOffset, bool bWorldShift) override;
+#if UE_WITH_REMOTE_OBJECT_HANDLE
+	ENGINE_API virtual FRemoteServerId GetPreferredRemoteServer() const override;
+#endif
 	//~ End AActor Interface
 
 	/** Static.  Plays an in-world camera shake that affects all nearby players, with radial distance-based attenuation.
diff --git a/Engine/Source/Runtime/Engine/Classes/GameFramework/Actor.h b/Engine/Source/Runtime/Engine/Classes/GameFramework/Actor.h
--- a/Engine/Source/Runtime/Engine/Classes/GameFramework/Actor.h
+++ b/Engine/Source/Runtime/Engine/Classes/GameFramework/Actor.h
@@ -17,6 +17,10 @@
 #include "Net/Core/Misc/NetSubObjectRegistry.h"
 #include "Engine/ReplicatedState.h"
 
+#if UE_WITH_REMOTE_OBJECT_HANDLE
+#include "UObject/RemoteObjectTypes.h"
+#endif
+
 #include <atomic>
 
 #if WITH_EDITOR
@@ -2326,6 +2330,21 @@ public:
 #if UE_WITH_REMOTE_OBJECT_HANDLE
 	ENGINE_API virtual void PostMigrate(const struct FUObjectMigrationContext& MigrationContext) override;
 	ENGINE_API virtual bool IsMigrationRoot() const override;
+
+	/**
+	 * Returns the server this actor should live on for push migration (ownership change).
+	 * Override for class-specific behavior (e.g., Controller follows Pawn).
+	 * Default:
+	 *  - If the Actor is local and movable, use GetPreferredServerDelegate for game-specific implementation.
+	 *  - Otherwise, report the current owner server (no change).
+	 */
+	ENGINE_API virtual FRemoteServerId GetPreferredRemoteServer() const;
+
+	/**
+	 * Static delegate for plugins or games to provide an implementation (typically a spatial-partition-based server lookup).
+	 * Called by the default GetPreferredRemoteServer() implementation.
+	 */
+	static ENGINE_API TDelegate<FRemoteServerId(const AActor*)> GetPreferredServerDelegate;
 #endif
 	ENGINE_API virtual void BeginDestroy() override;
 #if !UE_STRIP_DEPRECATED_PROPERTIES
diff --git a/Engine/Source/Runtime/Engine/Classes/GameFramework/Controller.h b/Engine/Source/Runtime/Engine/Classes/GameFramework/Controller.h
--- a/Engine/Source/Runtime/Engine/Classes/GameFramework/Controller.h
+++ b/Engine/Source/Runtime/Engine/Classes/GameFramework/Controller.h
@@ -224,6 +224,9 @@ public:
 
 	ENGINE_API virtual void Reset() override;
 	ENGINE_API virtual void Destroyed() override;
+#if UE_WITH_REMOTE_OBJECT_HANDLE
+	ENGINE_API virtual FRemoteServerId GetPreferredRemoteServer() const override;
+#endif
 	//~ End AActor Interface
 
 	/** Getter for Pawn */
diff --git a/Engine/Source/Runtime/Engine/Classes/GameFramework/PlayerController.h b/Engine/Source/Runtime/Engine/Classes/GameFramework/PlayerController.h
--- a/Engine/Source/Runtime/Engine/Classes/GameFramework/PlayerController.h
+++ b/Engine/Source/Runtime/Engine/Classes/GameFramework/PlayerController.h
@@ -2384,6 +2384,9 @@ public:
 	virtual void PostInitializeComponents() override;
 	virtual class AActor* GetViewTarget() const override;
 	virtual class UNetConnection* GetNetConnection() const override;
+#if UE_WITH_REMOTE_OBJECT_HANDLE
+	ENGINE_API virtual FRemoteServerId GetPreferredRemoteServer() const override;
+#endif
 
 	/** Set a view target position for use in replication relevancy. */
 	UFUNCTION(Server, Unreliable)
diff --git a/Engine/Source/Runtime/Engine/Classes/GameFramework/PlayerState.h b/Engine/Source/Runtime/Engine/Classes/GameFramework/PlayerState.h
--- a/Engine/Source/Runtime/Engine/Classes/GameFramework/PlayerState.h
+++ b/Engine/Source/Runtime/Engine/Classes/GameFramework/PlayerState.h
@@ -186,6 +186,9 @@ public:
 	ENGINE_API virtual void Destroyed() override;
 	ENGINE_API virtual void Reset() override;
 	ENGINE_API virtual FString GetHumanReadableName() const override;
+#if UE_WITH_REMOTE_OBJECT_HANDLE
+	ENGINE_API virtual FRemoteServerId GetPreferredRemoteServer() const override;
+#endif
 	//~ End AActor Interface
 
 	/** Return the pawn controlled by this Player State. */
diff --git a/Engine/Source/Runtime/Engine/Private/Actor.cpp b/Engine/Source/Runtime/Engine/Private/Actor.cpp
--- a/Engine/Source/Runtime/Engine/Private/Actor.cpp
+++ b/Engine/Source/Runtime/Engine/Private/Actor.cpp
@@ -65,6 +65,7 @@
 
 #if UE_WITH_REMOTE_OBJECT_HANDLE
 #include "UObject/RemoteExecutor.h"
+#include "UObject/RemoteObject.h"
 #include "UObject/RemoteObjectTransfer.h"
 #include "UObject/UObjectMigrationContext.h"
 #endif
@@ -1257,6 +1258,9 @@ void AActor::PostLoadSubobjects(FObjectInstancingGraph* OuterInstanceGraph)
 
 #if UE_WITH_REMOTE_OBJECT_HANDLE
 
+// Initialize the static delegate
+decltype(AActor::GetPreferredServerDelegate) AActor::GetPreferredServerDelegate;
+
 static USceneComponent* FixupNativeActorComponents(AActor* Actor);
 
 void AActor::PostMigrate(const FUObjectMigrationContext& MigrationContext)
@@ -1297,14 +1301,13 @@ void AActor::PostMigrate(const FUObjectMigrationContext& MigrationContext)
 	TGuardValue<ITransaction*> SuppressTransaction(GUndo, nullptr);
 	if (MigrationContext.MigrationSide == EObjectMigrationSide::Send)
 	{
-		EObjectMigrationSendType MigrationSendType = MigrationContext.GetObjectMigrationSendType(this);
-		UE_LOG(LogActor, Verbose, TEXT("%hs: Send %s (LocalSrvId: %s. DstSrvId: %s. HasOwnership: %d. MigrationType: %s)"), __func__, *GetName(), *LocalServerId.ToString(), *MigrationContext.RemoteServerId.ToString(), bHasOwnership, ToString(MigrationSendType));
 		ensureMsgf(MigrationContext.RemoteServerId != LocalServerId, TEXT("%hs: Sending Object %s from %s to %s but destination is the same as the Local Server Id."), __func__, *GetName(), *LocalServerId.ToString(), *MigrationContext.RemoteServerId.ToString());
 
 		// Remove the actor and it's sub-objects from the network id cache since these objects are now considered remote. It's assumed that the
 		// network id cache will borrow any objects that are not in the cache but are considered remote.
 		//
 		// Loan is an exception because the server is still retaining ownership of this object and can still be dereferenced.
+		const EObjectMigrationSendType MigrationSendType = MigrationContext.GetObjectMigrationSendType(this);
 		if (MigrationSendType != EObjectMigrationSendType::Loan)
 		{
 			GetWorld()->GetNetDriver()->RemoveObjectFromNetworkIdCache(this);
@@ -1362,9 +1365,6 @@ void AActor::PostMigrate(const FUObjectMigrationContext& MigrationContext)
 	}
 	else if (MigrationContext.MigrationSide == EObjectMigrationSide::Receive)
 	{
-		EObjectMigrationRecvType MigrationRecvType = MigrationContext.GetObjectMigrationRecvType(this);
-		UE_LOG(LogActor, Verbose, TEXT("%hs: Recv %s (LocalSrvId: %s. FromSrvId: %s. HasOwnership: %d. MigrationType: %s)"), __func__, *GetName(), *LocalServerId.ToString(), *MigrationContext.RemoteServerId.ToString(), bHasOwnership, ToString(MigrationRecvType));
-
 		// Build up the list of Components we have again (even if we previously owned this object, it could have been mutated)
 		EAddComponentForReplicationBehavior AddComponentForRepType = EAddComponentForReplicationBehavior::NeverAdd;
 		
@@ -1393,6 +1393,7 @@ void AActor::PostMigrate(const FUObjectMigrationContext& MigrationContext)
 		}
 
 		// If we're being passed ownership, register us with the global systems
+		const EObjectMigrationRecvType MigrationRecvType = MigrationContext.GetObjectMigrationRecvType(this);
 		if (MigrationRecvType == EObjectMigrationRecvType::AssignedOwnership)
 		{
 			ensureMsgf(!bActorIsBeingDestroyed, TEXT("We do not want to be assigned ownership of Actors that have already been destroyed. We do not yet know what to do with it."));
@@ -1461,6 +1462,31 @@ bool AActor::IsMigrationRoot() const
 	return true;
 }
 
+FRemoteServerId AActor::GetPreferredRemoteServer() const
+{
+	if (GetPreferredServerDelegate.IsBound())
+	{
+		// Only use spatial lookup if the actor has a local, valid, movable root component.
+		// Access RootComponent TObjectPtr without dereferencing — GetResidence() and
+		// nullptr checks are non-dereferencing. Only call .Get() once confirmed Local.
+		if (RootComponent.GetResidence() == EResidence::Local)
+		{
+			// Using a BP-generated default scene root.  This happens when a C++ class in inherited in Blueprint
+			// and no other components are added.  So we want to treat these actors as having no root component
+			// because they effectively have no transform thus should not be migrated.
+			const bool bBlueprintHasRootComponent = (RootComponent.GetFName() != USceneComponent::GetDefaultSceneRootVariableName());
+
+			USceneComponent* Root = RootComponent.Get();
+			if (Root && Root->IsRegistered() && Root->Mobility == EComponentMobility::Movable
+				&& bBlueprintHasRootComponent)
+			{
+				return GetPreferredServerDelegate.Execute(this);
+			}
+		}
+	}
+	return UE::RemoteObject::Handle::GetOwnerServerId(this);
+}
+
 #endif // UE_WITH_REMOTE_OBJECT_HANDLE
 
 void AActor::ProcessEvent(UFunction* Function, void* Parameters)
diff --git a/Engine/Source/Runtime/Engine/Private/Controller.cpp b/Engine/Source/Runtime/Engine/Private/Controller.cpp
--- a/Engine/Source/Runtime/Engine/Private/Controller.cpp
+++ b/Engine/Source/Runtime/Engine/Private/Controller.cpp
@@ -22,6 +22,10 @@
 #include "GameFramework/PlayerState.h"
 #include "ObjectTrace.h"
 
+#if UE_WITH_REMOTE_OBJECT_HANDLE
+#include "UObject/RemoteObject.h"
+#endif
+
 #include UE_INLINE_GENERATED_CPP_BY_NAME(Controller)
 
 DEFINE_LOG_CATEGORY(LogController);
@@ -607,6 +611,31 @@ void AController::Destroyed()
 	Super::Destroyed();
 }
 
+#if UE_WITH_REMOTE_OBJECT_HANDLE
+FRemoteServerId AController::GetPreferredRemoteServer() const
+{
+	TObjectPtr<APawn> PawnPtr = GetPawn();  // TObjectPtr by value, no auto-pull
+	if (!AutoRTFM::IsClosed() && PawnPtr.GetResidence() == EResidence::Remote)
+	{
+		// We are outside of a transaction (probably push migration step) and the
+		// Pawn is already on another server — that IS its preferred server
+		FRemoteServerId PawnServerId = UE::RemoteObject::Handle::GetOwnerServerId(PawnPtr.GetRemoteId());
+		if (PawnServerId.IsValid())
+		{
+			return PawnServerId;
+		}
+	}
+	else if (APawn* LocalPawn = PawnPtr.Get())
+	{
+		// Pawn is local — ask where it WANTS to be
+		return LocalPawn->GetPreferredRemoteServer();
+	}
+
+	// No Pawn — stay where we are (don't use spatial delegate; controllers
+	// without pawns have no meaningful location to determine preferred server)
+	return UE::RemoteObject::Handle::GetOwnerServerId(this);
+}
+#endif
 
 void AController::CleanupPlayerState()
 {
diff --git a/Engine/Source/Runtime/Engine/Private/PlayerCameraManager.cpp b/Engine/Source/Runtime/Engine/Private/PlayerCameraManager.cpp
--- a/Engine/Source/Runtime/Engine/Private/PlayerCameraManager.cpp
+++ b/Engine/Source/Runtime/Engine/Private/PlayerCameraManager.cpp
@@ -25,6 +25,10 @@
 #include "GameDelegates.h"
 #include "ChaosVisualDebugger/ChaosVisualDebuggerTrace.h"
 
+#if UE_WITH_REMOTE_OBJECT_HANDLE
+#include "UObject/RemoteObject.h"
+#endif
+
 #include UE_INLINE_GENERATED_CPP_BY_NAME(PlayerCameraManager)
 
 DEFINE_LOG_CATEGORY_STATIC(LogPlayerCameraManager, Log, All);
@@ -683,6 +687,29 @@ void APlayerCameraManager::Destroyed()
 	Super::Destroyed();
 }
 
+#if UE_WITH_REMOTE_OBJECT_HANDLE
+FRemoteServerId APlayerCameraManager::GetPreferredRemoteServer() const
+{
+	if (!AutoRTFM::IsClosed() && PCOwner.GetResidence() == EResidence::Remote)
+	{
+		// We are outside of a transaction (probably push migration step) and the
+		// PCOwner is already on another server - that IS its preferred server
+		FRemoteServerId PCServerId = UE::RemoteObject::Handle::GetOwnerServerId(PCOwner.GetRemoteId());
+		if (PCServerId.IsValid())
+		{
+			return PCServerId;
+		}
+	}
+	else if (APlayerController* LocalPC = PCOwner.Get())
+	{
+		return LocalPC->GetPreferredRemoteServer();
+	}
+	
+	// No PCOwner — stay where we are (don't use spatial delegate, our location is not meaningful)
+	return UE::RemoteObject::Handle::GetOwnerServerId(this);
+}
+#endif
+
 void APlayerCameraManager::InitializeFor(APlayerController* PC)
 {
 	FMinimalViewInfo DefaultFOVCache = GetCameraCacheView();
diff --git a/Engine/Source/Runtime/Engine/Private/PlayerController.cpp b/Engine/Source/Runtime/Engine/Private/PlayerController.cpp
--- a/Engine/Source/Runtime/Engine/Private/PlayerController.cpp
+++ b/Engine/Source/Runtime/Engine/Private/PlayerController.cpp
@@ -77,6 +77,7 @@
 #include "AutoRTFM.h"
 #include "UObject/UObjectMigrationContext.h"
 #include "UObject/RemoteExecutor.h"
+#include "UObject/RemoteObject.h"
 #endif
 
 #include UE_INLINE_GENERATED_CPP_BY_NAME(PlayerController)
@@ -6871,5 +6872,13 @@ void ANoPawnPlayerController::ServerEnableReplicationToConnection_Implementation
 	bEnableReplicationToConnection = bEnable;
 }
 
+#if UE_WITH_REMOTE_OBJECT_HANDLE
+FRemoteServerId ANoPawnPlayerController::GetPreferredRemoteServer() const
+{
+	// NoPawnPlayerControllers are server-specific & never migrate servers
+	return UE::RemoteObject::Handle::GetOwnerServerId(this);
+}
+#endif
+
 #undef LOCTEXT_NAMESPACE
 
diff --git a/Engine/Source/Runtime/Engine/Private/PlayerState.cpp b/Engine/Source/Runtime/Engine/Private/PlayerState.cpp
--- a/Engine/Source/Runtime/Engine/Private/PlayerState.cpp
+++ b/Engine/Source/Runtime/Engine/Private/PlayerState.cpp
@@ -11,6 +11,10 @@
 #include "Iris/ReplicationSystem/ReplicationSystem.h"
 #include "Iris/ReplicationSystem/Prioritization/NetObjectPrioritizer.h"
 
+#if UE_WITH_REMOTE_OBJECT_HANDLE
+#include "UObject/RemoteObject.h"
+#endif
+
 #include UE_INLINE_GENERATED_CPP_BY_NAME(PlayerState)
 
 APlayerState::APlayerState(const FObjectInitializer& ObjectInitializer)
@@ -230,6 +234,28 @@ void APlayerState::Reset()
 	ForceNetUpdate();
 }
 
+#if UE_WITH_REMOTE_OBJECT_HANDLE
+FRemoteServerId APlayerState::GetPreferredRemoteServer() const
+{
+	if (!AutoRTFM::IsClosed() && PawnPrivate.GetResidence() == EResidence::Remote)
+	{
+		// We are not in a transaction (probably during push migration) and our Pawn is non-local
+		// So the Pawn must have already migrated, choose that server.
+		FRemoteServerId PawnServerId = UE::RemoteObject::Handle::GetOwnerServerId(PawnPrivate.GetRemoteId());
+		if (PawnServerId.IsValid())
+		{
+			return PawnServerId;
+		}
+	}
+	else if (APawn* LocalPawn = PawnPrivate.Get())
+	{
+		return LocalPawn->GetPreferredRemoteServer();
+	}
+	// No Pawn — stay where we are (don't use spatial delegate, we have no location)
+	return UE::RemoteObject::Handle::GetOwnerServerId(this);
+}
+#endif
+
 FString APlayerState::GetHumanReadableName() const
 {
 	return GetPlayerName();
