DataDriven cvars fix for console variables not working when loading from plugins

#Synced-CL 51652354

diff --git a/Engine/Source/Runtime/Engine/Classes/DataDrivenCVars/DataDrivenCVars.h b/Engine/Source/Runtime/Engine/Classes/DataDrivenCVars/DataDrivenCVars.h
--- a/Engine/Source/Runtime/Engine/Classes/DataDrivenCVars/DataDrivenCVars.h
+++ b/Engine/Source/Runtime/Engine/Classes/DataDrivenCVars/DataDrivenCVars.h
@@ -98,6 +98,7 @@ public:
 
 	//~ Begin UObject Interface
 	ENGINE_API virtual void PostInitProperties() override;
+	ENGINE_API virtual void PostReloadConfig(FProperty* PropertyThatWasLoaded) override;
 
 #if WITH_EDITOR
 	ENGINE_API virtual void PostEditChangeProperty(FPropertyChangedEvent& PropertyChangedEvent) override;
diff --git a/Engine/Source/Runtime/Engine/Private/DataDrivenCVars/DataDrivenCVars.cpp b/Engine/Source/Runtime/Engine/Private/DataDrivenCVars/DataDrivenCVars.cpp
--- a/Engine/Source/Runtime/Engine/Private/DataDrivenCVars/DataDrivenCVars.cpp
+++ b/Engine/Source/Runtime/Engine/Private/DataDrivenCVars/DataDrivenCVars.cpp
@@ -107,6 +107,20 @@ void UDataDrivenConsoleVariableSettings::PostInitProperties()
 	}
 }
 
+void UDataDrivenConsoleVariableSettings::PostReloadConfig(FProperty* PropertyThatWasLoaded)
+{
+	Super::PostReloadConfig(PropertyThatWasLoaded);
+
+	// Re-register all CVars after a config reload (e.g. from a plugin INI dynamic layer).
+	// PostInitProperties only runs once at construction, so new entries added via plugin INIs
+	// would otherwise never have Register() called on them.
+	// Register() is safe to call on already-registered CVars since it skips creation if the CVar exists.
+	for (FDataDrivenConsoleVariable& CVar : CVarsArray)
+	{
+		CVar.Register();
+	}
+}
+
 void UDataDrivenConsoleVariableSettings::OnDataDrivenChange(IConsoleVariable* CVar)
 {
 	if (GEngine != nullptr)
