Bag of new material translation fixes to get closer to supporting the M_4Layers material.

Major fixes are:
- FunctionHLSLs used in custom outputs are now automatically deduplicated, and only emitter if a call to a function is analyzed (proven useful).
- Fixed Material UB referencing vertex or non-vertex identifiers. MaterialVS doesn't appear to be declared in the shader, will fix in followup change.
- MaterialAttributes can now be branched upon, but it's the user's responsibility to handle the PixelNormalWS reading cycle.
- VertexInterpolators are now automatically deduped.
- Utils in MaterialIRCommon.h moved to MaterialIRUtils.h/cpp for better code organization.

#rnx
#Synced-CL 51637005

diff --git a/Engine/Config/BaseMaterialExpressions.ini b/Engine/Config/BaseMaterialExpressions.ini
--- a/Engine/Config/BaseMaterialExpressions.ini
+++ b/Engine/Config/BaseMaterialExpressions.ini
@@ -1041,8 +1041,8 @@ bShaderInputData=true
 +ExternalCodeDeclarations=(ReturnType=Float4,			Name="DecalColor",						Definition="DecalColor()", Derivative=Zero, ShaderFrequency=Pixel, Domains=(MD_DeferredDecal))
 +ExternalCodeDeclarations=(ReturnType=Float,			Name="DecalLifetimeOpacity",			Definition="DecalLifetimeOpacity()", Derivative=Zero, ShaderFrequency=Pixel, Domains=(MD_DeferredDecal))
 +ExternalCodeDeclarations=(ReturnType=Float,			Name="DeltaTime",						Definition="View.DeltaTime", Derivative=Zero, bIsInlined=true)
-+ExternalCodeDeclarations=(ReturnType=Float,			Name="EyeAdaptation",					Definition="EyeAdaptationLookup()", Derivative=Zero, ShaderFrequency=Pixel, bIsInlined=true, EnvironmentDefines=((Name="USES_EYE_ADAPTATION")))
-+ExternalCodeDeclarations=(ReturnType=Float,			Name="EyeAdaptationInverse",			Definition="EyeAdaptationInverseLookup($0)", ShaderFrequency=Pixel, bIsInlined=true, EnvironmentDefines=((Name="USES_EYE_ADAPTATION")))
++ExternalCodeDeclarations=(ReturnType=Float,			Name="EyeAdaptation",					Definition="EyeAdaptationLookup()", Derivative=Zero, ShaderFrequency=Pixel|Compute, bIsInlined=true, EnvironmentDefines=((Name="USES_EYE_ADAPTATION")))
++ExternalCodeDeclarations=(ReturnType=Float,			Name="EyeAdaptationInverse",			Definition="EyeAdaptationInverseLookup($0)", ShaderFrequency=Pixel|Compute, bIsInlined=true, EnvironmentDefines=((Name="USES_EYE_ADAPTATION")))
 +ExternalCodeDeclarations=(ReturnType=Float4,			Name="FontSignedDistanceData",			Definition="GetFontSignedDistanceData(Parameters)", ShaderFrequency=Pixel, Domains=(MD_UI))
 +ExternalCodeDeclarations=(ReturnType=Float,			Name="IsFirstPerson",					Definition="(IsFirstPerson(Parameters) ? 1.0f : 0.0f)", Derivative=Zero, bIsInlined=true, Domains=(MD_Surface, MD_Volume))
 +ExternalCodeDeclarations=(ReturnType=Float,			Name="IsOrthographic",					Definition="IsOrthoProjectionFloat()", Derivative=Zero, bIsInlined=true)
@@ -1131,7 +1131,7 @@ bShaderInputData=true
 +ExternalCodeDeclarations=(ReturnType=LWCMatrix,	Name="GetInstanceToWorldDF",		Definition="Get<PREV>InstanceToWorldDF(Parameters)", Derivative=Zero, bIsInlined=true, EnvironmentDefines=((Name="NEEDS_INSTANCE_LOCAL_TO_WORLD_PS", ShaderFrequency=Pixel)))
 +ExternalCodeDeclarations=(ReturnType=LWCMatrix,	Name="GetLocalToWorldDF",			Definition="Get<PREV>LocalToWorldDF(Parameters)", Derivative=Zero, bIsInlined=true)
 +ExternalCodeDeclarations=(ReturnType=LWCMatrix,	Name="GetParticleToWorld",			Definition="GetParticleToWorld(Parameters)", Derivative=Zero, bIsInlined=true, EnvironmentDefines=((Name="NEEDS_PARTICLE_LOCAL_TO_WORLD")))
-+ExternalCodeDeclarations=(ReturnType=Float3,		Name="GetPreViewTranslation",		Definition="Get<PREV>PreViewTranslation(Parameters)", Derivative=Zero, bIsInlined=true)
++ExternalCodeDeclarations=(ReturnType=LWCVector3,	Name="GetPreViewTranslation",		Definition="Get<PREV>PreViewTranslation(Parameters)", Derivative=Zero, bIsInlined=true)
 +ExternalCodeDeclarations=(ReturnType=LWCMatrix,	Name="GetWorldToInstance",			Definition="GetWorldToInstance(Parameters)", Derivative=Zero, bIsInlined=true, EnvironmentDefines=((Name="NEEDS_INSTANCE_WORLD_TO_LOCAL_PS", ShaderFrequency=Pixel)))
 +ExternalCodeDeclarations=(ReturnType=LWCMatrix,	Name="GetWorldToInstanceDF",		Definition="GetWorldToInstanceDF(Parameters)", Derivative=Zero, bIsInlined=true, EnvironmentDefines=((Name="NEEDS_INSTANCE_WORLD_TO_LOCAL_PS", ShaderFrequency=Pixel)))
 +ExternalCodeDeclarations=(ReturnType=LWCMatrix,	Name="GetWorldToLocal",				Definition="Get<PREV>WorldToLocal(Parameters)", Derivative=Zero, bIsInlined=true, Domains=(MD_Surface, MD_Volume))
diff --git a/Engine/Source/Runtime/Engine/Private/Materials/MaterialExpressionsToMIR.cpp b/Engine/Source/Runtime/Engine/Private/Materials/MaterialExpressionsToMIR.cpp
--- a/Engine/Source/Runtime/Engine/Private/Materials/MaterialExpressionsToMIR.cpp
+++ b/Engine/Source/Runtime/Engine/Private/Materials/MaterialExpressionsToMIR.cpp
@@ -2693,11 +2693,9 @@ static FValueRef EmitTransformVector(MIR::FEmitter& Em, FValueRef InputValue, EM
 		FValueRef Zero = Em.ConstantZero(MIR::EScalarKind::Float);
 		InputValue = Em.Vector3(Zero, Zero, InputValue);
 	}
-	else
-	{
-		InputValue = Em.CastToVector(InputValue, 3);
-	}
-
+	
+	// Cast the input value to a float3 then make sure it's valid.
+	InputValue = Em.CastToVector(InputValue, 3);
 	if (!InputValue.IsValid())
 	{
 		return InputValue.ToPoison();
@@ -3462,7 +3460,7 @@ static FValueRef EmitActorOrObjectPosition(MIR::FEmitter& Em, EActorOrObjectPosi
 		{
 			if (Printer.Differential != MIR::EExternDifferential::None)
 			{
-				Printer << (bIsAbsolute ? TEXTVIEW("WSPromote(MaterialFloat3(0))") : TEXTVIEW("MaterialFloat3(0)"));
+				Printer << (bIsAbsolute ? TEXTVIEW("WSPromote((MaterialFloat3)0)") : TEXTVIEW("(MaterialFloat3)0"));
 				return;
 			}
 			// Get<Previous><Actor|Object><Translated>WorldPosition(Parameters)
@@ -4565,7 +4563,7 @@ void UMaterialExpressionCustom::Build(MIR::FEmitter& Em)
 	UE_MIR_CHECKPOINT(Em); // Make sure checks above did not fail
 
 	// Declare the HLSL function with the description we generated
-	const MIR::FFunction* Func = Em.FunctionHLSL(FuncDesc);
+	MIR::FFunction* Func = Em.FunctionHLSL(FuncDesc);
 	
 	UE_MIR_CHECKPOINT(Em); // To guarantee a function was succesfully emitted.
 
diff --git a/Engine/Source/Runtime/Engine/Private/Materials/MaterialIR.cpp b/Engine/Source/Runtime/Engine/Private/Materials/MaterialIR.cpp
--- a/Engine/Source/Runtime/Engine/Private/Materials/MaterialIR.cpp
+++ b/Engine/Source/Runtime/Engine/Private/Materials/MaterialIR.cpp
@@ -913,6 +913,11 @@ void FStageSwitch::SetArgs(FValue* PixelStageArg, FValue* OtherStagesArg)
 	}
 }
 
+uint32 GetTypeHash(const FFunctionParameter& F)
+{
+	return HashCombine(GetTypeHash(F.Name), GetTypeHash(F.Type));
+}
+
 bool FFunction::Equals(const FFunction* Other) const
 {
 	return Kind == Other->Kind
@@ -921,7 +926,36 @@ bool FFunction::Equals(const FFunction* Other) const
 		&& NumInputOnlyParams == Other->NumInputOnlyParams
 		&& NumInputAndOutputParams == Other->NumInputAndOutputParams
 		&& NumParameters == Other->NumParameters
-		&& FMemory::Memcmp(Parameters, Other->Parameters, NumParameters * sizeof(FFunctionParameter)) == 0;
+		&& CompareItems(Parameters, Other->Parameters, NumParameters);
+}
+
+uint32 GetTypeHash(const FFunction& F)
+{
+	uint32 Hash = GetTypeHash(F.Name);
+	Hash = HashCombine(Hash, ::GetTypeHash((uint8)F.Kind));
+	Hash = HashCombine(Hash, ::GetTypeHash(F.NumInputOnlyParams));
+	Hash = HashCombine(Hash, ::GetTypeHash(F.NumInputAndOutputParams));
+	Hash = HashCombine(Hash, ::GetTypeHash(F.NumParameters));
+	Hash = HashCombine(Hash, GetTypeHash(F.ReturnType));
+	Hash = HashCombine(Hash, ArrayViewHash(TConstArrayView<FFunctionParameter>{ F.Parameters, F.NumParameters }));
+	return Hash;
+}
+
+bool FFunctionHLSL::Equals(const FFunctionHLSL* Other) const
+{
+	return FFunction::Equals(Other)
+		&& Code == Other->Code
+		&& MIR::ArrayViewEquals(Defines, Other->Defines)
+		&& MIR::ArrayViewEquals(Includes, Other->Includes);
+}
+
+uint32 GetTypeHash(const FFunctionHLSL& F)
+{
+	uint32 Hash = GetTypeHash(static_cast<const FFunction&>(F));
+	Hash = HashCombine(Hash, GetTypeHash(F.Code));
+	Hash = HashCombine(Hash, ArrayViewHash(F.Defines));
+	Hash = HashCombine(Hash, ArrayViewHash(F.Includes));
+	return Hash;
 }
 
 int32 FPrimitiveUniform::Analysis_NumComponents() const
diff --git a/Engine/Source/Runtime/Engine/Private/Materials/MaterialIREmitter.cpp b/Engine/Source/Runtime/Engine/Private/Materials/MaterialIREmitter.cpp
--- a/Engine/Source/Runtime/Engine/Private/Materials/MaterialIREmitter.cpp
+++ b/Engine/Source/Runtime/Engine/Private/Materials/MaterialIREmitter.cpp
@@ -2947,13 +2947,6 @@ FValueRef FEmitter::Branch(FValueRef Condition, FValueRef True, FValueRef False)
 		return Poison();
 	}
 
-	// Branching on material attributes isn't supported yet.
-	if (CommonType.AsAggregate() == MaterialAttributesAggregate::Get())
-	{
-		Error(TEXT("Dynamic branching on MaterialAttributes is not yet supported due to the Normal attribute being evaluated before all other attributes. Branching on any other user defined aggregate is supported."));
-		return Poison();
-	}
-
 	True = Cast(True, CommonType);
 	False = Cast(False, CommonType);
 	if (!True || !False)
@@ -4175,14 +4168,14 @@ FValueRef FEmitter::VertexInterpolator(FValueRef VertexValue)
 	}
 
 	// Create the set/get pair of vertex interpolator instructions
-	FSetVertexInterpolator* SetInstr = EmitNew<FSetVertexInterpolator>(*this, FType::MakeVoid());
-	SetInstr->Arg = VertexValue;
+	FSetVertexInterpolator SetInstrProto = MakePrototype<FSetVertexInterpolator>(FType::MakeVoid());
+	SetInstrProto.Arg = VertexValue;
 
-	FGetVertexInterpolator* GetInstr = EmitNew<FGetVertexInterpolator>(*this, VertexValue->Type);
-	GetInstr->SetSubgraphProperties(EGraphProperties::PixelStageOnly);
-	GetInstr->SetInstr = SetInstr;
+	FGetVertexInterpolator GetInstrProto = MakePrototype<FGetVertexInterpolator>(VertexValue->Type);
+	GetInstrProto.SetSubgraphProperties(EGraphProperties::PixelStageOnly);
+	GetInstrProto.SetInstr = EmitPrototype(*this, SetInstrProto)->As<FSetVertexInterpolator>();
 
-	return GetInstr;
+	return EmitPrototype(*this, GetInstrProto);
 }
 
 static void SetFunctionFromDesc(FFunction& Function, FFunctionKind Kind, const FFunctionDesc& Desc)
@@ -4193,7 +4186,7 @@ static void SetFunctionFromDesc(FFunction& Function, FFunctionKind Kind, const F
 	Function.NumInputAndOutputParams = Desc.NumInputOnlyParams + Desc.NumInputOutputParams;
 	Function.NumParameters = Function.NumInputAndOutputParams + Desc.NumOutputOnlyParams;
 	Function.ReturnType = Desc.ReturnType;
- 
+
 	// Copy over the parameter declarations from the description to the function prototype
 	for (uint32 i = 0; i < Function.NumParameters; ++i)
 	{
@@ -4201,7 +4194,7 @@ static void SetFunctionFromDesc(FFunction& Function, FFunctionKind Kind, const F
 	} 
 }
 
-const FFunction* FEmitter::FunctionHLSL(const FFunctionHLSLDesc& Desc)
+FFunction* FEmitter::FunctionHLSL(const FFunctionHLSLDesc& Desc)
 {
 	check(Desc.GetNumParameters() <= MIR::MaxNumFunctionParameters);
 
@@ -4210,22 +4203,30 @@ const FFunction* FEmitter::FunctionHLSL(const FFunctionHLSLDesc& Desc)
 	Prototype.Code = Desc.Code;
 	Prototype.Defines = Desc.Defines;
 	Prototype.Includes = Desc.Includes;
-	
+
+	// Return existing function if an equivalent one was already emitted.
+	if (FFunctionHLSL* const* Found = FunctionHLSLSet.Find(&Prototype))
+	{
+		return *Found;
+	}
+
 	// Create the new MIR HLSL function instance and set it up
 	FFunctionHLSL* Function = new (FPrivate::Allocate(*this, sizeof(FFunctionHLSL), alignof(FFunctionHLSL))) FFunctionHLSL { Prototype };
 	Function->Name = Module->InternString(Function->Name);
 	Function->Code = Module->InternString(Desc.Code);
 	Function->Defines = MakeArrayCopy(*this, Desc.Defines);
 	Function->Includes = MakeArrayCopy(*this, Desc.Includes);
-	Function->UniqueId = Module->FunctionHLSLs.Num();
-
-	// Add it to the module list
-	Module->FunctionHLSLs.Push(Function);
 
+	FunctionHLSLSet.Add(Function);
 	return Function;
 }
 
-FValueRef FEmitter::Call(const FFunction* Function, TConstArrayView<FValueRef> InputArguments)
+uint32 FEmitter::FFunctionHLSLKeyFuncs::GetKeyHash(const FFunctionHLSL* Key)
+{
+	return GetTypeHash(*Key);
+}
+
+FValueRef FEmitter::Call(FFunction* Function, TConstArrayView<FValueRef> InputArguments)
 {
 	if (!Function)
 	{
@@ -4626,34 +4627,6 @@ bool FEmitter::FValueKeyFuncs::Matches(KeyInitType A, KeyInitType B)
 	return A->Equals(B);
 }
 
-static uint32 HashBytes(const void* InPtr, uint32 Size)
-{
-	check((uintptr_t)InPtr % 4 == 0); // Require the pointer to be 4-byte aligned
-
-	uint32 Hash = 0;
-
-	const uint32* WordPtr = (const uint32*)InPtr;
-	uint32 NumWords = Size / 4;
-	for (uint32 i = 0; i < NumWords; ++i)
-	{
-		Hash = HashCombineFast(Hash, WordPtr[i]);
-	}
-
-	const uint8* EndPtr = (const uint8*)(WordPtr + NumWords);
-
-	uint32 End = 0;
-	switch (Size % 4)
-	{
-		case 3: End |= uint32(EndPtr[2]) << 16; // fallthrough
-		case 2: End |= uint32(EndPtr[1]) << 8;  // fallthrough
-		case 1: End |= uint32(EndPtr[0]);       // fallthrough
-			Hash = HashCombineFast(Hash, End);
-		default: break;
-	}
-
-	return Hash;
-}
-
 uint32 FEmitter::FValueKeyFuncs::GetKeyHash(KeyInitType Key)
 {
 	return MIR::HashBytes(Key, Key->GetSizeInBytes());
diff --git a/Engine/Source/Runtime/Engine/Private/Materials/MaterialIRModule.cpp b/Engine/Source/Runtime/Engine/Private/Materials/MaterialIRModule.cpp
--- a/Engine/Source/Runtime/Engine/Private/Materials/MaterialIRModule.cpp
+++ b/Engine/Source/Runtime/Engine/Private/Materials/MaterialIRModule.cpp
@@ -116,6 +116,11 @@ int32 FMaterialIRModule::FindOrAddParameterCollection(UMaterialParameterCollecti
 	return CollectionIndex;
 }
 
+void FMaterialIRModule::AddFunctionHLSL(const MIR::FFunctionHLSL* Function)
+{
+	FunctionHLSLs.Add(Function);
+}
+
 void FMaterialIRModule::AddShadingModel(EMaterialShadingModel InShadingModel)
 {
 	ShadingModelsFromCompilation.AddShadingModel(InShadingModel);
@@ -126,4 +131,5 @@ bool FMaterialIRModule::IsMaterialPropertyUsed(EMaterialProperty InProperty) con
 	return PropertyValues[InProperty] && !PropertyValues[InProperty]->EqualsConstant(FMaterialAttributeDefinitionMap::GetDefaultValue(InProperty));
 }
 
+
 #endif // #if WITH_EDITOR
diff --git a/Engine/Source/Runtime/Engine/Private/Materials/MaterialIRToHLSLTranslator.cpp b/Engine/Source/Runtime/Engine/Private/Materials/MaterialIRToHLSLTranslator.cpp
--- a/Engine/Source/Runtime/Engine/Private/Materials/MaterialIRToHLSLTranslator.cpp
+++ b/Engine/Source/Runtime/Engine/Private/Materials/MaterialIRToHLSLTranslator.cpp
@@ -23,6 +23,7 @@
 #include "ParameterCollection.h"
 #include "RenderUtils.h"
 #include "ShaderCore.h"
+#include "ShaderCompiler.h"
 
 #if WITH_EDITOR
 
@@ -574,7 +575,15 @@ struct FEntryPointCodeGen
 		if ((ExternArg->GetInfo().Flags & MIR::EExternFlags::ZeroDifferentials) != MIR::EExternFlags::None)
 		{
 			check(ExternDifferential->Type.IsPrimitive());
-			Printer << ExternDifferential->Type << TEXT("(0)");
+			if (ExternDifferential->Type.IsDouble())
+			{
+				// Example: WSPromote((MaterialFloat3)0)
+				Printer << TEXT("WSPromote((") << ExternDifferential->Type.GetPrimitive().ToScalarKind(MIR::EScalarKind::Float) << TEXT(")0)");
+			}
+			else
+			{
+				Printer << TEXT('(') << ExternDifferential->Type << TEXT(")0");
+			}
 			return;
 		}
 
@@ -586,12 +595,12 @@ struct FEntryPointCodeGen
 
 	void LowerTextureUniform(const MIR::FTextureUniform* TextureUniform)
 	{
-		Printer << TEXTVIEW("Material.") << LowerTextureMaterialType(TextureUniform->Texture->GetMaterialType(), true) << TEXTVIEW("_") << TextureUniform->Analysis_UniformIndex;
+		Printer << GetMaterialUBName() << TEXTVIEW(".") << LowerTextureMaterialType(TextureUniform->Texture->GetMaterialType(), true) << TEXTVIEW("_") << TextureUniform->Analysis_UniformIndex;
 	}
 
 	void LowerVirtualTextureUniform(const MIR::FVirtualTextureUniform* VirtualTextureUniform)
 	{
-		Printer << TEXTVIEW("Material.") << LowerTextureMaterialType(MCT_TextureVirtual, true) << TEXTVIEW("_") << VirtualTextureUniform->Analysis_UniformIndex;
+		Printer << GetMaterialUBName() << TEXTVIEW(".") << LowerTextureMaterialType(MCT_TextureVirtual, true) << TEXTVIEW("_") << VirtualTextureUniform->Analysis_UniformIndex;
 	}
 
 	void EmitPreshaderBufferReadDoubleVector(uint32 GlobalComponentOffset, int32 SourceComponents, int32 UsedComponents)
@@ -602,6 +611,7 @@ struct FEntryPointCodeGen
 
 		// Index of the float4 slot.
 		const uint32 BufferSlotIndex = GlobalComponentOffset / 4;
+		const TCHAR* UBName = GetMaterialUBName();
 
 		if (UsedComponents == 1)
 		{
@@ -619,40 +629,40 @@ struct FEntryPointCodeGen
 		{
 		case 1:
 			// Emit two individual components
-			Printer << TEXTVIEW("(Material.PreshaderBuffer[") << BufferSlotIndex << TEXTVIEW("].") << Components[GlobalComponentOffset % 4];
-			Printer << TEXTVIEW(", Material.PreshaderBuffer[") << BufferSlotIndex << TEXTVIEW("].") << Components[GlobalComponentOffset % 4 + 1] << TEXTVIEW("))");
+			Printer << TEXTVIEW("(") << UBName << TEXTVIEW(".PreshaderBuffer[") << BufferSlotIndex << TEXTVIEW("].") << Components[GlobalComponentOffset % 4];
+			Printer << TEXTVIEW(", ") << UBName << TEXTVIEW(".PreshaderBuffer[") << BufferSlotIndex << TEXTVIEW("].") << Components[GlobalComponentOffset % 4 + 1] << TEXTVIEW("))");
 			break;
 		case 2:
 			// Emit XY and ZW, or just X and Z
-			Printer << TEXTVIEW("(Material.PreshaderBuffer[") << BufferSlotIndex << (UsedComponents == 2 ? TEXTVIEW("].xy") : TEXTVIEW("].x"));
-			Printer << TEXTVIEW(", Material.PreshaderBuffer[") << BufferSlotIndex << (UsedComponents == 2 ? TEXTVIEW("].zw))") : TEXTVIEW("].z))"));
+			Printer << TEXTVIEW("(") << UBName << TEXTVIEW(".PreshaderBuffer[") << BufferSlotIndex << (UsedComponents == 2 ? TEXTVIEW("].xy") : TEXTVIEW("].x"));
+			Printer << TEXTVIEW(", ") << UBName << TEXTVIEW(".PreshaderBuffer[") << BufferSlotIndex << (UsedComponents == 2 ? TEXTVIEW("].zw))") : TEXTVIEW("].z))"));
 			break;
 		case 3:
 			if (UsedComponents == 3)
 			{
 				// Emit XYZ and float3(W, XY)
-				Printer << TEXTVIEW("(Material.PreshaderBuffer[") << BufferSlotIndex << TEXTVIEW("].xyz");
-				Printer << TEXTVIEW(", float3(Material.PreshaderBuffer[") << BufferSlotIndex << TEXTVIEW("].w");
-				Printer << TEXTVIEW(", Material.PreshaderBuffer[") << BufferSlotIndex + 1 << TEXTVIEW("].xy)))");
+				Printer << TEXTVIEW("(") << UBName << TEXTVIEW(".PreshaderBuffer[") << BufferSlotIndex << TEXTVIEW("].xyz");
+				Printer << TEXTVIEW(", float3(") << UBName << TEXTVIEW(".PreshaderBuffer[") << BufferSlotIndex << TEXTVIEW("].w");
+				Printer << TEXTVIEW(", ") << UBName << TEXTVIEW(".PreshaderBuffer[") << BufferSlotIndex + 1 << TEXTVIEW("].xy)))");
 			}
 			else if (UsedComponents == 2)
 			{
 				// Emit XY and float2(W, X)
-				Printer << TEXTVIEW("(Material.PreshaderBuffer[") << BufferSlotIndex << TEXTVIEW("].xy");
-				Printer << TEXTVIEW(", float2(Material.PreshaderBuffer[") << BufferSlotIndex << TEXTVIEW("].w");
-				Printer << TEXTVIEW(", Material.PreshaderBuffer[") << BufferSlotIndex + 1 << TEXTVIEW("].x)))");
+				Printer << TEXTVIEW("(") << UBName << TEXTVIEW(".PreshaderBuffer[") << BufferSlotIndex << TEXTVIEW("].xy");
+				Printer << TEXTVIEW(", float2(") << UBName << TEXTVIEW(".PreshaderBuffer[") << BufferSlotIndex << TEXTVIEW("].w");
+				Printer << TEXTVIEW(", ") << UBName << TEXTVIEW(".PreshaderBuffer[") << BufferSlotIndex + 1 << TEXTVIEW("].x)))");
 			}
 			else
 			{
 				// Emit X and W
-				Printer << TEXTVIEW("(Material.PreshaderBuffer[") << BufferSlotIndex << TEXTVIEW("].x");
-				Printer << TEXTVIEW(", Material.PreshaderBuffer[") << BufferSlotIndex << TEXTVIEW("].w))");
+				Printer << TEXTVIEW("(") << UBName << TEXTVIEW(".PreshaderBuffer[") << BufferSlotIndex << TEXTVIEW("].x");
+				Printer << TEXTVIEW(", ") << UBName << TEXTVIEW(".PreshaderBuffer[") << BufferSlotIndex << TEXTVIEW("].w))");
 			}
 			break;
 		case 4:
 			// Emit used components
-			Printer << TEXTVIEW("(Material.PreshaderBuffer[") << BufferSlotIndex << TEXTVIEW("].") << ComponentSubsets[UsedComponents - 1];
-			Printer << TEXTVIEW(", Material.PreshaderBuffer[") << BufferSlotIndex + 1 << TEXTVIEW("].") << ComponentSubsets[UsedComponents - 1] << TEXTVIEW("))");
+			Printer << TEXTVIEW("(") << UBName << TEXTVIEW(".PreshaderBuffer[") << BufferSlotIndex << TEXTVIEW("].") << ComponentSubsets[UsedComponents - 1];
+			Printer << TEXTVIEW(", ") << UBName << TEXTVIEW(".PreshaderBuffer[") << BufferSlotIndex + 1 << TEXTVIEW("].") << ComponentSubsets[UsedComponents - 1] << TEXTVIEW("))");
 			break;
 		}
 	}
@@ -672,7 +682,7 @@ struct FEntryPointCodeGen
 			Printer << TEXTVIEW("asint(");
 		}
 
-		Printer << TEXTVIEW("Material.PreshaderBuffer[") << BufferSlotIndex << TEXTVIEW("]");
+		Printer << GetMaterialUBName() << TEXTVIEW(".PreshaderBuffer[") << BufferSlotIndex << TEXTVIEW("]");
 
 		if (NumComponents < 4)
 		{
@@ -1487,16 +1497,16 @@ struct FEntryPointCodeGen
 			// The wrap/clamp/mirror operation is handled in the shader explicitly.
 			// This generates: GetMaterialSharedSampler(Material.VirtualTexturePhysical_<VirtualTextureIndex>Sampler, <SharedSamplerName>)
 			FStringView SharedSamplerName = TextureRead->bUseAnisoSampler ? TEXTVIEW("View.SharedBilinearAnisoClampedSampler") : TEXTVIEW("View.SharedBilinearClampedSampler");
-			Printer << TEXTVIEW("GetMaterialSharedSampler(Material.VirtualTexturePhysical_") << VirtualTextureIndex << TEXTVIEW("Sampler, ") << SharedSamplerName << TEXTVIEW(")");
+			Printer << TEXTVIEW("GetMaterialSharedSampler(") << GetMaterialUBName() << TEXTVIEW(".VirtualTexturePhysical_") << VirtualTextureIndex << TEXTVIEW("Sampler, ") << SharedSamplerName << TEXTVIEW(")");
 		}
 		else
 		{
-			Printer << TEXTVIEW("Material.VirtualTexturePhysical_") << VirtualTextureIndex << TEXTVIEW("Sampler");
+			Printer << GetMaterialUBName() << TEXTVIEW(".VirtualTexturePhysical_") << VirtualTextureIndex << TEXTVIEW("Sampler");
 		}
 
 		Printer << ListSeparator << LowerValue(TextureRead->VTPageTable)
 			<< ListSeparator << VTPageTableResult->VTPageTableIndex
-			<< ListSeparator << TEXTVIEW("VTUniform_Unpack(") << TEXTVIEW("Material.VTPackedUniform[") << VirtualTextureIndex << TEXTVIEW("]") << TEXTVIEW(")")
+			<< ListSeparator << TEXTVIEW("VTUniform_Unpack(") << GetMaterialUBName() << TEXTVIEW(".VTPackedUniform[") << VirtualTextureIndex << TEXTVIEW("]") << TEXTVIEW(")")
 			<< EndArgs;
 	}
 
@@ -1806,6 +1816,12 @@ struct FEntryPointCodeGen
 		Printer << TEXT('_') << ParamIndex;
 	}
 	
+	// Returns the name of the material uniform buffer for the current stage.
+	const TCHAR* GetMaterialUBName() const
+	{
+		return CurrentStage == MIR::Stage_Vertex ? TEXT("__Material_UB_VS__") : TEXT("__Material_UB__");
+	}
+
 	void LowerSubstrateDefaultSlab(const MIR::FSubstrateDefaultSlab* SubstrateParameter)
 	{
 		Printer << TEXTVIEW("GetInitialisedSubstrateData()");
@@ -2256,43 +2272,58 @@ static void GenerateCustomFunctionsHLSL(const FMaterialIRModule* Module, TMap<FS
 			Printer << TEXTVIEW("#include \"") << Include << TEXT('\"') << NewLine;
 		}
 
-		// Write the custom function signature, for example "C5_MyCustomNode".
-		// - 'C' is only a "namespace" for custom functions
-		// - '5' is a unique id used to disambiguate distinct custom functions with the same name.
-		Printer << Function->ReturnType << TEXTVIEW(" C") << Function->UniqueId << TEXT('_') << Function->Name << BeginArgs << TEXT("FMaterialPixelParameters Parameters");
-			
-		// Write the parameter declarations
-		for (uint16 i = 0; i < Function->NumParameters; ++i)
+		// Emit one overload per stage so that the correct Parameters type (FMaterialVertexParameters vs
+		// FMaterialPixelParameters) is used. When a function is called from multiple stages, HLSL overload
+		// resolution at the call site picks the right overload based on the type of 'Parameters' in scope.
+		const MIR::EStage Stages[] = { MIR::Stage_Vertex, MIR::Stage_Pixel }; // Compute still uses FMaterialPixelParameters
+		for (int StageIndex = 0; StageIndex < UE_ARRAY_COUNT(Stages); ++StageIndex)
 		{
-			Printer << TEXTVIEW(", ");
+			MIR::EStage Stage = Stages[StageIndex];
+			if (!MIR::IsStageInMask(Stage, Function->Analysis_CallStageMask))
+			{
+				continue;
+			}
 
-			// Print the io keyword
-			Printer << (i < Function->NumInputOnlyParams ? TEXTVIEW("")
-			: i < Function->NumInputAndOutputParams ? TEXTVIEW("inout ")
-			: TEXTVIEW("out "));
+			// Write the custom function signature, for example "C5_MyCustomNode".
+			// - 'C' is only a "namespace" for custom functions
+			// - '5' is a unique id used to disambiguate distinct custom functions with the same name.
+			Printer << Function->ReturnType << TEXTVIEW(" C") << Function->UniqueId << TEXT('_') << Function->Name << BeginArgs;
+			Printer << (Stage == MIR::Stage_Vertex ? TEXTVIEW("FMaterialVertexParameters") : TEXTVIEW("FMaterialPixelParameters"));
+			Printer << TEXTVIEW(" Parameters");
 
-			// Type and name
-			Printer << Function->Parameters[i].Type << TEXT(' ') << *Function->Parameters[i].Name.ToString();
-		}
+			// Write the parameter declarations
+			for (uint16 i = 0; i < Function->NumParameters; ++i)
+			{
+				Printer << TEXTVIEW(", ");
 
-		Printer << EndArgs << OpenBlock;
-			
-		// If the function does not contain a "return" keyword, add one.
-		bool bContainsReturn = Function->Code.Contains(TEXT("return"));
-		if (!bContainsReturn)
-		{
-			Printer << TEXTVIEW("return") << NewLine;
-		}
-			
-		// Write the function code
-		Printer << Function->Code << NewLine;
+				// Print the io keyword
+				Printer << (i < Function->NumInputOnlyParams ? TEXTVIEW("")
+				: i < Function->NumInputAndOutputParams ? TEXTVIEW("inout ")
+				: TEXTVIEW("out "));
 
-		if (!bContainsReturn)
-		{
-			Printer << TEXT(';');
-		}
+				// Type and name
+				Printer << Function->Parameters[i].Type << TEXT(' ') << *Function->Parameters[i].Name.ToString();
+			}
+
+			Printer << EndArgs << OpenBlock;
+
+			// If the function does not contain a "return" keyword, add one.
+			bool bContainsReturn = Function->Code.Contains(TEXT("return"));
+			if (!bContainsReturn)
+			{
+				Printer << TEXTVIEW("return") << NewLine;
+			}
 
-		Printer << NewLine << CloseBlock << NewLine << NewLine;
+			// Write the function code
+			Printer << Function->Code << NewLine;
+
+			if (!bContainsReturn)
+			{
+				Printer << TEXT(';');
+			}
+
+			Printer << NewLine << CloseBlock << NewLine << NewLine;
+		}
 	}
 
 	OutParameters.Add(TEXT("custom_functions"), MoveTemp(Printer.Buffer));
@@ -2367,12 +2398,36 @@ static void GenerateOtherTemplateStringParameters(const FMaterial* Material, con
 	// user_scene_texture_remap
 	OutParameters.Add(TEXT("user_scene_texture_remap"), UE::MaterialTranslatorUtils::GenerateUserSceneTextureRemapHLSLDefines(Module->GetCompilationOutput()));
 	
-	if (Substrate::IsSubstrateEnabled())
 	{
-		// No one else write to uniform_material_expressions so no need to search before append.
-		OutParameters.Add(TEXT("uniform_material_expressions"), SubstrateData.TreeHLSL);
-		// SUBSTRATE_TODO also add in there DescriptionStringCommentForDebug
-	}  
+		FString UniformMaterialExpressions = TEXT("#define __Material_UB__ Material" HLSL_LINE_TERMINATOR);
+		if (UseMaterialVSUniformBuffer())
+		{
+			UniformMaterialExpressions += TEXT(
+				"#if (VERTEXSHADER || MESHSHADER || (SHADER_USES_ONLY_VERTEX_DATA && SHADER_SUPPORTS_MATERIAL_VS_UB))" HLSL_LINE_TERMINATOR
+				"#define __Material_UB_VS__ MaterialVS" HLSL_LINE_TERMINATOR
+				"#define USE_MATERIAL_VS_UB 1" HLSL_LINE_TERMINATOR
+				"#else" HLSL_LINE_TERMINATOR
+				"#define __Material_UB_VS__ Material" HLSL_LINE_TERMINATOR
+				"#define USE_MATERIAL_VS_UB 0" HLSL_LINE_TERMINATOR
+				"#endif" HLSL_LINE_TERMINATOR);
+		}
+		else
+		{
+			// When we want to compile a Nanite compute shader which is also calling SF_Vertex level 
+			// functions, we actually remap __Material_UB_VS__ to the generic Material UB. 
+			// Since the indices are the same for VS-accessed uniforms, we only need to replace the UB name
+			UniformMaterialExpressions += TEXT("#define __Material_UB_VS__ Material" HLSL_LINE_TERMINATOR);
+			UniformMaterialExpressions += TEXT("#define USE_MATERIAL_VS_UB 0" HLSL_LINE_TERMINATOR);
+		}
+
+		if (Substrate::IsSubstrateEnabled())
+		{
+			UniformMaterialExpressions += SubstrateData.TreeHLSL;
+			// SUBSTRATE_TODO also add in there DescriptionStringCommentForDebug
+		}
+
+		OutParameters.Add(TEXT("uniform_material_expressions"), MoveTemp(UniformMaterialExpressions));
+	}
 }
 
 static const TCHAR* GetShadingModelParameterName(EMaterialShadingModel InModel)
@@ -2625,13 +2680,14 @@ static void GetShaderCompilerEnvironment(const FMaterial* Material, const FMater
 	for (int32 VTStackIndex = 0; VTStackIndex < CompilationOutput.UniformExpressionSet.GetVTStacks().Num(); ++VTStackIndex)
 	{
 		const FMaterialVirtualTextureStack& VTStack = CompilationOutput.UniformExpressionSet.GetVTStack(VTStackIndex);
+		const TCHAR* UniformUBName = (VTStack.ShaderFrequencyMask & (1ULL << ((uint64)SF_Vertex))) ? TEXT("__Material_UB_VS__") : TEXT("__Material_UB__");
 
 		// Setup page table defines to map each VT stack to either 1 or 2 page table textures, depending on how many layers it uses
-		FString PageTableValue = FString::Printf(TEXT("Material.VirtualTexturePageTable0_%d"), VTStackIndex);
+		FString PageTableValue = FString::Printf(TEXT("%s.VirtualTexturePageTable0_%d"), UniformUBName, VTStackIndex);
 		OutEnvironment.SetDefine(*FString::Printf(TEXT("VIRTUALTEXTURE_PAGETABLE_%d"), VTStackIndex), *PageTableValue);
 
 		// Setup page table uniform defines.
-		FString PageTableUniformValue = FString::Printf(TEXT("Material.VTPackedPageTableUniform[%d*2], Material.VTPackedPageTableUniform[%d*2+1]"), VTStackIndex, VTStackIndex);
+		FString PageTableUniformValue = FString::Printf(TEXT("%s.VTPackedPageTableUniform[%d*2], %s.VTPackedPageTableUniform[%d*2+1]"), UniformUBName, VTStackIndex, UniformUBName, VTStackIndex);
 		OutEnvironment.SetDefine(*FString::Printf(TEXT("VIRTUALTEXTURE_PAGETABLE_UNIFORM_%d"), VTStackIndex), *PageTableUniformValue);
 	}
 
diff --git a/Engine/Source/Runtime/Engine/Private/Materials/MaterialIRTypes.cpp b/Engine/Source/Runtime/Engine/Private/Materials/MaterialIRTypes.cpp
--- a/Engine/Source/Runtime/Engine/Private/Materials/MaterialIRTypes.cpp
+++ b/Engine/Source/Runtime/Engine/Private/Materials/MaterialIRTypes.cpp
@@ -338,6 +338,27 @@ bool FType::operator==(FType Other) const
 	}
 }
 
+uint32 GetTypeHash(const FType& T)
+{
+	uint32 Hash = ::GetTypeHash((uint32)T.Kind);
+	if (T.Kind == ETypeKind::Primitive)
+	{
+		static_assert(sizeof(FPrimitive) == sizeof(uint32));
+		union
+		{
+			FPrimitive Primitive;
+			uint32     Uint;
+		} Test;
+		Test.Primitive = T.Primitive;
+		Hash = HashCombine(Hash,::GetTypeHash(Test.Uint));
+	}
+	else if (T.Kind == ETypeKind::Aggregate)
+	{
+		Hash = HashCombine(Hash, ::GetTypeHash(T.Aggregate));
+	}
+	return Hash;
+}
+
 } // namespace MIR
 
 #endif // #if WITH_EDITOR
diff --git a/Engine/Source/Runtime/Engine/Private/Materials/MaterialIRUtils.cpp b/Engine/Source/Runtime/Engine/Private/Materials/MaterialIRUtils.cpp
new file mode 100644
--- /dev/null
+++ b/Engine/Source/Runtime/Engine/Private/Materials/MaterialIRUtils.cpp
@@ -0,0 +1,44 @@
+// Copyright Epic Games, Inc. All Rights Reserved.
+
+#include "Materials/MaterialIRUtils.h"
+
+#if WITH_EDITOR
+
+namespace MIR
+{
+
+uint32 HashBytes(const void* InPtr, uint32 Size)
+{
+	if (!ensure((uintptr_t)InPtr % 4 == 0))
+	{
+		check(false);
+		return 0;
+	}
+
+	uint32 Hash = 0;
+
+	const uint32* WordPtr = (const uint32*)InPtr;
+	uint32 NumWords = Size / 4;
+	for (uint32 i = 0; i < NumWords; ++i)
+	{
+		Hash = HashCombineFast(Hash, WordPtr[i]);
+	}
+
+	const uint8* EndPtr = (const uint8*)(WordPtr + NumWords);
+
+	uint32 End = 0;
+	switch (Size % 4)
+	{
+		case 3: End |= uint32(EndPtr[2]) << 16; // fallthrough
+		case 2: End |= uint32(EndPtr[1]) << 8;  // fallthrough
+		case 1: End |= uint32(EndPtr[0]);        // fallthrough
+			Hash = HashCombineFast(Hash, End);
+		default: break;
+	}
+
+	return Hash;
+}
+
+} // namespace MIR
+
+#endif // #if WITH_EDITOR
diff --git a/Engine/Source/Runtime/Engine/Private/Materials/MaterialIRValueAnalyzer.cpp b/Engine/Source/Runtime/Engine/Private/Materials/MaterialIRValueAnalyzer.cpp
--- a/Engine/Source/Runtime/Engine/Private/Materials/MaterialIRValueAnalyzer.cpp
+++ b/Engine/Source/Runtime/Engine/Private/Materials/MaterialIRValueAnalyzer.cpp
@@ -390,6 +390,23 @@ static void SetVertexInterpolator_Analyze(FMaterialIRValueAnalyzer& Analyzer, MI
 	}
 }
 
+static void Call_Analyze(FMaterialIRValueAnalyzer& Analyzer, MIR::FCall* Call)
+{
+	if (Call->Function->UniqueId == 0)
+	{
+		const MIR::FFunctionHLSL* FunctionHLSL = static_cast<const MIR::FFunctionHLSL*>(Call->Function);
+		Analyzer.Module->AddFunctionHLSL(FunctionHLSL);
+		uint32 UniqueId = Analyzer.Module->GetFunctionHLSLs().Num();
+		check(UniqueId <= UINT16_MAX);
+		Call->Function->UniqueId = UniqueId;
+	}
+}
+
+static void Call_AnalyzeInStage(FMaterialIRValueAnalyzer& Analyzer, MIR::FCall* Call, MIR::EStage Stage)
+{
+	Call->Function->Analysis_CallStageMask |= MIR::StageToMask(Stage);
+}
+
 static void RecurseMarkPreshaderTemp(MIR::FValue* Value)
 {
 	for (MIR::FValue* Use : Value->GetUses())
@@ -650,6 +667,7 @@ void FMaterialIRValueAnalyzer::Analyze(MIR::FValue* Value)
 		EXPAND_CASE(Extern);
 		EXPAND_CASE(GetVertexInterpolator);
 		EXPAND_CASE(SetVertexInterpolator);
+		EXPAND_CASE(Call);
 		default: break;
 	}
 
@@ -679,6 +697,7 @@ void FMaterialIRValueAnalyzer::AnalyzeInStage(MIR::FValue* Value, MIR::EStage St
 		EXPAND_CASE(Extern);
 		EXPAND_CASE(VTPageTableRead);
 		EXPAND_CASE(HardwarePartialDerivative);
+		EXPAND_CASE(Call);
 		default: break;
 	}
 
diff --git a/Engine/Source/Runtime/Engine/Public/Materials/MaterialIR.h b/Engine/Source/Runtime/Engine/Public/Materials/MaterialIR.h
--- a/Engine/Source/Runtime/Engine/Public/Materials/MaterialIR.h
+++ b/Engine/Source/Runtime/Engine/Public/Materials/MaterialIR.h
@@ -973,7 +973,7 @@ struct FSetVertexInterpolator : TInstruction<VK_SetVertexInterpolator, 1>
 };
 
 // Kind of user-specified function.
-enum class FFunctionKind
+enum class FFunctionKind : uint8
 {
 	HLSL,
 };
@@ -986,6 +986,10 @@ struct FFunctionParameter
 	
 	// The parameter type.
 	FType Type;
+
+	bool operator==(const FFunctionParameter&) const = default;
+
+	friend uint32 GetTypeHash(const FFunctionParameter& F);
 };
 
 // Maximum number of parameters a user function can have.
@@ -997,6 +1001,9 @@ struct FFunction
 	// The kind of function this is.
 	FFunctionKind Kind;
 
+	// Mask of stages this function is called in.
+	EStageMask Analysis_CallStageMask;
+
 	// This function's name (could be used as identifier).
 	FStringView Name;
 
@@ -1032,6 +1039,8 @@ struct FFunction
 
 	// Whether this function equals specified (they have the same data).
 	bool Equals(const FFunction* Other) const;
+
+	friend uint32 GetTypeHash(const FFunction& F);
 };
 
 // Describes an HLSL preprocessor define, .e.g "#define MyDefine 123"
@@ -1044,6 +1053,11 @@ struct FFunctionHLSLDefine
 	FStringView Value;
 
 	bool operator==(const FFunctionHLSLDefine& Other) const = default;
+
+	friend uint32 GetTypeHash(const FFunctionHLSLDefine& D)
+	{
+		return HashCombine(GetTypeHash(D.Name), GetTypeHash(D.Value));
+	}
 };
 
 // A user-function defined with user-provided HLSL code.
@@ -1060,13 +1074,15 @@ struct FFunctionHLSL : FFunction
 
 	// Whether this function equals specified (they have the same data).
 	bool Equals(const FFunctionHLSL* Other) const;
+
+	friend uint32 GetTypeHash(const FFunctionHLSL& F);
 };
 
 // A call instruction to a user-function.
 struct FCall : TInstruction<VK_Call, 0>
 {
 	// The user-function to call.
-	const FFunction* Function;
+	FFunction* Function;
 
 	// The arguments to pass to the function parameters.
 	FValue* Arguments[MaxNumFunctionParameters];
diff --git a/Engine/Source/Runtime/Engine/Public/Materials/MaterialIRCommon.h b/Engine/Source/Runtime/Engine/Public/Materials/MaterialIRCommon.h
--- a/Engine/Source/Runtime/Engine/Public/Materials/MaterialIRCommon.h
+++ b/Engine/Source/Runtime/Engine/Public/Materials/MaterialIRCommon.h
@@ -6,6 +6,7 @@
 #include "Containers/Array.h"
 #include "Math/MathFwd.h"
 #include "Misc/MemStack.h"
+#include "Materials/MaterialIRUtils.h"
 
 #if WITH_EDITOR
 
@@ -66,10 +67,6 @@ namespace MIR
 namespace MIR
 {
 
-//------------------------------------------------------------------------
-//                          Utilities
-//------------------------------------------------------------------------
-
 // Identifies the block of execution in which an instruction runs.
 // Note: If you introduce a new stage, make sure you update EValueFlags accordingly.
 enum EStage
@@ -93,6 +90,12 @@ enum class EStageMask : uint8
 
 ENUM_CLASS_FLAGS(EStageMask);
 
+// Gets the mask from the specified stage.
+static inline EStageMask StageToMask(EStage Stage)
+{
+	return (EStageMask)(1 << (uint32)Stage);
+}
+
 // Returns whether given Stage is in given Mask.
 static inline bool IsStageInMask(EStage Stage, EStageMask Mask)
 {
@@ -169,80 +172,6 @@ enum class EDerivativeAxis
 	X, Y,
 };
 
-//------------------------------------------------------------------------
-//                          Utilities
-//------------------------------------------------------------------------
-
-// A simple string view without constructors, to be used inside FValues.
-struct FSimpleStringView
-{
-	const TCHAR* Ptr; // The string characters
-	int32        Len; // The string length (excluding null byte)
-
-	FStringView ToView() const { return { Ptr, Len }; }
-	FString ToString() const { return FString{ ToView() }; }
-	operator FStringView() const { return ToView(); }
-	FSimpleStringView& operator=(FStringView View) { Ptr = View.GetData(); Len = View.Len(); return *this; }
-	const TCHAR* operator*() const { return Ptr; }
-};
-
-// Use this to efficiently allocate a temporary array using MemStack, instead
-// of using TArray and going through the global allocator. Only declare local variable
-// of this struct (i.e. do not "new TTemporaryArray")
-// Remark: The allocated memory lifespan is the same as the TTemporaryArray local variable
-// (it is deallocated when TTemporaryArray goes out of scope).
-template <typename T>
-struct TTemporaryArray : public TArrayView<T>
-{
-	FMemMark MemMark;
-
-	TTemporaryArray(int32 Num)
-	: MemMark{ FMemStack::Get() }
-	{
-		*static_cast<TArrayView<T>*>(this) = { (T*)FMemStack::Get().Alloc(sizeof(T) * Num, alignof(T)), Num };
-	}
-
-	TTemporaryArray(TConstArrayView<T> Source)
-	: TTemporaryArray { Source.Num() }
-	{
-		static_assert(std::is_trivially_copyable_v<T>);
-		FMemory::Memcpy(TArrayView<T>::GetData(), Source.GetData(), Source.NumBytes());
-	}
-
-	operator TConstArrayView<T>() const
-	{
-		return { TArrayView<T>::GetData(), TArrayView<T>::Num() };
-	}
-
-	void Zero()
-	{
-		FMemory::Memzero(TArrayView<T>::GetData(), sizeof(T) * TArrayView<T>::Num());
-	}
-};
-
-// Zeroes an array of elements.
-template <typename T>
-void ZeroArray(TArrayView<T> Array)
-{
-	static_assert(TIsTriviallyCopyAssignable<T>::Value);
-	FMemory::Memzero(Array.GetData(), Array.Num() * sizeof(T));
-}
-
-// Tries to find a value with given Key into given TMap intance. If it's found, its value is assigned to OutValue.
-// Returns whether the value was found and OutValue was set.
-template <typename TKey, typename TValue>
-bool Find(const TMap<TKey, TValue>& Map, const TKey& Key, TValue& OutValue)
-{
-	if (auto ValuePtr = Map.Find(Key)) {
-		OutValue = *ValuePtr;
-		return true;
-	}
-	return false;
-}
-
 } // namespace MIR
 
-#define UE_MIR_UNREACHABLE() { check(!"Unreachable"); UE_ASSUME(false); }
-#define UE_MIR_TODO() UE_MIR_UNREACHABLE()
-
 #endif // #if WITH_EDITOR
diff --git a/Engine/Source/Runtime/Engine/Public/Materials/MaterialIREmitter.h b/Engine/Source/Runtime/Engine/Public/Materials/MaterialIREmitter.h
--- a/Engine/Source/Runtime/Engine/Public/Materials/MaterialIREmitter.h
+++ b/Engine/Source/Runtime/Engine/Public/Materials/MaterialIREmitter.h
@@ -748,11 +748,11 @@ public:
 
 	// Defines a user-specified HLSL function to embed in the generated HLSL when generated IR Module is translated to HLSL.
 	// This will create a local function with the specified description. You can issue calls to the returned functions using the Call() API.
-	const FFunction* FunctionHLSL(const FFunctionHLSLDesc& Desc);
+	FFunction* FunctionHLSL(const FFunctionHLSLDesc& Desc);
 
 	// Emits a call to a specified user function using specified input arguments.
 	// Note that the InputArguments must provide an argument for *all* function input parameters (that is including the input-output ones).
-	FValueRef Call(const FFunction* Function, TConstArrayView<FValueRef> InputArguments);
+	FValueRef Call(FFunction* Function, TConstArrayView<FValueRef> InputArguments);
 
 	// Emits an instruction that retrieves the output parameter value from a function call.
 	FValueRef CallParameterOutput(FValueRef Call, uint32 ParameterIndex);
@@ -951,6 +951,16 @@ private:
 	// calculations.
 	TSet<FValue*, FValueKeyFuncs> ValueSet{};
 
+	// Key funcs for deduplicating HLSL functions by content.
+	struct FFunctionHLSLKeyFuncs : DefaultKeyFuncs<const FFunctionHLSL*>
+	{
+		static bool Matches(const FFunctionHLSL* A, const FFunctionHLSL* B) { return A->Equals(B); }
+		static uint32 GetKeyHash(const FFunctionHLSL* Key);
+	};
+
+	// The set of unique HLSL functions emitted so far, used to deduplicate identical function definitions.
+	TSet<FFunctionHLSL*, FFunctionHLSLKeyFuncs> FunctionHLSLSet{};
+
 	friend FMaterialIRModuleBuilder;
 	friend FMaterialIRModuleBuilderImpl;
 };
diff --git a/Engine/Source/Runtime/Engine/Public/Materials/MaterialIRModule.h b/Engine/Source/Runtime/Engine/Public/Materials/MaterialIRModule.h
--- a/Engine/Source/Runtime/Engine/Public/Materials/MaterialIRModule.h
+++ b/Engine/Source/Runtime/Engine/Public/Materials/MaterialIRModule.h
@@ -175,6 +175,9 @@ public:
 	// Returns the list of material parameter collections used by this material.
 	TConstArrayView<UMaterialParameterCollection*> GetParameterCollections() const { return ParameterCollections; }
 
+	// Adds an HLSL function in the module.
+	void AddFunctionHLSL(const MIR::FFunctionHLSL* Function);
+
 	// Returns the array of HLSL-defined user-functions in the module.
 	TConstArrayView<const MIR::FFunctionHLSL*> GetFunctionHLSLs() const { return FunctionHLSLs; };
 
diff --git a/Engine/Source/Runtime/Engine/Public/Materials/MaterialIRTypes.h b/Engine/Source/Runtime/Engine/Public/Materials/MaterialIRTypes.h
--- a/Engine/Source/Runtime/Engine/Public/Materials/MaterialIRTypes.h
+++ b/Engine/Source/Runtime/Engine/Public/Materials/MaterialIRTypes.h
@@ -291,6 +291,8 @@ struct FType
 	bool operator==(FType Other) const;
 	bool operator!=(FType Other) const { return !operator==(Other); }
 
+	friend uint32 GetTypeHash(const FType& T);
+
 private:
 	union
 	{
diff --git a/Engine/Source/Runtime/Engine/Public/Materials/MaterialIRUtils.h b/Engine/Source/Runtime/Engine/Public/Materials/MaterialIRUtils.h
new file mode 100644
--- /dev/null
+++ b/Engine/Source/Runtime/Engine/Public/Materials/MaterialIRUtils.h
@@ -0,0 +1,120 @@
+// Copyright Epic Games, Inc. All Rights Reserved.
+#pragma once
+
+#include "Containers/Array.h"
+#include "Misc/MemStack.h"
+
+#if WITH_EDITOR
+
+#define UE_MIR_UNREACHABLE() { check(!"Unreachable"); UE_ASSUME(false); }
+#define UE_MIR_TODO() UE_MIR_UNREACHABLE()
+
+namespace MIR
+{
+
+// A simple string view without constructors, to be used inside FValues.
+struct FSimpleStringView
+{
+	const TCHAR* Ptr; // The string characters
+	int32        Len; // The string length (excluding null byte)
+
+	FStringView ToView() const { return { Ptr, Len }; }
+	FString ToString() const { return FString{ ToView() }; }
+	operator FStringView() const { return ToView(); }
+	FSimpleStringView& operator=(FStringView View) { Ptr = View.GetData(); Len = View.Len(); return *this; }
+	const TCHAR* operator*() const { return Ptr; }
+};
+
+// Use this to efficiently allocate a temporary array using MemStack, instead
+// of using TArray and going through the global allocator. Only declare local variable
+// of this struct (i.e. do not "new TTemporaryArray")
+// Remark: The allocated memory lifespan is the same as the TTemporaryArray local variable
+// (it is deallocated when TTemporaryArray goes out of scope).
+template <typename T>
+struct TTemporaryArray : public TArrayView<T>
+{
+	FMemMark MemMark;
+
+	TTemporaryArray(int32 Num)
+	: MemMark{ FMemStack::Get() }
+	{
+		*static_cast<TArrayView<T>*>(this) = { (T*)FMemStack::Get().Alloc(sizeof(T) * Num, alignof(T)), Num };
+	}
+
+	TTemporaryArray(TConstArrayView<T> Source)
+	: TTemporaryArray { Source.Num() }
+	{
+		static_assert(std::is_trivially_copyable_v<T>);
+		FMemory::Memcpy(TArrayView<T>::GetData(), Source.GetData(), Source.NumBytes());
+	}
+
+	operator TConstArrayView<T>() const
+	{
+		return { TArrayView<T>::GetData(), TArrayView<T>::Num() };
+	}
+
+	void Zero()
+	{
+		FMemory::Memzero(TArrayView<T>::GetData(), sizeof(T) * TArrayView<T>::Num());
+	}
+};
+
+// Zeroes an array of elements.
+template <typename T>
+void ZeroArray(TArrayView<T> Array)
+{
+	static_assert(TIsTriviallyCopyAssignable<T>::Value);
+	FMemory::Memzero(Array.GetData(), Array.Num() * sizeof(T));
+}
+
+// Tries to find a value with given Key into given TMap intance. If it's found, its value is assigned to OutValue.
+// Returns whether the value was found and OutValue was set.
+template <typename TKey, typename TValue>
+bool Find(const TMap<TKey, TValue>& Map, const TKey& Key, TValue& OutValue)
+{
+	if (auto ValuePtr = Map.Find(Key)) {
+		OutValue = *ValuePtr;
+		return true;
+	}
+	return false;
+}
+
+// Returns whether two array views are equal by comparing their elements.
+// Requires operator== to be defined for T.
+template <typename T>
+bool ArrayViewEquals(TArrayView<T> A, TArrayView<T> B)
+{
+	if (A.Num() != B.Num())
+	{
+		return false;
+	}
+	for (int32 i = 0; i < A.Num(); ++i)
+	{
+		if (!(A[i] == B[i]))
+		{
+			return false;
+		}
+	}
+	return true;
+}
+
+// Computes a hash from a byte buffer using word-aligned reads and HashCombineFast.
+// The pointer must be 4-byte aligned.
+uint32 HashBytes(const void* InPtr, uint32 Size);
+
+// Computes the hash of an array view by combining the hashes of its elements.
+// Requires GetTypeHash to be defined for T.
+template <typename T>
+uint32 ArrayViewHash(TArrayView<T> Array)
+{
+	uint32 Hash = ::GetTypeHash(Array.Num());
+	for (const T& Element : Array)
+	{
+		Hash = HashCombine(Hash, GetTypeHash(Element));
+	}
+	return Hash;
+}
+
+} // namespace MIR
+
+#endif // #if WITH_EDITOR
