New material translator minor fix:

- TexCoords weren't declared as used when sampled in the vertex shader only. Fixed now, but we still should optimize this in the future and avoid interpolating texcoords from vertex to pixel shader when they're only used in the vertex shader. This will require ush files edits.

#rnx

#Synced-CL 51637005

diff --git a/Engine/Source/Runtime/Engine/Private/Materials/MaterialExpressionsToMIRCommon.cpp b/Engine/Source/Runtime/Engine/Private/Materials/MaterialExpressionsToMIRCommon.cpp
--- a/Engine/Source/Runtime/Engine/Private/Materials/MaterialExpressionsToMIRCommon.cpp
+++ b/Engine/Source/Runtime/Engine/Private/Materials/MaterialExpressionsToMIRCommon.cpp
@@ -37,14 +37,18 @@ MIR::FValueRef EmitTexCoord(MIR::FEmitter& Em, int32 Index)
 
 		void AnalyzeInStage(MIR::FExternAnalysisContext& Context, MIR::EStage Stage)
 		{
-			// Texture coordinates read in the vertex stage don't need to be interpolated.
-			if (Stage == MIR::EStage::Stage_Vertex)
-			{
-				return;
-			}
-
-			// @massimo.tristano this follows the logic in the old translator. A better way would instead be to **allocate** an index
-			// to a texture coordinate based on usage and map that index here. 
+			// @massimo.tristano I think we can avoid fully interpolating texture coordinates when they're
+			// only used inside a vertex stage. We should make a different #define to store the number of
+			// texture coordinates we need in the MaterialVertexParameter, so that it is not automatically
+			// interpolated for the pixel shader, saving perf. 
+			// 	 if (Stage == MIR::EStage::Stage_Vertex)
+			// 	 {
+			//       <here bump the number of texcoords only used in vertex shader>
+			//       return;
+			// 	 }
+
+			// @massimo.tristano this follows the logic in the old translator. A better way would instead be
+			// to **allocate** an index to a texture coordinate based on usage and map that index here. 
 			FMaterialIRModule::FStatistics& Statistics = Context.Module->GetStatistics();
 			Statistics.NumInterpolatedTexCoords = FMath::Max(Statistics.NumInterpolatedTexCoords, Index + 1);
 			Statistics.InterpolatedTexCoordsMask |= 1 << Index;
