Iris
- Prevent crash in NetConnection due to suspected bad incoming packet.

diff --git a/Engine/Source/Runtime/Engine/Private/NetConnection.cpp b/Engine/Source/Runtime/Engine/Private/NetConnection.cpp
--- a/Engine/Source/Runtime/Engine/Private/NetConnection.cpp
+++ b/Engine/Source/Runtime/Engine/Private/NetConnection.cpp
@@ -305,6 +305,8 @@ static void PushPacketId(FWrittenChannelsRecord& WrittenChannelsRecord, int32 Pa
 // Push written ChannelIndex to WrittenChanneRecord, Push new packetId if PackedIdDiffers from Last pushed entry */
 static void PushChannelRecord(FWrittenChannelsRecord& WrittenChannelsRecord, int32 PacketId, int32 ChannelIndex);
 
+static bool HasChannelRecordsToConsumeForPacket(const FWrittenChannelsRecord& WrittenChannelsRecord, int32 PacketId);
+
 // Consume all FChannelRecordEntries for the given PacketId and execute a function with the signature (void*)(int32 PacketId, uint32 Channelndex) for each entry
 template<class Functor>
 static void ConsumeChannelRecordsForPacket(FWrittenChannelsRecord& WrittenChannelsRecord, int32 PacketId, Functor&& Func);
@@ -3403,7 +3405,16 @@ void UNetConnection::ReceivedPacket( FBitReader& Reader, bool bIsReinjectedPacke
 				// Sanity check
 				if (FNetPacketNotify::SequenceNumberT(LastNotifiedPacketId) != AckedSequence)
 				{
-					UE_LOG(LogNet, Warning, TEXT("LastNotifiedPacketId != AckedSequence"));
+					UE_LOGF(LogNet, Error, "LastNotifiedPacketId != AckedSequence");
+
+					Close(ENetCloseResult::AckSequenceMismatch);
+
+					return;
+				}
+
+				if (!FChannelRecordImpl::HasChannelRecordsToConsumeForPacket(ChannelRecord, LastNotifiedPacketId))
+				{
+					UE_LOGF(LogNet, Error, "Could not find matching record for packet %d", LastNotifiedPacketId);
 
 					Close(ENetCloseResult::AckSequenceMismatch);
 
@@ -6546,6 +6557,18 @@ SIZE_T CountBytes(FWrittenChannelsRecord& WrittenChannelsRecord)
 	return WrittenChannelsRecord.ChannelRecord.AllocatedCapacity() * sizeof(FWrittenChannelsRecord::FChannelRecordEntry);
 }
 
+bool HasChannelRecordsToConsumeForPacket(const FWrittenChannelsRecord& WrittenChannelsRecord, int32 PacketId)
+{
+	const FWrittenChannelsRecord::FChannelRecordEntryQueue& Record = WrittenChannelsRecord.ChannelRecord;
+	if (!Record.IsEmpty())
+	{
+		FWrittenChannelsRecord::FChannelRecordEntry PacketEntry = Record.Peek();
+		return PacketEntry.IsSequence == 1u && PacketEntry.Value == (uint32)PacketId;
+	}
+
+	return false;
+}
+
 template<class Functor>
 void ConsumeChannelRecordsForPacket(FWrittenChannelsRecord& WrittenChannelsRecord, int32 PacketId, Functor&& Func)
 {
