[Actor] ComponentToWorld needs to be up-to-date before attempting a DetachFromComponent
- Otherwise, parented actors may change positions when entering PIE, when Runtime Cells Transformers are enabled, because RerunConstructionScript is invoked on unregistered components.

#Synced-CL 51591079
#Synced-CL 51591079

diff --git a/Engine/Source/Runtime/Engine/Private/ActorConstruction.cpp b/Engine/Source/Runtime/Engine/Private/ActorConstruction.cpp
--- a/Engine/Source/Runtime/Engine/Private/ActorConstruction.cpp
+++ b/Engine/Source/Runtime/Engine/Private/ActorConstruction.cpp
@@ -467,6 +467,9 @@ void AActor::RerunConstructionScripts()
 
 			if (RootComponent != nullptr)
 			{
+				// ComponentToWorld must be up-to-date for the detach below, and to get OldTransform
+				RootComponent->UpdateComponentToWorld();
+
 				// Do not need to detach if root component is not going away
 				if (RootComponent->GetAttachParent() != nullptr && RootComponent->IsCreatedByConstructionScript())
 				{
@@ -483,9 +486,6 @@ void AActor::RerunConstructionScripts()
 					RootComponent->DetachFromComponent(FDetachmentTransformRules::KeepWorldTransform);
 				}
 
-				// Update component transform and remember it so it can be reapplied to any new root component which exists after construction.
-				// (Component transform may be stale if we are here following an Undo)
-				RootComponent->UpdateComponentToWorld();
 				OldTransform = RootComponent->GetComponentTransform();
 				OldTransformRotationCache = RootComponent->GetRelativeRotationCache();
 			}
