Converting Characters, Character Movement and dependent functionality to use a UObject representation of Movement Bases instead of UPrimitiveComponents. This opens up movement bases to non-actor objects such as FastGeoStreaming objects or Scenegraph objects.
-Added new FMovementBaseInterfaceData struct that houses UObject PysicsObjectOwner and IPhysicsBodyOwnerInterface that implements necessary functions for basic movement base functionality. FMovementBaseInterfaceData replaces UPrimitiveComponent as how code interacts with movement bases but a UPrimitiveComponent can be used to make a FMovementBaseInterfaceData easily.
-Modified important classes/structs such as ACharacter, UCharacterMovementComponent, FBasedMovementInfo, FCharacterNetworkMoveData, FClientAdjustment and MovementBaseUtility with deprecation and fixup so they properly use new FMovementBaseInterfaceData struct.
-ACharacter and UCharacterMovementComponent have multiple examples of how to network and use this new struct.

#Synced-CL 51675492

#Synced-CL 51675492

diff --git a/Engine/Plugins/Experimental/FastGeoStreaming/Source/FastGeoStreaming/Internal/FastGeoPhysicsBodyInstanceOwner.h b/Engine/Plugins/Experimental/FastGeoStreaming/Source/FastGeoStreaming/Internal/FastGeoPhysicsBodyInstanceOwner.h
--- a/Engine/Plugins/Experimental/FastGeoStreaming/Source/FastGeoStreaming/Internal/FastGeoPhysicsBodyInstanceOwner.h
+++ b/Engine/Plugins/Experimental/FastGeoStreaming/Source/FastGeoStreaming/Internal/FastGeoPhysicsBodyInstanceOwner.h
@@ -38,6 +38,15 @@ public:
 	virtual int32 GetNumMaterials() const override;
 	virtual UMaterialInterface* GetMaterial(int32 Index) const override;
 	virtual UBodySetup* GetPhysicsBodySetup() const override;
+	virtual UObject* GetSourceObjectOwner() const override;
+	virtual FTransform GetPhysicsOwnerSocketTransform(FName InSocketName) const override;
+	virtual bool IsPhysicsOwnerMovable() const override;
+	virtual bool IsPhysicsOwnerSimulatingPhysics() const override;
+	virtual FVector GetPhysicsOwnerVelocity() const override;
+	virtual UObject* GetPhysicsOwnerAttachmentRoot() const override;
+	virtual bool IsPhysicsObjectWorldGeometry() const override;
+	virtual bool DoesSocketExistOnPhysicsOwner(FName InSocketName) const override;
+	virtual const FWalkableSlopeOverride& GetWalkableSlopeOverride() const override;
 	//~ End IPhysicsBodyInstanceOwner interface
 
 private:
diff --git a/Engine/Plugins/Experimental/FastGeoStreaming/Source/FastGeoStreaming/Private/FastGeoPhysicsBodyInstanceOwner.cpp b/Engine/Plugins/Experimental/FastGeoStreaming/Source/FastGeoStreaming/Private/FastGeoPhysicsBodyInstanceOwner.cpp
--- a/Engine/Plugins/Experimental/FastGeoStreaming/Source/FastGeoStreaming/Private/FastGeoPhysicsBodyInstanceOwner.cpp
+++ b/Engine/Plugins/Experimental/FastGeoStreaming/Source/FastGeoStreaming/Private/FastGeoPhysicsBodyInstanceOwner.cpp
@@ -134,4 +134,55 @@ UBodySetup* FFastGeoPhysicsBodyInstanceOwner::GetPhysicsBodySetup() const
 {
 	check(OwnerContainer.IsValid());
 	return OwnerComponent->GetBodySetup();
-}
\ No newline at end of file
+}
+
+UObject* FFastGeoPhysicsBodyInstanceOwner::GetSourceObjectOwner() const
+{
+	check(OwnerContainer.IsValid());
+	return OwnerComponent->GetOwner();
+}
+
+FTransform FFastGeoPhysicsBodyInstanceOwner::GetPhysicsOwnerSocketTransform(FName InSocketName) const
+{
+	return FTransform::Identity;
+}
+
+bool FFastGeoPhysicsBodyInstanceOwner::IsPhysicsOwnerMovable() const
+{
+	// FastGeo doesn't support movables
+	return false;
+}
+
+bool FFastGeoPhysicsBodyInstanceOwner::IsPhysicsOwnerSimulatingPhysics() const
+{
+	check(OwnerContainer.IsValid());
+	return OwnerComponent->BodyInstance.IsInstanceSimulatingPhysics();
+}
+
+FVector FFastGeoPhysicsBodyInstanceOwner::GetPhysicsOwnerVelocity() const
+{
+	// FastGeo doesn't support movables so return empty vector
+	return FVector::ZeroVector;
+}
+
+UObject* FFastGeoPhysicsBodyInstanceOwner::GetPhysicsOwnerAttachmentRoot() const
+{
+	check(OwnerContainer.IsValid());
+	return OwnerComponent->GetSourceObject();
+}
+
+bool FFastGeoPhysicsBodyInstanceOwner::IsPhysicsObjectWorldGeometry() const
+{
+	return true;
+}
+
+bool FFastGeoPhysicsBodyInstanceOwner::DoesSocketExistOnPhysicsOwner(FName InSocketName) const
+{
+	return false;
+}
+
+const FWalkableSlopeOverride& FFastGeoPhysicsBodyInstanceOwner::GetWalkableSlopeOverride() const
+{
+	check(OwnerContainer.IsValid());
+	return OwnerComponent->BodyInstance.GetWalkableSlopeOverride();
+}
diff --git a/Engine/Plugins/FX/Niagara/Source/Niagara/Private/NiagaraDataChannel_GameplayBurst.cpp b/Engine/Plugins/FX/Niagara/Source/Niagara/Private/NiagaraDataChannel_GameplayBurst.cpp
--- a/Engine/Plugins/FX/Niagara/Source/Niagara/Private/NiagaraDataChannel_GameplayBurst.cpp
+++ b/Engine/Plugins/FX/Niagara/Source/Niagara/Private/NiagaraDataChannel_GameplayBurst.cpp
@@ -186,11 +186,13 @@ bool FNDCGameplayBurstAttachmentSettings::UseAttachedPathForComponent(const USce
 	//Check the owners speed to check to auto attach to moving components/actors.
 	if (SpeedThreshold >= 0.0f)
 	{
-		const UPrimitiveComponent* PrimComp = Cast<const UPrimitiveComponent>(Component);
+		const UPrimitiveComponent* ConstPrimComp = Cast<const UPrimitiveComponent>(Component);
+		UPrimitiveComponent* PrimComp = const_cast<UPrimitiveComponent*>(ConstPrimComp);
+		const FMovementBaseInterfaceData MovementBaseData(PrimComp);
 		//If the searching component is not a Niagara component then we can bail out immediately for static objects without ever searching our PerComponentData map.
-		if (MovementBaseUtility::IsDynamicBase(PrimComp))
+		if (MovementBaseUtility::IsDynamicBase(&MovementBaseData))
 		{
-			FVector BaseVel = MovementBaseUtility::GetMovementBaseVelocity(PrimComp, NAME_None);
+			FVector BaseVel = MovementBaseUtility::GetMovementBaseVelocity(&MovementBaseData, NAME_None);
 			if (BaseVel.SquaredLength() >= (SpeedThreshold * SpeedThreshold))
 			{
 				return true;
@@ -201,9 +203,10 @@ bool FNDCGameplayBurstAttachmentSettings::UseAttachedPathForComponent(const USce
 		if (AttachParent)
 		{
 			// Check the attached parent for movement.
-			if (const UPrimitiveComponent* PrimitiveComponentBase = Cast<const UPrimitiveComponent>(AttachParent->GetRootComponent()))
+			if (UPrimitiveComponent* PrimitiveComponentBase = Cast<UPrimitiveComponent>(AttachParent->GetRootComponent()))
 			{
-				FVector BaseVel = MovementBaseUtility::GetMovementBaseVelocity(PrimitiveComponentBase, NAME_None);
+				const FMovementBaseInterfaceData AttachParentMovementBaseData(PrimitiveComponentBase);
+				FVector BaseVel = MovementBaseUtility::GetMovementBaseVelocity(&AttachParentMovementBaseData, NAME_None);
 				if (BaseVel.SquaredLength() >= (SpeedThreshold * SpeedThreshold))
 				{
 					return true;
diff --git a/Engine/Source/Runtime/AIModule/Private/GameplayDebugger/GameplayDebuggerCategory_AI.cpp b/Engine/Source/Runtime/AIModule/Private/GameplayDebugger/GameplayDebuggerCategory_AI.cpp
--- a/Engine/Source/Runtime/AIModule/Private/GameplayDebugger/GameplayDebuggerCategory_AI.cpp
+++ b/Engine/Source/Runtime/AIModule/Private/GameplayDebugger/GameplayDebuggerCategory_AI.cpp
@@ -159,9 +159,9 @@ void FGameplayDebuggerCategory_AI::CollectData(APlayerController* OwnerPC, AActo
 		UCharacterMovementComponent* CharMovementComp = MyChar ? MyChar->GetCharacterMovement() : nullptr;
 		if (CharMovementComp)
 		{
-			UPrimitiveComponent* FloorComponent = MyPawn->GetMovementBase();
-			AActor* FloorActor = FloorComponent ? FloorComponent->GetOwner() : nullptr;
-			DataPack.MovementBaseInfo = FloorComponent ? FString::Printf(TEXT("%s.%s"), *GetNameSafe(FloorActor), *FloorComponent->GetName()) : FString("None");
+			const FMovementBaseInterfaceData* MovementBaseData = MyPawn->GetMovementBaseInterfaceData();
+			const UObject* FloorObject = MovementBaseData && MovementBaseData->IsValid() ? MovementBaseData->GetMovementBaseObjectOwner() : nullptr;
+			DataPack.MovementBaseInfo = MovementBaseData && MovementBaseData->GetMovementBaseObject() ? FString::Printf(TEXT("%s.%s"), *GetNameSafe(FloorObject), *MovementBaseData->GetMovementBaseObject()->GetName()) : FString("None");
 			DataPack.MovementModeInfo = CharMovementComp->GetMovementName();
 		}
 
@@ -351,9 +351,9 @@ void FGameplayDebuggerCategory_AI::DrawData(APlayerController* OwnerPC, FGamepla
 			// and we do want to display client-side data here as well, so we bend the rules a bit and gather and display data instantly here
 			if (const UCharacterMovementComponent* CharMovementComp = MyChar ? MyChar->GetCharacterMovement() : nullptr)
 			{
-				const UPrimitiveComponent* FloorComponent = MyChar->GetMovementBase();
-				const AActor* FloorActor = FloorComponent ? FloorComponent->GetOwner() : nullptr;
-				const FString ClientMovementBaseInfo = FloorComponent ? FString::Printf(TEXT("%s.%s"), *GetNameSafe(FloorActor), *FloorComponent->GetName()) : FString("None");
+				const FMovementBaseInterfaceData* MovementBaseData = MyChar->GetMovementBaseInterfaceData();
+				const UObject* FloorObject = MovementBaseData && MovementBaseData->IsValid() ? MovementBaseData->GetMovementBaseObjectOwner() : nullptr;
+				const FString ClientMovementBaseInfo = MovementBaseData && MovementBaseData->GetMovementBaseObject() ? FString::Printf(TEXT("%s.%s"), *GetNameSafe(FloorObject), *MovementBaseData->GetMovementBaseObject()->GetName()) : FString("None");
 				const FString ClientMovementModeInfo = CharMovementComp->GetMovementName();
 
 				CanvasContext.Printf(TEXT("(Client) Movement Mode: {yellow}%s{white}, Base: {yellow}%s"), *ClientMovementModeInfo, *ClientMovementBaseInfo);
diff --git a/Engine/Source/Runtime/Engine/Classes/Components/PrimitiveComponent.h b/Engine/Source/Runtime/Engine/Classes/Components/PrimitiveComponent.h
--- a/Engine/Source/Runtime/Engine/Classes/Components/PrimitiveComponent.h
+++ b/Engine/Source/Runtime/Engine/Classes/Components/PrimitiveComponent.h
@@ -285,7 +285,7 @@ DECLARE_DYNAMIC_MULTICAST_SPARSE_DELEGATE_TwoParams( FComponentEndTouchOverSigna
  * ShapeComponents generate geometry that is used for collision detection but are not rendered, while StaticMeshComponents and SkeletalMeshComponents contain pre-built geometry that is rendered, but can also be used for collision detection.
  */
 UCLASS(abstract, HideCategories=(Mobility, VirtualTexture), ShowCategories=(PhysicsVolume), MinimalAPI)
-class UPrimitiveComponent : public USceneComponent, public INavRelevantInterface, public IInterface_AsyncCompilation, public IPhysicsComponent, public FRenderAssetOwnerStreamingState, public IPhysicsBodyInstanceOwner
+class UPrimitiveComponent : public USceneComponent, public INavRelevantInterface, public IInterface_AsyncCompilation, public IPhysicsComponent, public FRenderAssetOwnerStreamingState, public IPhysicsBodyInstanceOwner, public IPhysicsBodyInstanceOwnerResolver
 {
 	GENERATED_BODY()
 	UE_DECLARE_COMPONENT_ACTOR_INTERFACE(PrimitiveComponent)
@@ -1596,7 +1596,7 @@ public:
 
 	/** Returns the slope override struct for this component. */
 	UFUNCTION(BlueprintPure, Category="Physics")
-	ENGINE_API const struct FWalkableSlopeOverride& GetWalkableSlopeOverride() const;
+	ENGINE_API virtual const struct FWalkableSlopeOverride& GetWalkableSlopeOverride() const override;
 
 	/** Sets a new slope override for this component instance. */
 	UFUNCTION(BlueprintCallable, Category="Physics")
@@ -3178,6 +3178,10 @@ public:
 	ENGINE_API virtual Chaos::FPhysicsObjectId GetIdFromGTParticle(Chaos::FGeometryParticle* Particle) const override;
 	//~ End IPhysicsComponent Interface.
 
+	//~Begin IPhysicsBodyInstanceOwnerResolver
+	ENGINE_API virtual IPhysicsBodyInstanceOwner* ResolvePhysicsBodyInstanceOwner(Chaos::FConstPhysicsObjectHandle PhysicsObject) override;
+	//~End IPhysicsBodyInstanceOwnerResolver
+	
 	//~ Begin IPhysicsBodyInstanceOwner Interface.
 
 	/** Whether the physics is static for this component. */
@@ -3193,6 +3197,9 @@ public:
 	/** Returns the source UObject of this body instance owner. */
 	virtual UObject* GetSourceObject() const override { return const_cast<UPrimitiveComponent*>(this); }
 
+	/** Returns the owner of the source UObject of this body instance owner. */
+	ENGINE_API virtual UObject* GetSourceObjectOwner() const override;
+
 	/** Returns the transform of this body instance owner. */
 	virtual FTransform GetPhysicsOwnerTransform() const override { return GetComponentTransform(); }
 
@@ -3205,6 +3212,26 @@ public:
 	/** Returns the body instance owner mesh material. */
 	ENGINE_API virtual UMaterialInterface* GetPhysicsMaterialBase() const override;
 
+	/** Whether this body instance owner is movable or not. On UPrimitiveComponents this checks the Mobility type property. */
+	ENGINE_API virtual bool IsPhysicsOwnerMovable() const override;
+	
+	/** Whether this body instance owner is simulating physics or not */
+	ENGINE_API virtual bool IsPhysicsOwnerSimulatingPhysics() const override;
+	
+	/** Gets the body instance owner velocity. Since this is a UPrimitiveComponent we will also check component velocity. */
+	ENGINE_API virtual FVector GetPhysicsOwnerVelocity() const override;
+	
+	/** Returns the attachment root of the body instance owner if it exists */
+	ENGINE_API virtual UObject* GetPhysicsOwnerAttachmentRoot() const override;
+	
+	/** Whether this body instance is world geometry or not */
+	ENGINE_API virtual bool IsPhysicsObjectWorldGeometry() const override;
+	
+	ENGINE_API virtual FTransform GetPhysicsOwnerSocketTransform(FName InSocketName) const override;
+	
+	/** Checks for a socket on the body instance owner */
+	ENGINE_API virtual bool DoesSocketExistOnPhysicsOwner(FName InSocketName) const override;
+	
 	//~ End IPhysicsBodyInstanceOwner Interface.
 
 	/** Wheter the component is relevant for first person */
diff --git a/Engine/Source/Runtime/Engine/Classes/GameFramework/Character.h b/Engine/Source/Runtime/Engine/Classes/GameFramework/Character.h
--- a/Engine/Source/Runtime/Engine/Classes/GameFramework/Character.h
+++ b/Engine/Source/Runtime/Engine/Classes/GameFramework/Character.h
@@ -15,6 +15,7 @@
 #include "GameFramework/CharacterMovementReplication.h"
 #include "Animation/AnimationAsset.h"
 #include "GameFramework/RootMotionSource.h"
+#include "Interfaces/MovementBaseInterface.h"
 #include "Character.generated.h"
 
 class ACharacter;
@@ -48,6 +49,15 @@ struct FRepRootMotionMontage
 {
 	GENERATED_USTRUCT_BODY()
 
+	PRAGMA_DISABLE_DEPRECATION_WARNINGS
+	FRepRootMotionMontage() = default;
+	~FRepRootMotionMontage() = default;
+	FRepRootMotionMontage(const FRepRootMotionMontage&) = default;
+	FRepRootMotionMontage(FRepRootMotionMontage&&) = default;
+	FRepRootMotionMontage& operator=(const FRepRootMotionMontage&) = default;
+	FRepRootMotionMontage& operator=(FRepRootMotionMontage&&) = default;
+	PRAGMA_ENABLE_DEPRECATION_WARNINGS
+	
 	/** Animation providing Root Motion */
 	UPROPERTY()
 	TObjectPtr<UAnimSequenceBase> Animation = nullptr;
@@ -77,9 +87,14 @@ struct FRepRootMotionMontage
 	FRotator Rotation = FRotator(0.f);
 
 	/** Movement Relative to Base */
-	UPROPERTY()
+	UE_DEPRECATED(5.8, "This property has been deprecated. Please use PhysicsObjectOwner instead and use that to create a FMovementBaseInterfaceData. See ACharacter::OnRep_ReplicatedBasedMovement() for an example of this.")
+	UPROPERTY(meta = (DeprecatedProperty, DeprecationMessage="Please use PhysicsObjectOwner instead and use that to create a FMovementBaseInterfaceData."))
 	TObjectPtr<UPrimitiveComponent> MovementBase = nullptr;
 
+	/** PhysicsObjectOwner for BodyInstance used as movement base */
+	UPROPERTY()
+	TObjectPtr<UObject> PhysicsObjectOwner = nullptr;
+	
 	/** Bone on the MovementBase, if a skeletal mesh. */
 	UPROPERTY()
 	FName MovementBaseBoneName;
@@ -131,44 +146,112 @@ struct FSimulatedRootMotionReplicatedMove
 /** MovementBaseUtility provides utilities for working with movement bases, for which we may need relative positioning info. */
 namespace MovementBaseUtility
 {
+	// Start deprecated movement base utility functions
 	/** Determine whether MovementBase can possibly move. */
+	UE_DEPRECATED(5.8, "MovementBaseUtility functions using UPrimitiveComponents have been deprecated. Use functions using the FMovementBaseInterfaceData struct instead.")
 	ENGINE_API bool IsDynamicBase(const UPrimitiveComponent* MovementBase);
 
 	/** Determine whether MovementBase is simulating or attached to a simulating object. */
+	UE_DEPRECATED(5.8, "MovementBaseUtility functions using UPrimitiveComponents have been deprecated. Use functions using the FMovementBaseInterfaceData struct instead.")
 	ENGINE_API bool IsSimulatedBase(const UPrimitiveComponent* MovementBase);
 
 	/** Determine if we should use relative positioning when based on a component (because it may move). */
+	UE_DEPRECATED(5.8, "MovementBaseUtility functions using UPrimitiveComponents have been deprecated. Use functions using the FMovementBaseInterfaceData struct instead.")
 	inline bool UseRelativeLocation(const UPrimitiveComponent* MovementBase)
 	{
+		PRAGMA_DISABLE_DEPRECATION_WARNINGS
 		return IsDynamicBase(MovementBase);
+		PRAGMA_ENABLE_DEPRECATION_WARNINGS
 	}
-
+	
 	/** Ensure that BasedObjectTick ticks after NewBase */
+	UE_DEPRECATED(5.8, "MovementBaseUtility functions using UPrimitiveComponents have been deprecated. Use functions using the FMovementBaseInterfaceData struct instead.")
 	ENGINE_API void AddTickDependency(FTickFunction& BasedObjectTick, UPrimitiveComponent* NewBase);
 
 	/** Remove tick dependency of BasedObjectTick on OldBase */
+	UE_DEPRECATED(5.8, "MovementBaseUtility functions using UPrimitiveComponents have been deprecated. Use functions using the FMovementBaseInterfaceData struct instead.")
 	ENGINE_API void RemoveTickDependency(FTickFunction& BasedObjectTick, UPrimitiveComponent* OldBase);
 
 	/** Get the velocity of the given component, first checking the ComponentVelocity and falling back to the physics velocity if necessary. */
+	UE_DEPRECATED(5.8, "MovementBaseUtility functions using UPrimitiveComponents have been deprecated. Use functions using the FMovementBaseInterfaceData struct instead.")
 	ENGINE_API FVector GetMovementBaseVelocity(const UPrimitiveComponent* MovementBase, const FName BoneName);
 
 	/** Get the tangential velocity at WorldLocation for the given component. */
+	UE_DEPRECATED(5.8, "MovementBaseUtility functions using UPrimitiveComponents have been deprecated. Use functions using the FMovementBaseInterfaceData struct instead.")
 	ENGINE_API FVector GetMovementBaseTangentialVelocity(const UPrimitiveComponent* MovementBase, const FName BoneName, const FVector& WorldLocation);
 
 	/** Get the transform (local-to-world) for the given MovementBase, optionally at the location of a bone. Returns false if MovementBase is nullptr, or if BoneName is not a valid bone. */
+	UE_DEPRECATED(5.8, "MovementBaseUtility functions using UPrimitiveComponents have been deprecated. Use functions using the FMovementBaseInterfaceData struct instead.")
 	ENGINE_API bool GetMovementBaseTransform(const UPrimitiveComponent* MovementBase, const FName BoneName, FVector& OutLocation, FQuat& OutQuat);
 
 	/** Convert a local location to a world location for a given MovementBase. Returns false if MovementBase is nullptr, or if BoneName is not a valid bone. Scaling is ignored. */
+	UE_DEPRECATED(5.8, "MovementBaseUtility functions using UPrimitiveComponents have been deprecated. Use functions using the FMovementBaseInterfaceData struct instead.")
 	ENGINE_API bool TransformLocationToWorld(const UPrimitiveComponent* MovementBase, const FName BoneName, const FVector& LocalLocation, FVector& OutLocationWorldSpace);
 
 	/** Convert a world location to a local location for a given MovementBase, optionally at the location of a bone. Returns false if MovementBase is nullptr, or if BoneName is not a valid bone. Scaling is ignored. */
+	UE_DEPRECATED(5.8, "MovementBaseUtility functions using UPrimitiveComponents have been deprecated. Use functions using the FMovementBaseInterfaceData struct instead.")
 	ENGINE_API bool TransformLocationToLocal(const UPrimitiveComponent* MovementBase, const FName BoneName, const FVector& WorldSpaceLocation, FVector& OutLocalLocation);
 
 	/** Convert a local direction to a world direction for a given MovementBase. Returns false if MovementBase is nullptr, or if BoneName is not a valid bone. Scaling is ignored. */
+	UE_DEPRECATED(5.8, "MovementBaseUtility functions using UPrimitiveComponents have been deprecated. Use functions using the FMovementBaseInterfaceData struct instead.")
 	ENGINE_API bool TransformDirectionToWorld(const UPrimitiveComponent* MovementBase, const FName BoneName, const FVector& LocalDirection, FVector& OutDirectionWorldSpace);
 
 	/** Convert a world direction to a local direction for a given MovementBase, optionally relative to the orientation of a bone. Returns false if MovementBase is nullptr, or if BoneName is not a valid bone. Scaling is ignored. */
+	UE_DEPRECATED(5.8, "MovementBaseUtility functions using UPrimitiveComponents have been deprecated. Use functions using the FMovementBaseInterfaceData struct instead.")
 	ENGINE_API bool TransformDirectionToLocal(const UPrimitiveComponent* MovementBase, const FName BoneName, const FVector& WorldSpaceDirection, FVector& OutLocalDirection);
+	// End deprecated movement base utility functions
+
+
+	/** Determine whether MovementBase can possibly move. */
+	ENGINE_API bool IsDynamicBase(const FMovementBaseInterfaceData* MovementBase);
+
+	/** Determine whether MovementBase is simulating or attached to a simulating object. */
+	ENGINE_API bool IsSimulatedBase(const FMovementBaseInterfaceData* MovementBase);
+
+	/** Ensure that BasedObjectTick ticks after NewBase */
+	ENGINE_API void AddTickDependency(FTickFunction& BasedObjectTick, FMovementBaseInterfaceData* MovementBase);
+
+	/** Remove tick dependency of BasedObjectTick on OldBase */
+	ENGINE_API void RemoveTickDependency(FTickFunction& BasedObjectTick, FMovementBaseInterfaceData* MovementBase);
+	
+	/** Determine if we should use relative positioning when based on a component (because it may move). */
+	inline bool UseRelativeLocation(const FMovementBaseInterfaceData* MovementBase)
+	{
+		return IsDynamicBase(MovementBase);
+	}
+
+	/** Get the velocity of the MovementBase, first checking the ComponentVelocity and falling back to the physics velocity if necessary. */
+	ENGINE_API FVector GetMovementBaseVelocity(const FMovementBaseInterfaceData* MovementBase, const FName BoneName);
+
+	/** Get the tangential velocity at WorldLocation for the MovementBase. */
+	ENGINE_API FVector GetMovementBaseTangentialVelocity(const FMovementBaseInterfaceData* MovementBase, const FName BoneName, const FVector& WorldLocation);
+	
+	/** Get the transform (local-to-world) for the given MovementBase, optionally at the location of a bone. Returns false if MovementBase is nullptr, or if BoneName is not a valid bone. */
+	ENGINE_API bool GetMovementBaseTransform(const FMovementBaseInterfaceData* MovementBase, const FName BoneName, FVector& OutLocation, FQuat& OutQuat);
+
+	/** Convert a local location to a world location for a given MovementBase. Returns false if MovementBase is nullptr, or if BoneName is not a valid bone. Scaling is ignored. */
+	ENGINE_API bool TransformLocationToWorld(const FMovementBaseInterfaceData* MovementBase, const FName BoneName, const FVector& LocalLocation, FVector& OutLocationWorldSpace);
+
+	/** Convert a world location to a local location for a given MovementBase, optionally at the location of a bone. Returns false if MovementBase is nullptr, or if BoneName is not a valid bone. Scaling is ignored. */
+	ENGINE_API bool TransformLocationToLocal(const FMovementBaseInterfaceData* MovementBase, const FName BoneName, const FVector& WorldSpaceLocation, FVector& OutLocalLocation);
+
+	/** Convert a local direction to a world direction for a given MovementBase. Returns false if MovementBase is nullptr, or if BoneName is not a valid bone. Scaling is ignored. */
+	ENGINE_API bool TransformDirectionToWorld(const FMovementBaseInterfaceData* MovementBase, const FName BoneName, const FVector& LocalDirection, FVector& OutDirectionWorldSpace);
+
+	/** Convert a world direction to a local direction for a given MovementBase, optionally relative to the orientation of a bone. Returns false if MovementBase is nullptr, or if BoneName is not a valid bone. Scaling is ignored. */
+	ENGINE_API bool TransformDirectionToLocal(const FMovementBaseInterfaceData* MovementBase, const FName BoneName, const FVector& WorldSpaceDirection, FVector& OutLocalDirection);
+
+	/** Returns true if movement base data passed in matches. Even if both are invalid. */
+	ENGINE_API bool DoesMovementBaseDataMatch(const FMovementBaseInterfaceData* MovementBaseDataA, const FMovementBaseInterfaceData* MovementBaseDataB);
+	
+	/** Returns true if MovementBaseData is valid (both not null and underlying FMovementBaseInterfaceData::IsValid check) */
+	ENGINE_API bool IsMovementBaseDataValid(const FMovementBaseInterfaceData* MovementBaseData);
+	
+	/** Returns FMovementBaseInterfaceData with movement base data extracted from the PhysicsObjectOwner passed in. */
+	ENGINE_API FMovementBaseInterfaceData GetMovementBaseDataFromPhysicsOwner(UObject* PhysicsObjectOwner);
+
+	/** Returns FMovementBaseInterfaceData with movement base data extracted from the HitResult passed in. */
+	ENGINE_API FMovementBaseInterfaceData GetMovementBaseDataFromHitResult(const FHitResult* HitResult);
 }
 
 /** Struct to hold information about the "base" object the character is standing on. */
@@ -177,6 +260,15 @@ struct FBasedMovementInfo
 {
 	GENERATED_USTRUCT_BODY()
 
+	PRAGMA_DISABLE_DEPRECATION_WARNINGS
+	FBasedMovementInfo() = default;
+	~FBasedMovementInfo() = default;
+	FBasedMovementInfo(const FBasedMovementInfo&) = default;
+	FBasedMovementInfo(FBasedMovementInfo&&) = default;
+	FBasedMovementInfo& operator=(const FBasedMovementInfo&) = default;
+	FBasedMovementInfo& operator=(FBasedMovementInfo&&) = default;
+	PRAGMA_ENABLE_DEPRECATION_WARNINGS
+	
 	/** Unique (within a reasonable timespan) ID of the base component. Can be used to detect changes in the base when the pointer can't replicate, eg during fast shared replication. */
 	UPROPERTY()
 	uint16 BaseID = 0;
@@ -198,7 +290,8 @@ struct FBasedMovementInfo
 	FName BoneName;
 
 	/** Component we are based on */
-	UPROPERTY()
+	UE_DEPRECATED(5.8, "This property has been deprecated. Please use MovementBaseInterfaceData instead.")
+	UPROPERTY(meta = (DeprecatedProperty, DeprecationMessage="Use MovementBaseInterfaceData instead."))
 	TObjectPtr<UPrimitiveComponent> MovementBase = nullptr;
 
 	/** Location relative to MovementBase. Only valid if HasRelativeLocation() is true. */
@@ -209,10 +302,14 @@ struct FBasedMovementInfo
 	UPROPERTY()
 	FRotator Rotation = FRotator(0.f);
 
+	/** Struct containing Movement base functionality and relevant data for the UObject we are based on */
+	UPROPERTY()
+	FMovementBaseInterfaceData MovementBaseInterfaceData;
+	
 	/** Is location relative? */
 	inline bool HasRelativeLocation() const
 	{
-		return MovementBaseUtility::UseRelativeLocation(MovementBase);
+		return MovementBaseUtility::UseRelativeLocation(&MovementBaseInterfaceData);
 	}
 
 	/** Is rotation relative or absolute? It can only be relative if location is also relative. */
@@ -224,7 +321,7 @@ struct FBasedMovementInfo
 	/** Return true if the client should have MovementBase, but it hasn't replicated (possibly component has not streamed in). */
 	inline bool IsBaseUnresolved() const
 	{
-		return (MovementBase == nullptr) && bServerHasBaseComponent;
+		return (!MovementBaseInterfaceData.IsValid()) && bServerHasBaseComponent;
 	}
 };
 
@@ -415,7 +512,11 @@ public:
 #endif
 
 	/** Sets the component the Character is walking on, used by CharacterMovement walking movement to be able to follow dynamic objects. */
+	UE_DEPRECATED(5.8, "Movement base functions using UPrimitiveComponents have been deprecated. Use functions using the FMovementBaseInterfaceData struct instead.")
 	ENGINE_API virtual void SetBase(UPrimitiveComponent* NewBase, const FName BoneName = NAME_None, bool bNotifyActor=true);
+
+	/** Sets the component the Character is walking on, used by CharacterMovement walking movement to be able to follow dynamic objects. */
+	ENGINE_API virtual void SetBase(FMovementBaseInterfaceData* MovementBaseInterfaceData, const FName BoneName = NAME_None, bool bNotifyActor=true);
 	
 	/**
 	 * Cache mesh offset from capsule. This is used as the target for network smoothing interpolation, when the mesh is offset with lagged smoothing.
@@ -692,7 +793,14 @@ public:
 	//~ Begin APawn Interface.
 	ENGINE_API virtual void PostInitializeComponents() override;
 	ENGINE_API virtual UPawnMovementComponent* GetMovementComponent() const override;
-	virtual UPrimitiveComponent* GetMovementBase() const override final { return BasedMovement.MovementBase; }
+	
+	UE_DEPRECATED(5.8, "GetMovementBase is deprecated. Use GetMovementBaseObject or GetMovementBaseInterfaceData instead.")
+	virtual UPrimitiveComponent* GetMovementBase() const override final { return Cast<UPrimitiveComponent>(GetMovementBaseObject()); }
+	
+	virtual UObject* GetMovementBaseObject() const override final { return BasedMovement.MovementBaseInterfaceData.GetMovementBaseObject(); }
+	ENGINE_API virtual const FMovementBaseInterfaceData* GetMovementBaseInterfaceData() const override final;
+	ENGINE_API virtual FMovementBaseInterfaceData* GetMovementBaseInterfaceData_Mutable() override final;
+	
 	ENGINE_API virtual float GetDefaultHalfHeight() const override;
 	ENGINE_API virtual void TurnOff() override;
 	ENGINE_API virtual void Restart() override;
diff --git a/Engine/Source/Runtime/Engine/Classes/GameFramework/CharacterMovementComponent.h b/Engine/Source/Runtime/Engine/Classes/GameFramework/CharacterMovementComponent.h
--- a/Engine/Source/Runtime/Engine/Classes/GameFramework/CharacterMovementComponent.h
+++ b/Engine/Source/Runtime/Engine/Classes/GameFramework/CharacterMovementComponent.h
@@ -20,6 +20,7 @@
 #include "GameFramework/CharacterMovementReplication.h"
 #include "Interfaces/NetworkPredictionInterface.h"
 #include "CharacterMovementComponentAsync.h"
+#include "Interfaces/MovementBaseInterface.h"
 #include "CharacterMovementComponent.generated.h"
 
 class ACharacter;
@@ -902,7 +903,11 @@ public:
 	TEnumAsByte<enum EMovementMode> DefaultWaterMovementMode;
 
 	/** Accessor to the Last Server Movement Base, only relevant during server update*/
-	UPrimitiveComponent* GetLastServerMovementBase() const { return LastServerMovementBase.Get(); };
+	UE_DEPRECATED(5.8, "Movement functions using UPrimitiveComponents have been deprecated (GetLastServerMovementBase included). Use GetLastServerMovementBaseInterfaceData instead.")
+	UPrimitiveComponent* GetLastServerMovementBase() const { return Cast<UPrimitiveComponent>(LastMovementBaseInterfaceData.PhysicsObjectOwner.Get()); }
+
+	/** Accessor to the Last Server Movement Base Interface Data, only relevant during server update*/
+	const FMovementBaseInterfaceData* GetLastServerMovementBaseInterfaceData() const { return &LastMovementBaseInterfaceData; }
 
 private:
 	/**
@@ -914,7 +919,11 @@ private:
 	TEnumAsByte<enum EMovementMode> GroundMovementMode;
 
 	/** Remember last server movement base so we can detect mounts/dismounts and respond accordingly. */
+	UE_DEPRECATED(5.8, "This property has been deprecated. Please use LastMovementBaseInterfaceData instead.")
 	TWeakObjectPtr<UPrimitiveComponent> LastServerMovementBase = nullptr;
+	
+	/** Remember last server movement base interface data so we can detect mounts/dismounts and respond accordingly. */
+	FMovementBaseInterfaceData LastMovementBaseInterfaceData;
 
 public:
 	/**
@@ -1381,9 +1390,20 @@ public:
 	ENGINE_API virtual void AdjustFloorHeight();
 
 	/** Return PrimitiveComponent we are based on (standing and walking on). */
-	UFUNCTION(BlueprintCallable, Category="Pawn|Components|CharacterMovement")
+	UE_DEPRECATED(5.8, "GetMovementBase has been deprecated. Use GetMovementBaseObject cast if necessary instead.")
+	UFUNCTION(BlueprintCallable, Category="Pawn|Components|CharacterMovement", meta=(DeprecatedFunction="Note", DeprecationMessage="GetMovementBase has been deprecated. Use GetMovementBaseObject cast if necessary instead."))
 	ENGINE_API UPrimitiveComponent* GetMovementBase() const;
 
+	/** Return Uobject we are based on (standing and walking on). */
+	UFUNCTION(BlueprintCallable, Category="Pawn|Components|CharacterMovement")
+	ENGINE_API UObject* GetMovementBaseObject() const;
+	
+	/** Returns struct containing references and functionality for UObject we're using as a MovementBase. Uses CharacterOwner to get relevant movement base data */
+	ENGINE_API const FMovementBaseInterfaceData* GetMovementBaseInterfaceData() const;
+	
+	/** Returns mutable struct containing references and functionality for UObject we're using as a MovementBase. Uses CharacterOwner to get relevant movement base data */
+	ENGINE_API FMovementBaseInterfaceData* GetMovementBaseInterfaceData_Mutable();
+	
 	/** Update or defer updating of position based on Base movement */
 	ENGINE_API virtual void MaybeUpdateBasedMovement(float DeltaSeconds);
 
@@ -1576,7 +1596,12 @@ public:
 	ENGINE_API virtual bool StepUp(const FVector& GravDir, const FVector& Delta, const FHitResult &Hit, FStepDownResult* OutStepDownResult = NULL);
 
 	/** Update the base of the character, which is the PrimitiveComponent we are standing on. */
+	UE_DEPRECATED(5.8, "Movement functions using UPrimitiveComponents have been deprecated. Use functions using the FMovementBaseInterfaceData struct instead.")
 	ENGINE_API virtual void SetBase(UPrimitiveComponent* NewBase, const FName BoneName = NAME_None, bool bNotifyActor=true);
+	
+	// TODO NS: Should this be a pointer?
+	/** Update the base of the character, which is the Object we are standing on. */
+	ENGINE_API virtual void SetBase(FMovementBaseInterfaceData* MovementBaseInterfaceData, const FName BoneName = NAME_None, bool bNotifyActor=true);
 
 	/**
 	 * Update the base of the character, using the given floor result if it is walkable, or null if not. Calls SetBase().
@@ -1816,11 +1841,14 @@ public:
 	/** Check if pawn is falling */
 	ENGINE_API virtual bool CheckFall(const FFindFloorResult& OldFloor, const FHitResult& Hit, const FVector& Delta, const FVector& OldLocation, float remainingTime, float timeTick, int32 Iterations, bool bMustJump);
 	
+	UE_DEPRECATED(5.8, "Movement functions using UPrimitiveComponents have been deprecated. Use functions using the FMovementBaseInterfaceData struct instead.")
+	ENGINE_API void RevertMove(const FVector& OldLocation, UPrimitiveComponent* OldBase, const FVector& InOldBaseLocation, const FFindFloorResult& OldFloor, bool bFailMove);
+
 	/** 
 	 *  Revert to previous position OldLocation, return to being based on OldBase.
 	 *  if bFailMove, stop movement and notify controller
 	 */	
-	ENGINE_API void RevertMove(const FVector& OldLocation, UPrimitiveComponent* OldBase, const FVector& InOldBaseLocation, const FFindFloorResult& OldFloor, bool bFailMove);
+	ENGINE_API void RevertMove(const FVector& OldLocation, FMovementBaseInterfaceData* OldMovementBaseInterfaceData, const FVector& InOldBaseLocation, const FFindFloorResult& OldFloor, bool bFailMove);
 
 	/** Perform rotation over deltaTime */
 	ENGINE_API virtual void PhysicsRotation(float DeltaTime);
@@ -2379,29 +2407,43 @@ protected:
 	 */
 	ENGINE_API virtual void CallServerMovePacked(const FSavedMove_Character* NewMove, const FSavedMove_Character* PendingMove, const FSavedMove_Character* OldMove);
 
+	UE_DEPRECATED(5.8, "Movement functions using UPrimitiveComponents have been deprecated. Use functions using the FMovementBaseInterfaceData struct instead.")
+	ENGINE_API virtual void ServerMoveHandleClientError(float ClientTimeStamp, float DeltaTime, const FVector& Accel, const FVector& RelativeClientLocation, UPrimitiveComponent* ClientMovementBase, FName ClientBaseBoneName, uint8 ClientMovementMode);
+	
 	/**
 	 * Have the server check if the client is outside an error tolerance, and queue a client adjustment if so.
 	 * If either GetPredictionData_Server_Character()->bForceClientUpdate or ServerCheckClientError() are true, the client adjustment will be sent.
 	 * RelativeClientLocation will be a relative location if MovementBaseUtility::UseRelativePosition(ClientMovementBase) is true, or a world location if false.
 	 * @see ServerCheckClientError()
 	 */
-	ENGINE_API virtual void ServerMoveHandleClientError(float ClientTimeStamp, float DeltaTime, const FVector& Accel, const FVector& RelativeClientLocation, UPrimitiveComponent* ClientMovementBase, FName ClientBaseBoneName, uint8 ClientMovementMode);
+	ENGINE_API virtual void ServerMoveHandleClientError(float ClientTimeStamp, float DeltaTime, const FVector& Accel, const FVector& RelativeClientLocation, FMovementBaseInterfaceData* ClientMovementBaseInterfaceData, FName ClientBaseBoneName, uint8 ClientMovementMode);
+
+
+	UE_DEPRECATED(5.8, "Movement functions using UPrimitiveComponents have been deprecated. Use functions using the FMovementBaseInterfaceData struct instead.")
+	ENGINE_API virtual bool ServerCheckClientError(float ClientTimeStamp, float DeltaTime, const FVector& Accel, const FVector& ClientWorldLocation, const FVector& RelativeClientLocation, UPrimitiveComponent* ClientMovementBase, FName ClientBaseBoneName, uint8 ClientMovementMode);
 
 	/**
 	 * Check for Server-Client disagreement in position or other movement state important enough to trigger a client correction.
 	 * @see ServerMoveHandleClientError()
 	 */
-	ENGINE_API virtual bool ServerCheckClientError(float ClientTimeStamp, float DeltaTime, const FVector& Accel, const FVector& ClientWorldLocation, const FVector& RelativeClientLocation, UPrimitiveComponent* ClientMovementBase, FName ClientBaseBoneName, uint8 ClientMovementMode);
+	ENGINE_API virtual bool ServerCheckClientError(float ClientTimeStamp, float DeltaTime, const FVector& Accel, const FVector& ClientWorldLocation, const FVector& RelativeClientLocation, FMovementBaseInterfaceData* ClientMovementBaseInterfaceData, FName ClientBaseBoneName, uint8 ClientMovementMode);
+
+	
+	UE_DEPRECATED(5.8, "Movement functions using UPrimitiveComponents have been deprecated. Use functions using the FMovementBaseInterfaceData struct instead.")
+	ENGINE_API virtual bool ServerExceedsAllowablePositionError(float ClientTimeStamp, float DeltaTime, const FVector& Accel, const FVector& ClientWorldLocation, const FVector& RelativeClientLocation, UPrimitiveComponent* ClientMovementBase, FName ClientBaseBoneName, uint8 ClientMovementMode);
 
 	/**
 	 * Check position error within ServerCheckClientError(). Set bNetworkLargeClientCorrection to true if the correction should be prioritized (delayed less in SendClientAdjustment).
 	 */
-	ENGINE_API virtual bool ServerExceedsAllowablePositionError(float ClientTimeStamp, float DeltaTime, const FVector& Accel, const FVector& ClientWorldLocation, const FVector& RelativeClientLocation, UPrimitiveComponent* ClientMovementBase, FName ClientBaseBoneName, uint8 ClientMovementMode);
+	ENGINE_API virtual bool ServerExceedsAllowablePositionError(float ClientTimeStamp, float DeltaTime, const FVector& Accel, const FVector& ClientWorldLocation, const FVector& RelativeClientLocation, FMovementBaseInterfaceData* ClientMovementBaseInterfaceData, FName ClientBaseBoneName, uint8 ClientMovementMode);
+
+	UE_DEPRECATED(5.8, "Movement functions using UPrimitiveComponents have been deprecated. Use functions using the FMovementBaseInterfaceData struct instead.")
+	ENGINE_API virtual bool ServerShouldUseAuthoritativePosition(float ClientTimeStamp, float DeltaTime, const FVector& Accel, const FVector& ClientWorldLocation, const FVector& RelativeClientLocation, UPrimitiveComponent* ClientMovementBase, FName ClientBaseBoneName, uint8 ClientMovementMode);
 
 	/**
 	 * If ServerCheckClientError() does not find an error, this determines if the server should also copy the client's movement params rather than keep the server sim result.
 	 */
-	ENGINE_API virtual bool ServerShouldUseAuthoritativePosition(float ClientTimeStamp, float DeltaTime, const FVector& Accel, const FVector& ClientWorldLocation, const FVector& RelativeClientLocation, UPrimitiveComponent* ClientMovementBase, FName ClientBaseBoneName, uint8 ClientMovementMode);
+	ENGINE_API virtual bool ServerShouldUseAuthoritativePosition(float ClientTimeStamp, float DeltaTime, const FVector& Accel, const FVector& ClientWorldLocation, const FVector& RelativeClientLocation, FMovementBaseInterfaceData* ClientMovementBaseInterfaceData, FName ClientBaseBoneName, uint8 ClientMovementMode);
 
 	/* Process a move at the given time stamp, given the compressed flags representing various events that occurred (ie jump). */
 	ENGINE_API virtual void MoveAutonomous( float ClientTimeStamp, float DeltaTime, uint8 CompressedFlags, const FVector& NewAccel);
@@ -2518,17 +2560,29 @@ public:
 	/* If no client adjustment is needed after processing received ServerMove(), ack the good move so client can remove it from SavedMoves */
 	ENGINE_API virtual void ClientAckGoodMove_Implementation(float TimeStamp);
 
-	/* Replicate position correction to client, associated with a timestamped servermove.  Client will replay subsequent moves after applying adjustment. */
+	UE_DEPRECATED(5.8, "Movement functions using UPrimitiveComponents have been deprecated. Use functions using the FMovementBaseInterfaceData struct instead.")
 	ENGINE_API virtual void ClientAdjustPosition_Implementation(float TimeStamp, FVector NewLoc, FVector NewVel, UPrimitiveComponent* NewBase, FName NewBaseBoneName, bool bHasBase, bool bBaseRelativePosition, uint8 ServerMovementMode, TOptional<FRotator> OptionalRotation = TOptional<FRotator>());
+	
+	/* Replicate position correction to client, associated with a timestamped servermove.  Client will replay subsequent moves after applying adjustment. */
+	ENGINE_API virtual void ClientAdjustPosition_Implementation(float TimeStamp, FVector NewLoc, FVector NewVel, FMovementBaseInterfaceData* NewMovementBaseInterfaceData, FName NewBaseBoneName, bool bHasBase, bool bBaseRelativePosition, uint8 ServerMovementMode, TOptional<FRotator> OptionalRotation = TOptional<FRotator>());
 
-	/* Bandwidth saving version, when velocity is zeroed */
+	UE_DEPRECATED(5.8, "Movement functions using UPrimitiveComponents have been deprecated. Use functions using the FMovementBaseInterfaceData struct instead.")
 	ENGINE_API virtual void ClientVeryShortAdjustPosition_Implementation(float TimeStamp, FVector NewLoc, UPrimitiveComponent* NewBase, FName NewBaseBoneName, bool bHasBase, bool bBaseRelativePosition, uint8 ServerMovementMode);
 
-	/* Replicate position correction to client when using root motion for movement. (animation root motion specific) */
+	/* Bandwidth saving version, when velocity is zeroed */
+	ENGINE_API virtual void ClientVeryShortAdjustPosition_Implementation(float TimeStamp, FVector NewLoc, FMovementBaseInterfaceData* NewMovementBaseInterfaceData, FName NewBaseBoneName, bool bHasBase, bool bBaseRelativePosition, uint8 ServerMovementMode);
+
+	UE_DEPRECATED(5.8, "Movement functions using UPrimitiveComponents have been deprecated. Use functions using the FMovementBaseInterfaceData struct instead.")
 	ENGINE_API virtual void ClientAdjustRootMotionPosition_Implementation(float TimeStamp, float ServerMontageTrackPosition, FVector ServerLoc, FVector_NetQuantizeNormal ServerRotation, float ServerVelZ, UPrimitiveComponent* ServerBase, FName ServerBoneName, bool bHasBase, bool bBaseRelativePosition, uint8 ServerMovementMode);
 
-	/* Replicate root motion source correction to client when using root motion for movement. */
+	/* Replicate position correction to client when using root motion for movement. (animation root motion specific) */
+	ENGINE_API virtual void ClientAdjustRootMotionPosition_Implementation(float TimeStamp, float ServerMontageTrackPosition, FVector ServerLoc, FVector_NetQuantizeNormal ServerRotation, float ServerVelZ, FMovementBaseInterfaceData* ServerMovementBaseInterfaceData, FName ServerBoneName, bool bHasBase, bool bBaseRelativePosition, uint8 ServerMovementMode);
+
+	UE_DEPRECATED(5.8, "Movement functions using UPrimitiveComponents have been deprecated. Use functions using the FMovementBaseInterfaceData struct instead.")
 	ENGINE_API virtual void ClientAdjustRootMotionSourcePosition_Implementation(float TimeStamp, FRootMotionSourceGroup ServerRootMotion, bool bHasAnimRootMotion, float ServerMontageTrackPosition, FVector ServerLoc, FVector_NetQuantizeNormal ServerRotation, float ServerVelZ, UPrimitiveComponent* ServerBase, FName ServerBoneName, bool bHasBase, bool bBaseRelativePosition, uint8 ServerMovementMode);
+	
+	/* Replicate root motion source correction to client when using root motion for movement. */
+	ENGINE_API virtual void ClientAdjustRootMotionSourcePosition_Implementation(float TimeStamp, FRootMotionSourceGroup ServerRootMotion, bool bHasAnimRootMotion, float ServerMontageTrackPosition, FVector ServerLoc, FVector_NetQuantizeNormal ServerRotation, float ServerVelZ, FMovementBaseInterfaceData* ServerMovementBaseInterfaceData, FName ServerBoneName, bool bHasBase, bool bBaseRelativePosition, uint8 ServerMovementMode);
 
 	/**
 	 * Handle movement data after it's unpacked from the ServerMovePacked_ServerReceive() call.
@@ -2624,9 +2678,12 @@ public:
 
 protected:
 
-	/** Event notification when client receives correction data from the server, before applying the data. Base implementation logs relevant data and draws debug info if "p.NetShowCorrections" is not equal to 0. */
+	UE_DEPRECATED(5.8, "Movement functions using UPrimitiveComponents have been deprecated. Use functions using the FMovementBaseInterfaceData struct instead.")
 	ENGINE_API virtual void OnClientCorrectionReceived(class FNetworkPredictionData_Client_Character& ClientData, float TimeStamp, FVector NewLocation, FVector NewVelocity, UPrimitiveComponent* NewBase, FName NewBaseBoneName, bool bHasBase, bool bBaseRelativePosition, uint8 ServerMovementMode, FVector ServerGravityDirection);
 
+	/** Event notification when client receives correction data from the server, before applying the data. Base implementation logs relevant data and draws debug info if "p.NetShowCorrections" is not equal to 0. */
+	ENGINE_API virtual void OnClientCorrectionReceived(class FNetworkPredictionData_Client_Character& ClientData, float TimeStamp, FVector NewLocation, FVector NewVelocity, FMovementBaseInterfaceData* NewMovementBaseInterfaceData, FName NewBaseBoneName, bool bHasBase, bool bBaseRelativePosition, uint8 ServerMovementMode, FVector ServerGravityDirection);
+
 	/**
 	 * Set custom struct used for client to server move RPC serialization.
 	 * This is typically set in the constructor for this component and should persist for the lifetime of the component.
@@ -2960,7 +3017,9 @@ public:
 	FQuat StartBaseRotation;	// rotation of the base component (or bone), only saved if it can move.
 	float StartCapsuleRadius;
 	float StartCapsuleHalfHeight;
+	UE_DEPRECATED(5.8, "This property has been deprecated. Please use StartMovementBaseInterfaceData instead.")
 	TWeakObjectPtr<UPrimitiveComponent> StartBase;
+	FMovementBaseInterfaceData StartMovementBaseInterfaceData;
 	FName StartBoneName;
 	uint32 StartActorOverlapCounter;
 	uint32 StartComponentOverlapCounter;
@@ -2976,7 +3035,9 @@ public:
 	FVector SavedRelativeLocation;
 	FVector SavedRelativeAcceleration;
 	FRotator SavedControlRotation;
+	UE_DEPRECATED(5.8, "This property has been deprecated. Please use EndMovementBaseInterfaceData instead.")
 	TWeakObjectPtr<UPrimitiveComponent> EndBase;
+	FMovementBaseInterfaceData EndMovementBaseInterfaceData;
 	FName EndBoneName;
 	uint32 EndActorOverlapCounter;
 	uint32 EndComponentOverlapCounter;
diff --git a/Engine/Source/Runtime/Engine/Classes/GameFramework/CharacterMovementReplication.h b/Engine/Source/Runtime/Engine/Classes/GameFramework/CharacterMovementReplication.h
--- a/Engine/Source/Runtime/Engine/Classes/GameFramework/CharacterMovementReplication.h
+++ b/Engine/Source/Runtime/Engine/Classes/GameFramework/CharacterMovementReplication.h
@@ -103,6 +103,7 @@ public:
 		OldMove
 	};
 
+	PRAGMA_DISABLE_DEPRECATION_WARNINGS
 	FCharacterNetworkMoveData()
 		: NetworkMoveType(ENetworkMoveType::NewMove)
 		, TimeStamp(0.f)
@@ -112,12 +113,18 @@ public:
 		, CompressedMoveFlags(0)
 		, MovementMode(0)
 		, MovementBase(nullptr)
+		, MovementBasePhysicsObjectOwner(nullptr)
 		, MovementBaseBoneName(NAME_None)
 	{
 	}
 	
 	virtual ~FCharacterNetworkMoveData() = default;
-
+	FCharacterNetworkMoveData(const FCharacterNetworkMoveData&) = default;
+	FCharacterNetworkMoveData(FCharacterNetworkMoveData&&) = default;
+	FCharacterNetworkMoveData& operator=(const FCharacterNetworkMoveData&) = default;
+	FCharacterNetworkMoveData& operator=(FCharacterNetworkMoveData&&) = default;
+	PRAGMA_ENABLE_DEPRECATION_WARNINGS
+	
 	/**
 	 * Given a FSavedMove_Character from UCharacterMovementComponent, fill in data in this struct with relevant movement data.
 	 * Note that the instance of the FSavedMove_Character is likely a custom struct of a derived struct of your own, if you have added your own saved move data.
@@ -145,7 +152,11 @@ public:
 	uint8 CompressedMoveFlags;
 
 	uint8 MovementMode;
+
+	UE_DEPRECATED(5.8, "This property has been deprecated. Please use MovementBasePhysicsObjectOwner instead and use that to create a FMovementBaseInterfaceData. See UCharacterMovementComponent::ServerMove_PerformMovement for an example of this.")
 	class UPrimitiveComponent* MovementBase;
+
+	TObjectPtr<UObject> MovementBasePhysicsObjectOwner;
 	FName MovementBaseBoneName;
 };
 
@@ -250,6 +261,7 @@ struct FClientAdjustment
 {
 public:
 
+	PRAGMA_DISABLE_DEPRECATION_WARNINGS
 	FClientAdjustment()
 		: TimeStamp(0.f)
 		, DeltaTime(0.f)
@@ -258,6 +270,7 @@ public:
 		, NewRot(ForceInitToZero)
 		, GravityDirection(FVector::DownVector)
 		, NewBase(NULL)
+		, NewMovementBasePhysicsObjectOwner(nullptr)
 		, NewBaseBoneName(NAME_None)
 		, bAckGoodMove(false)
 		, bBaseRelativePosition(false)
@@ -266,6 +279,14 @@ public:
 	{
 	}
 
+	~FClientAdjustment() = default;
+	FClientAdjustment(const FClientAdjustment&) = default;
+	FClientAdjustment(FClientAdjustment&&) = default;
+	FClientAdjustment& operator=(const FClientAdjustment&) = default;
+	FClientAdjustment& operator=(FClientAdjustment&&) = default;
+	PRAGMA_ENABLE_DEPRECATION_WARNINGS
+	
+	
 	float TimeStamp;
 	float DeltaTime;
 	FVector NewLoc; // Note: if bBaseRelativePosition is set, this is a relative location to the Movement base.
@@ -274,10 +295,14 @@ public:
 	FVector GravityDirection;
 
 #if UE_WITH_REMOTE_OBJECT_HANDLE
+	UE_DEPRECATED(5.8, "This property has been deprecated. Please use NewMovementBasePhysicsObjectOwner instead and use that to create a FMovementBaseInterfaceData. See UCharacterMovementComponent::ClientHandleMoveResponse for an example of this.")
 	TObjectPtr<UPrimitiveComponent> NewBase;
 #else
+	UE_DEPRECATED(5.8, "This property has been deprecated. Please use NewMovementBasePhysicsObjectOwner instead and use that to create a FMovementBaseInterfaceData. See UCharacterMovementComponent::ClientHandleMoveResponse for an example of this.")
 	UPrimitiveComponent* NewBase;
 #endif
+	
+	TObjectPtr<UObject> NewMovementBasePhysicsObjectOwner;
 	FName NewBaseBoneName;
 	bool bAckGoodMove;
 	bool bBaseRelativePosition;
@@ -294,7 +319,10 @@ public:
 			Archive << NewVel;
 			Archive << NewRot;
 			Archive << GravityDirection;
-			Archive << NewBase;
+			Archive << NewMovementBasePhysicsObjectOwner;
+			PRAGMA_DISABLE_DEPRECATION_WARNINGS
+			NewBase = Cast<UPrimitiveComponent>(NewMovementBasePhysicsObjectOwner.Get());
+			PRAGMA_ENABLE_DEPRECATION_WARNINGS
 			Archive << NewBaseBoneName;
 			Archive.SerializeBits(&bAckGoodMove, 1);
 			Archive.SerializeBits(&bBaseRelativePosition, 1);
diff --git a/Engine/Source/Runtime/Engine/Classes/GameFramework/Pawn.h b/Engine/Source/Runtime/Engine/Classes/GameFramework/Pawn.h
--- a/Engine/Source/Runtime/Engine/Classes/GameFramework/Pawn.h
+++ b/Engine/Source/Runtime/Engine/Classes/GameFramework/Pawn.h
@@ -13,6 +13,7 @@
 
 #include "Pawn.generated.h"
 
+struct FMovementBaseInterfaceData;
 class AController;
 class APhysicsVolume;
 class APlayerController;
@@ -55,8 +56,18 @@ public:
 	ENGINE_API virtual UPawnMovementComponent* GetMovementComponent() const;
 
 	/** Return PrimitiveComponent we are based on (standing on, attached to, and moving on). */
+	UE_DEPRECATED(5.8, "GetMovementBase is deprecated. Use GetMovementBaseObject or GetMovementBaseInterfaceData instead.")
 	virtual UPrimitiveComponent* GetMovementBase() const { return nullptr; }
 
+	/** Return UObject we are based on (standing on, attached to, and moving on). */
+	virtual UObject* GetMovementBaseObject() const { return nullptr; }
+
+	/** Return MovementBaseInterfaceData for the object we are based on (standing on, attached to, and moving on). */
+	virtual const FMovementBaseInterfaceData* GetMovementBaseInterfaceData() const { return nullptr; }
+	
+	/** Return mutable MovementBaseInterfaceData for the object we are based on (standing on, attached to, and moving on). */
+	virtual FMovementBaseInterfaceData* GetMovementBaseInterfaceData_Mutable(){ return nullptr; }
+	
 	/** If true, this Pawn's pitch will be updated to match the Controller's ControlRotation pitch, if controlled by a PlayerController. */
 	UPROPERTY(EditAnywhere, BlueprintReadWrite, Category=Pawn)
 	uint32 bUseControllerRotationPitch:1;
@@ -229,6 +240,10 @@ public:
 	UFUNCTION(BlueprintPure, Category=Pawn)
 	static ENGINE_API AActor* GetMovementBaseActor(const APawn* Pawn);
 
+	/** Gets the owning UObject of the Movement Base object on which the pawn is standing. */
+	UFUNCTION(BlueprintPure, Category=Pawn)
+	static ENGINE_API UObject* GetMovementBaseOwner(const APawn* Pawn);
+	
 	/** Return true if yaw is within AllowedYawError of desired yaw */
 	ENGINE_API virtual bool ReachedDesiredRotation();
 
diff --git a/Engine/Source/Runtime/Engine/Classes/Interfaces/MovementBaseInterface.cpp b/Engine/Source/Runtime/Engine/Classes/Interfaces/MovementBaseInterface.cpp
new file mode 100644
--- /dev/null
+++ b/Engine/Source/Runtime/Engine/Classes/Interfaces/MovementBaseInterface.cpp
@@ -0,0 +1,108 @@
+﻿// Copyright Epic Games, Inc. All Rights Reserved.
+
+#include "MovementBaseInterface.h"
+
+#include "IPhysicsComponent.h"
+#include "Engine/HitResult.h"
+
+DEFINE_LOG_CATEGORY_STATIC(LogCharacterMovementBase, Log, All);
+
+FMovementBaseInterfaceData::FMovementBaseInterfaceData()
+	: PhysicsObjectOwner(nullptr)
+	, BodyInstanceOwner(nullptr)
+{
+}
+
+FMovementBaseInterfaceData::FMovementBaseInterfaceData(const FMovementBaseInterfaceData* OtherMovementBaseData)
+{
+	if (OtherMovementBaseData)
+	{
+		BodyInstanceOwner = OtherMovementBaseData->BodyInstanceOwner;
+		PhysicsObjectOwner = OtherMovementBaseData->PhysicsObjectOwner;
+	}
+	else
+	{
+		BodyInstanceOwner = nullptr;
+		PhysicsObjectOwner = nullptr;
+	}
+}
+
+FMovementBaseInterfaceData::FMovementBaseInterfaceData(TObjectPtr<UObject> PhysicsObjectOwner)
+{
+	Set(PhysicsObjectOwner);
+}
+
+FMovementBaseInterfaceData::FMovementBaseInterfaceData(TObjectPtr<UObject> PhysicsObjectOwner, IPhysicsBodyInstanceOwner* PhysicsBodyInstanceOwner)
+	: PhysicsObjectOwner(PhysicsObjectOwner)
+	, BodyInstanceOwner(PhysicsBodyInstanceOwner)
+{
+}
+
+bool FMovementBaseInterfaceData::operator==(const FMovementBaseInterfaceData& Other) const
+{
+	return PhysicsObjectOwner == Other.PhysicsObjectOwner && BodyInstanceOwner == Other.BodyInstanceOwner;
+}
+
+bool FMovementBaseInterfaceData::operator!=(const FMovementBaseInterfaceData& Other) const
+{
+	return !(FMovementBaseInterfaceData::operator==(Other));
+}
+
+bool FMovementBaseInterfaceData::IsValid() const
+{
+	return PhysicsObjectOwner.IsValid() && BodyInstanceOwner != nullptr;
+}
+
+void FMovementBaseInterfaceData::Clear()
+{
+	PhysicsObjectOwner = nullptr;
+	BodyInstanceOwner = nullptr;
+}
+
+void FMovementBaseInterfaceData::Set(TObjectPtr<UObject> NewPhysicsObjectOwner, const IPhysicsBodyInstanceOwner* NewPhysicsBodyInstanceOwner)
+{
+	PhysicsObjectOwner = NewPhysicsObjectOwner;
+	BodyInstanceOwner = NewPhysicsBodyInstanceOwner;
+}
+
+void FMovementBaseInterfaceData::Set(TObjectPtr<UObject> NewPhysicsObjectOwner)
+{
+	IPhysicsBodyInstanceOwner* NewBodyInstanceOwner = nullptr;
+	
+	if (NewPhysicsObjectOwner)
+	{
+		// Todo: Pass along some sort of ID so we can get something useful for objects owning multiple physics objects
+		if (IPhysicsBodyInstanceOwnerResolver* PhysicsObjectOwnerResolver = Cast<IPhysicsBodyInstanceOwnerResolver>(NewPhysicsObjectOwner))
+		{
+			Chaos::FConstPhysicsObjectHandle PhysicsObject = nullptr;
+			if (const IPhysicsComponent* PhysicsComponent = Cast<IPhysicsComponent>(NewPhysicsObjectOwner))
+			{
+				PhysicsObject = PhysicsComponent->GetPhysicsObjectById(0);
+			}
+				
+			NewBodyInstanceOwner = PhysicsObjectOwnerResolver->ResolvePhysicsBodyInstanceOwner(PhysicsObject);
+		}
+		else
+		{
+			UE_LOGF(LogCharacterMovementBase, Warning, "%ls: UObject passed into FMovementBaseInterfaceData::Set is expected to implement IPhysicsBodyInstanceOwnerResolver. If the Object passed in implements IPhysicsBodyInstanceOwner directly consider using FMovementBaseInterfaceData::Set(TObjectPtr<UObject> NewPhysicsObjectOwner, const IPhysicsBodyInstanceOwner* NewPhysicsBodyInstanceOwner) instead.", *GetNameSafe(NewPhysicsObjectOwner));
+		}
+	}
+	
+	PhysicsObjectOwner = NewPhysicsObjectOwner;
+	BodyInstanceOwner = NewBodyInstanceOwner;
+}
+
+const IPhysicsBodyInstanceOwner* FMovementBaseInterfaceData::GetBodyInstanceOwner() const
+{
+	return PhysicsObjectOwner.IsValid() ? BodyInstanceOwner : nullptr;
+}
+
+UObject* FMovementBaseInterfaceData::GetMovementBaseObject() const
+{
+	return IsValid() ? BodyInstanceOwner->GetSourceObject() : nullptr;
+}
+
+UObject* FMovementBaseInterfaceData::GetMovementBaseObjectOwner() const
+{
+	return IsValid() ? BodyInstanceOwner->GetSourceObjectOwner() : nullptr;
+}
diff --git a/Engine/Source/Runtime/Engine/Classes/Interfaces/MovementBaseInterface.h b/Engine/Source/Runtime/Engine/Classes/Interfaces/MovementBaseInterface.h
new file mode 100644
--- /dev/null
+++ b/Engine/Source/Runtime/Engine/Classes/Interfaces/MovementBaseInterface.h
@@ -0,0 +1,74 @@
+﻿// Copyright Epic Games, Inc. All Rights Reserved.
+
+#pragma once
+
+#include "PhysicsEngine/PhysicsBodyInstanceOwnerInterface.h"
+#include "UObject/Interface.h"
+#include "MovementBaseInterface.generated.h"
+
+#define UE_API ENGINE_API
+
+USTRUCT()
+struct FMovementBaseInterfaceData
+{
+	GENERATED_BODY()
+
+	// UObject that owns the physics representation of the movement base object
+	UPROPERTY()
+	TWeakObjectPtr<UObject> PhysicsObjectOwner;
+
+protected:
+	// IPhysicsBodyInstanceOwner reference that implements movement base functionality 
+	const IPhysicsBodyInstanceOwner* BodyInstanceOwner;
+	
+public:
+	UE_API FMovementBaseInterfaceData();
+	UE_API FMovementBaseInterfaceData(const FMovementBaseInterfaceData* OtherMovementBaseData);
+	UE_API FMovementBaseInterfaceData(TObjectPtr<UObject> PhysicsObjectOwner);
+	UE_API FMovementBaseInterfaceData(TObjectPtr<UObject> PhysicsObjectOwner, IPhysicsBodyInstanceOwner* PhysicsBodyInstanceOwner);
+	
+	UE_API bool operator==(const FMovementBaseInterfaceData& Other) const;
+	UE_API bool operator!=(const FMovementBaseInterfaceData& Other) const;
+	explicit operator bool() const { return IsValid(); }
+	
+	/** Dangerous access: Ensures validity before returning */
+	const IPhysicsBodyInstanceOwner* operator->() const 
+	{ 
+		const IPhysicsBodyInstanceOwner* Resolved = GetBodyInstanceOwner();
+		check(Resolved);
+		return Resolved; 
+	}
+
+	/** Returns true if this struct has all relevant movement base data set */
+	UE_API bool IsValid() const;
+
+	/** Clears relevant movement base data */
+	UE_API void Clear();
+	
+	/** Sets relevant movement base data */
+	UE_API void Set(TObjectPtr<UObject> NewPhysicsObjectOwner, const IPhysicsBodyInstanceOwner* NewPhysicsBodyInstanceOwner);
+	
+	/** 
+	 * Sets relevant movement base data. 
+	 * Uses IPhysicsBodyInstanceOwnerResolver and IPhysicsComponent to find and set IPhysicsBodyInstanceOwner
+	 */
+	UE_API void Set(TObjectPtr<UObject> NewPhysicsObjectOwner);
+	
+	/**
+	 * Returns IPhysicsBodyInstanceOwner that implements movement base functionality
+	 * Note: The IPhysicsBodyInstanceOwner returned is only valid if the PhysicsObjectOwner is as well so the IPhysicsBodyInstanceOwner returned shouldn't be cached
+	 */
+	UE_API const IPhysicsBodyInstanceOwner* GetBodyInstanceOwner() const;
+
+	/**
+	 * Helper function that returns UObject that owns movement base functionality
+	 */
+	UE_API UObject* GetMovementBaseObject() const;
+
+	/**
+	 * Helper function that returns the owner of the UObject that owns movement base functionality
+	 */
+	UE_API UObject* GetMovementBaseObjectOwner() const;
+};
+
+#undef UE_API
diff --git a/Engine/Source/Runtime/Engine/Classes/PhysicsEngine/PhysicsBodyInstanceOwnerInterface.h b/Engine/Source/Runtime/Engine/Classes/PhysicsEngine/PhysicsBodyInstanceOwnerInterface.h
--- a/Engine/Source/Runtime/Engine/Classes/PhysicsEngine/PhysicsBodyInstanceOwnerInterface.h
+++ b/Engine/Source/Runtime/Engine/Classes/PhysicsEngine/PhysicsBodyInstanceOwnerInterface.h
@@ -65,9 +65,17 @@ public:
 	/** Returns the source UObject of this body instance owner. */
 	virtual UObject* GetSourceObject() const = 0;
 
+	/** Returns the owner of the source object for the body instance */
+	virtual UObject* GetSourceObjectOwner() const = 0;
+	
 	/** Returns the transform of this body instance owner. */
 	virtual FTransform GetPhysicsOwnerTransform() const = 0;
 
+	/** 
+	 * Get world-space socket transform of the socket matching InSocketName. 
+	 */
+	virtual FTransform GetPhysicsOwnerSocketTransform(FName InSocketName) const = 0;
+	
 	/** Returns the channel that the body instance owner belongs to. */
 	virtual ECollisionChannel GetCollisionObjectType() const = 0;
 
@@ -80,9 +88,30 @@ public:
 	/** Gets the response type given a specific channel. */
 	virtual ECollisionResponse GetCollisionResponseToChannel(ECollisionChannel Channel) const = 0;
 
+	/** Returns the slope override struct for this physics body instance. */
+	virtual const struct FWalkableSlopeOverride& GetWalkableSlopeOverride() const = 0;
+	
 	/** Returns a physics object corresponding to the specified Id. */
 	virtual Chaos::FPhysicsObject* GetPhysicsObjectById(Chaos::FPhysicsObjectId Id) const = 0;
 
+	/** Returns true if the body instance owner is movable */
+	virtual bool IsPhysicsOwnerMovable() const = 0;
+	
+	/** Returns true if the body instance owner simulating physics */
+	virtual bool IsPhysicsOwnerSimulatingPhysics() const = 0;
+
+	/** Returns velocity of body instance owner. It is often the velocity of the body instance unless overriden by another system like movement */
+	virtual FVector GetPhysicsOwnerVelocity() const = 0;
+
+	/** Get the attachment root of the body instance owner object */
+	virtual UObject* GetPhysicsOwnerAttachmentRoot() const = 0;
+	
+	/** Returns true if this body instance should be considered world geometry */
+	virtual bool IsPhysicsObjectWorldGeometry() const = 0;
+	
+	/** Checks if a socket of the name passed in exists on the object owning the body instance */
+	virtual bool DoesSocketExistOnPhysicsOwner(FName InSocketName) const = 0;
+
 	/** Returns all physics objects associated to the physics body instance owner. */
 	virtual TArray<Chaos::FPhysicsObject*> GetAllPhysicsObjects() const = 0;
 
diff --git a/Engine/Source/Runtime/Engine/Private/Character.cpp b/Engine/Source/Runtime/Engine/Private/Character.cpp
--- a/Engine/Source/Runtime/Engine/Private/Character.cpp
+++ b/Engine/Source/Runtime/Engine/Private/Character.cpp
@@ -20,6 +20,7 @@
 #include "Engine/Canvas.h"
 #include "Animation/AnimInstance.h"
 #include "Engine/DamageEvents.h"
+#include "PhysicsEngine/PhysicsBodyInstanceOwnerInterface.h"
 
 #include UE_INLINE_GENERATED_CPP_BY_NAME(Character)
 
@@ -213,6 +214,16 @@ UPawnMovementComponent* ACharacter::GetMovementComponent() const
 	return CharacterMovement;
 }
 
+const FMovementBaseInterfaceData* ACharacter::GetMovementBaseInterfaceData() const
+{
+	return &BasedMovement.MovementBaseInterfaceData;
+}
+
+FMovementBaseInterfaceData* ACharacter::GetMovementBaseInterfaceData_Mutable()
+{
+	return &BasedMovement.MovementBaseInterfaceData;
+}
+
 
 void ACharacter::SetupPlayerInputComponent(UInputComponent* PlayerInputComponent)
 {
@@ -563,9 +574,15 @@ void ACharacter::ApplyDamageMomentum(float DamageTaken, FDamageEvent const& Dama
 
 void ACharacter::ClearCrossLevelReferences()
 {
-	if( BasedMovement.MovementBase != nullptr && GetOutermost() != BasedMovement.MovementBase->GetOutermost() )
+	if(BasedMovement.MovementBaseInterfaceData.IsValid())
 	{
-		SetBase( nullptr );
+		const UObject* BodyInstanceOwnerSourceObject = BasedMovement.MovementBaseInterfaceData.GetMovementBaseObject();
+		const UPackage* BodyInstanceOwnerOuterMost = BodyInstanceOwnerSourceObject ? BodyInstanceOwnerSourceObject->GetOutermost() : nullptr;
+		if (GetOutermost() != BodyInstanceOwnerOuterMost)
+		{
+			FMovementBaseInterfaceData* DummyMovementBasePointer = nullptr;
+			SetBase(DummyMovementBasePointer);
+		}
 	}
 
 	Super::ClearCrossLevelReferences();
@@ -592,6 +609,7 @@ namespace MovementBaseUtility
 
 	void AddTickDependency(FTickFunction& BasedObjectTick, UPrimitiveComponent* NewBase)
 	{
+		PRAGMA_DISABLE_DEPRECATION_WARNINGS
 		if (NewBase && MovementBaseUtility::UseRelativeLocation(NewBase))
 		{
 			if (NewBase->PrimaryComponentTick.bCanEverTick)
@@ -618,10 +636,12 @@ namespace MovementBaseUtility
 				}
 			}
 		}
+		PRAGMA_ENABLE_DEPRECATION_WARNINGS
 	}
 
 	void RemoveTickDependency(FTickFunction& BasedObjectTick, UPrimitiveComponent* OldBase)
 	{
+		PRAGMA_DISABLE_DEPRECATION_WARNINGS
 		if (OldBase && MovementBaseUtility::UseRelativeLocation(OldBase))
 		{
 			BasedObjectTick.RemovePrerequisite(OldBase, OldBase->PrimaryComponentTick);
@@ -640,10 +660,12 @@ namespace MovementBaseUtility
 				}
 			}
 		}
+		PRAGMA_ENABLE_DEPRECATION_WARNINGS
 	}
 
 	FVector GetMovementBaseVelocity(const UPrimitiveComponent* MovementBase, const FName BoneName)
 	{
+		PRAGMA_DISABLE_DEPRECATION_WARNINGS
 		FVector BaseVelocity = FVector::ZeroVector;
 		if (MovementBaseUtility::IsDynamicBase(MovementBase))
 		{
@@ -680,10 +702,12 @@ namespace MovementBaseUtility
 		}
 		
 		return BaseVelocity;
+		PRAGMA_ENABLE_DEPRECATION_WARNINGS
 	}
 
 	FVector GetMovementBaseTangentialVelocity(const UPrimitiveComponent* MovementBase, const FName BoneName, const FVector& WorldLocation)
 	{
+		PRAGMA_DISABLE_DEPRECATION_WARNINGS
 		if (MovementBaseUtility::IsDynamicBase(MovementBase))
 		{
 			if (const FBodyInstance* BodyInstance = MovementBase->GetBodyInstance(BoneName))
@@ -702,6 +726,7 @@ namespace MovementBaseUtility
 				}
 			}			
 		}
+		PRAGMA_ENABLE_DEPRECATION_WARNINGS
 		
 		return FVector::ZeroVector;
 	}
@@ -743,7 +768,9 @@ namespace MovementBaseUtility
 	{
 		FVector OutLocation;
 		FQuat OutQuat;
+		PRAGMA_DISABLE_DEPRECATION_WARNINGS
 		const bool bResult = GetMovementBaseTransform(MovementBase, BoneName, OutLocation, OutQuat);
+		PRAGMA_ENABLE_DEPRECATION_WARNINGS
 		OutLocationWorldSpace = FTransform(OutQuat, OutLocation).TransformPositionNoScale(LocalLocation);
 		return bResult;
 	}
@@ -752,7 +779,9 @@ namespace MovementBaseUtility
 	{
 		FVector OutLocation;
 		FQuat OutQuat;
+		PRAGMA_DISABLE_DEPRECATION_WARNINGS
 		const bool bResult = GetMovementBaseTransform(MovementBase, BoneName, OutLocation, OutQuat);
+		PRAGMA_ENABLE_DEPRECATION_WARNINGS
 		OutLocalLocation = FTransform(OutQuat, OutLocation).InverseTransformPositionNoScale(WorldSpaceLocation);
 		return bResult;
 	}
@@ -761,12 +790,238 @@ namespace MovementBaseUtility
 	{
 		FVector IgnoredLocation;
 		FQuat OutQuat;
+		PRAGMA_DISABLE_DEPRECATION_WARNINGS
 		const bool bResult = GetMovementBaseTransform(MovementBase, BoneName, IgnoredLocation, OutQuat);
+		PRAGMA_ENABLE_DEPRECATION_WARNINGS
 		OutDirectionWorldSpace = OutQuat.RotateVector(LocalDirection);
 		return bResult;
 	}
 
 	bool TransformDirectionToLocal(const UPrimitiveComponent* MovementBase, const FName BoneName, const FVector& WorldSpaceDirection, FVector& OutLocalDirection)
+	{
+		FVector IgnoredLocation;
+		FQuat OutQuat;
+		PRAGMA_DISABLE_DEPRECATION_WARNINGS
+		const bool bResult = GetMovementBaseTransform(MovementBase, BoneName, IgnoredLocation, OutQuat);
+		PRAGMA_ENABLE_DEPRECATION_WARNINGS
+		OutLocalDirection = OutQuat.UnrotateVector(WorldSpaceDirection);
+		return bResult;
+	}
+
+	bool IsDynamicBase(const FMovementBaseInterfaceData* MovementBase)
+	{
+		return (MovementBase && MovementBase->IsValid() && MovementBase->GetBodyInstanceOwner()->IsPhysicsOwnerMovable());
+	}
+
+	bool IsSimulatedBase(const FMovementBaseInterfaceData* MovementBase)
+	{
+		return (MovementBase && MovementBase->IsValid() && MovementBase->GetBodyInstanceOwner()->IsPhysicsOwnerSimulatingPhysics());
+	}
+
+	void AddTickDependency(FTickFunction& BasedObjectTick, FMovementBaseInterfaceData* MovementBase)
+	{
+		// Todo: Eventually convert this to not use PrimitiveComponents
+		if (MovementBase)
+		{
+			UPrimitiveComponent* NewBase = Cast<UPrimitiveComponent>(MovementBase->PhysicsObjectOwner.Get());
+		
+			if (NewBase && MovementBaseUtility::UseRelativeLocation(MovementBase))
+			{
+				if (NewBase->PrimaryComponentTick.bCanEverTick)
+				{
+					BasedObjectTick.AddPrerequisite(NewBase, NewBase->PrimaryComponentTick);
+				}
+
+				AActor* NewBaseOwner = NewBase->GetOwner();
+				if (NewBaseOwner)
+				{
+					if (NewBaseOwner->PrimaryActorTick.bCanEverTick)
+					{
+						BasedObjectTick.AddPrerequisite(NewBaseOwner, NewBaseOwner->PrimaryActorTick);
+					}
+
+					// @TODO: We need to find a more efficient way of finding all ticking components in an actor.
+					for (UActorComponent* Component : NewBaseOwner->GetComponents())
+					{
+						// Dont allow a based component (e.g. a particle system) to push us into a different tick group
+						if (Component && Component->PrimaryComponentTick.bCanEverTick && Component->PrimaryComponentTick.TickGroup <= BasedObjectTick.TickGroup)
+						{
+							BasedObjectTick.AddPrerequisite(Component, Component->PrimaryComponentTick);
+						}
+					}
+				}
+			}
+		}
+	}
+
+	void RemoveTickDependency(FTickFunction& BasedObjectTick, FMovementBaseInterfaceData* MovementBase)
+	{
+		// Todo: Eventually convert this to not use PrimitiveComponents
+		if (MovementBase)
+		{
+			UPrimitiveComponent* OldBase = Cast<UPrimitiveComponent>(MovementBase->PhysicsObjectOwner.Get());
+		
+			if (OldBase && MovementBaseUtility::UseRelativeLocation(MovementBase))
+			{
+				BasedObjectTick.RemovePrerequisite(OldBase, OldBase->PrimaryComponentTick);
+				AActor* OldBaseOwner = OldBase->GetOwner();
+				if (OldBaseOwner)
+				{
+					BasedObjectTick.RemovePrerequisite(OldBaseOwner, OldBaseOwner->PrimaryActorTick);
+
+					// @TODO: We need to find a more efficient way of finding all ticking components in an actor.
+					for (UActorComponent* Component : OldBaseOwner->GetComponents())
+					{
+						if (Component && Component->PrimaryComponentTick.bCanEverTick)
+						{
+							BasedObjectTick.RemovePrerequisite(Component, Component->PrimaryComponentTick);
+						}
+					}
+				}
+			}
+		}
+	}
+
+	FVector GetMovementBaseVelocity(const FMovementBaseInterfaceData* MovementBase, const FName BoneName)
+	{
+		FVector BaseVelocity = FVector::ZeroVector;
+		if (IsDynamicBase(MovementBase))
+		{
+			if (BoneName != NAME_None)
+			{
+				if (const FBodyInstance* BodyInstance = MovementBase->GetBodyInstanceOwner()->GetBodyInstance(BoneName))
+				{
+					BaseVelocity = BodyInstance->GetUnrealWorldVelocity();
+					return BaseVelocity;
+				}
+			}
+
+			BaseVelocity = MovementBase->GetBodyInstanceOwner()->GetPhysicsOwnerVelocity();
+			if (BaseVelocity.IsZero())
+			{
+				// Fall back to actor's Root component
+				const AActor* Owner = Cast<const AActor>(MovementBase->GetMovementBaseObjectOwner());
+				if (Owner)
+				{
+					// Component might be moved manually (not by simulated physics or a movement component), see if the root component of the actor has a velocity.
+					BaseVelocity = Owner->GetVelocity();
+				}
+			}
+
+			// Fall back to physics velocity.
+			if (BaseVelocity.IsZero())
+			{
+				if (const FBodyInstance* BodyInstance = MovementBase->GetBodyInstanceOwner()->GetBodyInstance())
+				{
+					BaseVelocity = BodyInstance->GetUnrealWorldVelocity();
+				}
+			}
+		}
+		
+		return BaseVelocity;
+	}
+
+	FVector GetMovementBaseTangentialVelocity(const FMovementBaseInterfaceData* MovementBase, const FName BoneName, const FVector& WorldLocation)
+	{
+		if (IsDynamicBase(MovementBase))
+		{
+			if (const FBodyInstance* BodyInstance = MovementBase->GetBodyInstanceOwner()->GetBodyInstance(BoneName))
+			{
+				const FVector BaseAngVelInRad = BodyInstance->GetUnrealWorldAngularVelocityInRadians();
+				if (!BaseAngVelInRad.IsNearlyZero())
+				{
+					FVector BaseLocation;
+					FQuat BaseRotation;
+					if (GetMovementBaseTransform(MovementBase, BoneName, BaseLocation, BaseRotation))
+					{
+						const FVector RadialDistanceToBase = WorldLocation - BaseLocation;
+						const FVector TangentialVel = BaseAngVelInRad ^ RadialDistanceToBase;
+						return TangentialVel;
+					}
+				}
+			}
+			
+			const FBodyInstance* BodyInstance = MovementBase->GetBodyInstanceOwner()->GetBodyInstance();
+			const FVector BaseAngVelInRad = BodyInstance ? BodyInstance->GetUnrealWorldAngularVelocityInRadians() : FVector::ZeroVector;
+			if (!BaseAngVelInRad.IsNearlyZero())
+			{
+				FVector BaseLocation;
+				FQuat BaseRotation;
+				if (GetMovementBaseTransform(MovementBase, BoneName, BaseLocation, BaseRotation))
+				{
+					const FVector RadialDistanceToBase = WorldLocation - BaseLocation;
+					const FVector TangentialVel = BaseAngVelInRad ^ RadialDistanceToBase;
+					return TangentialVel;
+				}
+			}
+		}
+		
+		return FVector::ZeroVector;
+	}
+
+	bool GetMovementBaseTransform(const FMovementBaseInterfaceData* MovementBase, const FName BoneName, FVector& OutLocation, FQuat& OutQuat)
+	{
+		if (MovementBase && MovementBase->IsValid())
+		{
+			bool bGotTransformOfIntendedBone = true;
+
+			if (BoneName != NAME_None)
+			{
+				// Check if this socket or bone exists (DoesSocketExist checks for either, as does requesting the transform).
+				if (MovementBase->GetBodyInstanceOwner()->DoesSocketExistOnPhysicsOwner(BoneName))
+				{
+					const FTransform BodyOwnerSocketTransform = MovementBase->GetBodyInstanceOwner()->GetPhysicsOwnerSocketTransform(BoneName);
+					OutLocation = BodyOwnerSocketTransform.GetLocation();
+					OutQuat = BodyOwnerSocketTransform.GetRotation();
+					return true;
+				}
+				else
+				{
+					UE_LOG(LogCharacter, Warning, TEXT("GetMovementBaseTransform(): Invalid bone or socket '%s' for PrimitiveComponent base %s. Falling back to base's root transform."), *BoneName.ToString(), *GetPathNameSafe(MovementBase->GetMovementBaseObject()));
+					bGotTransformOfIntendedBone = false;
+				}
+			}
+			
+			// No bone supplied
+			OutLocation = MovementBase->GetBodyInstanceOwner()->GetPhysicsOwnerTransform().GetLocation();
+			OutQuat = MovementBase->GetBodyInstanceOwner()->GetPhysicsOwnerTransform().GetRotation();
+			return bGotTransformOfIntendedBone;
+		}
+
+		// nullptr MovementBase
+		OutLocation = FVector::ZeroVector;
+		OutQuat = FQuat::Identity;
+		return false;
+	}
+
+	bool TransformLocationToWorld(const FMovementBaseInterfaceData* MovementBase, const FName BoneName, const FVector& LocalLocation, FVector& OutLocationWorldSpace)
+	{
+		FVector OutLocation;
+		FQuat OutQuat;
+		const bool bResult = GetMovementBaseTransform(MovementBase, BoneName, OutLocation, OutQuat);
+		OutLocationWorldSpace = FTransform(OutQuat, OutLocation).TransformPositionNoScale(LocalLocation);
+		return bResult;
+	}
+
+	bool TransformLocationToLocal(const FMovementBaseInterfaceData* MovementBase, const FName BoneName, const FVector& WorldSpaceLocation, FVector& OutLocalLocation)
+	{
+		FVector OutLocation;
+		FQuat OutQuat;
+		const bool bResult = GetMovementBaseTransform(MovementBase, BoneName, OutLocation, OutQuat);
+		OutLocalLocation = FTransform(OutQuat, OutLocation).InverseTransformPositionNoScale(WorldSpaceLocation);
+		return bResult;
+	}
+
+	bool TransformDirectionToWorld(const FMovementBaseInterfaceData* MovementBase, const FName BoneName, const FVector& LocalDirection, FVector& OutDirectionWorldSpace)
+	{
+		FVector IgnoredLocation;
+		FQuat OutQuat;
+		const bool bResult = GetMovementBaseTransform(MovementBase, BoneName, IgnoredLocation, OutQuat);
+		OutDirectionWorldSpace = OutQuat.RotateVector(LocalDirection);
+		return bResult;
+	}
+
+	bool TransformDirectionToLocal(const FMovementBaseInterfaceData* MovementBase, const FName BoneName, const FVector& WorldSpaceDirection, FVector& OutLocalDirection)
 	{
 		FVector IgnoredLocation;
 		FQuat OutQuat;
@@ -775,44 +1030,175 @@ namespace MovementBaseUtility
 		return bResult;
 	}
 
+	bool DoesMovementBaseDataMatch(const FMovementBaseInterfaceData* MovementBaseDataA, const FMovementBaseInterfaceData* MovementBaseDataB)
+	{
+		// if they're both null, they match
+		if (MovementBaseDataA == nullptr && MovementBaseDataB == nullptr)
+		{
+			return true;
+		}
+		
+		// if only one is null, they don't match
+		if (MovementBaseDataA == nullptr || MovementBaseDataB == nullptr)
+		{
+			return false;
+		}
+		
+		// Compare actual values
+		if (MovementBaseDataA && MovementBaseDataB)
+		{
+			return *MovementBaseDataA == *MovementBaseDataB;
+		}
+		
+		return false;
+	}
+
+	bool IsMovementBaseDataValid(const FMovementBaseInterfaceData* MovementBaseData)
+	{
+		return MovementBaseData && MovementBaseData->IsValid();
+	}
+	
+	FMovementBaseInterfaceData GetMovementBaseDataFromPhysicsOwner(UObject* PhysicsObjectOwner)
+	{
+		/* TODO NS Movement Base Review Notes:
+ 		 * This is where we have to build up a IPhysicsBodyInstanceOwner interface from a PhysicsObjectOwner.
+ 		 * This is needed as IPhysicsBodyInstanceOwner isn't a UInterface so we need to instead replicate the PhysicsObjectOwner UObject and gather the rest of the info we need.
+ 		 * Which goes something like this:
+ 		 *  - First we need the Physics Handle (in order to get IPhysicsBodyInstanceOwner) - No good way to do this besides the IPhysicsComponent and UPrimitiveComponent route.
+ 		 *	  We also don't 100% need a physics handle. For instance, SG IPhysicsBodyInstanceOwnerResolver only has one body instance to point to so we just point to that rather
+ 		 *	  than using the handle to find something.
+ 		 *	- Now that we have the struct built we can use helper functions that use the struct or the IPhysicsBodyInstanceOwner directly through the struct
+ 		 *
+ 		 *	Overall, this is a lot of work to just build up the relevant data needed to make a MovementBaseData struct so I question if it's preferable. This can be avoided by
+ 		 *	just implementing functionality through a specific MovementBase UInterface but that requires adding that very specific interface to SG objects and Actor components
+ 		 *	so I want to balance complexity vs adding interfaces to solve every problem. IPhysicsBodyInstanceOwner avoids this problem as it's fairly broad and can be used
+ 		 *	for other things than just movement bases.
+ 		 *
+ 		 */
+		FMovementBaseInterfaceData ReturnMovementBaseData;
+		
+		if (PhysicsObjectOwner)
+		{
+			IPhysicsBodyInstanceOwner* BodyInstanceOwner = nullptr;
+			// TODO NS: Should we just implement IPhysicsBodyInstanceOwnerResolver on UPrimitiveComponent too? Avoid this specific check
+			if (UPrimitiveComponent* PhysicsObjectOwnerAsPrimitiveComponent = Cast<UPrimitiveComponent>(PhysicsObjectOwner))
+			{
+				BodyInstanceOwner = PhysicsObjectOwnerAsPrimitiveComponent;
+			}
+
+			if (IPhysicsBodyInstanceOwnerResolver* PhysicsObjectOwnerResolver = Cast<IPhysicsBodyInstanceOwnerResolver>(PhysicsObjectOwner))
+			{
+				Chaos::FConstPhysicsObjectHandle PhysicsObject = nullptr;
+				if (const IPhysicsComponent* PhysicsComponent = Cast<IPhysicsComponent>(PhysicsObjectOwner))
+				{
+					// Todo: Pass along some sort of ID so we can get something useful for objects owning multiple physics objects
+					PhysicsObject = PhysicsComponent->GetPhysicsObjectById(0);
+				}
+				
+				BodyInstanceOwner = PhysicsObjectOwnerResolver->ResolvePhysicsBodyInstanceOwner(PhysicsObject);
+			}
+			
+			ReturnMovementBaseData.Set(PhysicsObjectOwner, BodyInstanceOwner);
+		}
+		
+		return ReturnMovementBaseData;
+	}
+
+	FMovementBaseInterfaceData GetMovementBaseDataFromHitResult(const FHitResult* HitResult)
+	{
+		FMovementBaseInterfaceData ReturnMovementBaseData;
+
+		if (HitResult)
+		{
+			const IPhysicsBodyInstanceOwner* BodyInstanceOwner = nullptr;
+			// TODO NS: Should we just implement IPhysicsBodyInstanceOwnerResolver on UPrimitiveComponent too? Avoid this specific check
+			if (UPrimitiveComponent* PhysicsObjectOwnerAsPrimitiveComponent = Cast<UPrimitiveComponent>(HitResult->PhysicsObjectOwner.Get()))
+			{
+				BodyInstanceOwner = PhysicsObjectOwnerAsPrimitiveComponent;
+			}
+		
+			if (!BodyInstanceOwner)
+			{
+				BodyInstanceOwner = IPhysicsBodyInstanceOwner::GetPhysicsBodyInstanceOwnerFromHitResult(*HitResult);
+			}
+			
+			ReturnMovementBaseData.Set(HitResult->PhysicsObjectOwner.Get(), BodyInstanceOwner);
+			
+			// Fall back to using UPrimitiveComponent Component in hit result if we didn't get anything valid from Physics Object and Component is set
+			if (!ReturnMovementBaseData.IsValid() && HitResult->Component.IsValid())
+			{
+				ReturnMovementBaseData = FMovementBaseInterfaceData(HitResult->Component.Get());
+			}
+		}
+		
+		return ReturnMovementBaseData;
+	}
 }
 
 
 /**	Change the Pawn's base. */
 void ACharacter::SetBase( UPrimitiveComponent* NewBaseComponent, const FName InBoneName, bool bNotifyPawn )
+{
+	FMovementBaseInterfaceData MovementBaseInterfaceData(NewBaseComponent);
+	SetBase(&MovementBaseInterfaceData, InBoneName, bNotifyPawn);
+}
+
+void ACharacter::SetBase(FMovementBaseInterfaceData* MovementBaseInterfaceData, const FName InBoneName, bool bNotifyPawn)
 {
 	// If NewBaseComponent is nullptr, ignore bone name.
-	const FName BoneName = (NewBaseComponent ? InBoneName : NAME_None);
+	const FName BoneName = (MovementBaseInterfaceData ? InBoneName : NAME_None);
 
 	// See what changed.
-	const bool bBaseChanged = (NewBaseComponent != BasedMovement.MovementBase);
+	const bool bBaseChanged = MovementBaseInterfaceData ? (MovementBaseInterfaceData->GetBodyInstanceOwner() != BasedMovement.MovementBaseInterfaceData.GetBodyInstanceOwner()) : BasedMovement.MovementBaseInterfaceData.IsValid();
 	const bool bBoneChanged = (BoneName != BasedMovement.BoneName);
 
 	if (bBaseChanged || bBoneChanged)
 	{
 		// Verify no recursion.
-		APawn* Loop = (NewBaseComponent ? Cast<APawn>(NewBaseComponent->GetOwner()) : nullptr);
-		while (Loop)
+		if (MovementBaseInterfaceData && MovementBaseInterfaceData->IsValid())
 		{
-			if (Loop == this)
+			APawn* Loop = Cast<APawn>(MovementBaseInterfaceData->GetMovementBaseObjectOwner());
+			while (Loop)
 			{
-				UE_LOG(LogCharacter, Warning, TEXT(" SetBase failed! Recursion detected. Pawn %s already based on %s."), *GetName(), *NewBaseComponent->GetName());
-				return;
-			}
-			if (UPrimitiveComponent* LoopBase =	Loop->GetMovementBase())
-			{
-				Loop = Cast<APawn>(LoopBase->GetOwner());
-			}
-			else
-			{
-				break;
+				if (Loop == this)
+				{
+					UE_LOG(LogCharacter, Warning, TEXT(" SetBase failed! Recursion detected. Pawn %s already based on %s."), *GetName(), *MovementBaseInterfaceData->PhysicsObjectOwner->GetName());
+					return;
+				}
+				
+				bool bValidMovementBase = false;
+				if (const FMovementBaseInterfaceData* LoopMovementBaseData = Loop->GetMovementBaseInterfaceData())
+				{
+					if (LoopMovementBaseData->IsValid())
+					{
+						Loop = Cast<APawn>(LoopMovementBaseData->GetMovementBaseObjectOwner());
+						bValidMovementBase = true;
+					}
+				}
+				
+				if (!bValidMovementBase)
+				{
+					break;
+				}
 			}
 		}
 
 		// Set base.
-		UPrimitiveComponent* OldBase = BasedMovement.MovementBase;
-		BasedMovement.MovementBase = NewBaseComponent;
+		FMovementBaseInterfaceData OldMovementBaseInterfaceData = BasedMovement.MovementBaseInterfaceData;
 		BasedMovement.BoneName = BoneName;
+		if (MovementBaseInterfaceData)
+		{
+			BasedMovement.MovementBaseInterfaceData = *MovementBaseInterfaceData;
+		}
+		else
+		{
+			BasedMovement.MovementBaseInterfaceData.Clear();
+		}
+		
+		PRAGMA_DISABLE_DEPRECATION_WARNINGS
+		BasedMovement.MovementBase = MovementBaseInterfaceData ? Cast<UPrimitiveComponent>(BasedMovement.MovementBaseInterfaceData.PhysicsObjectOwner) : nullptr;
+		PRAGMA_ENABLE_DEPRECATION_WARNINGS
+		
 		if (bBaseChanged)
 		{
 			// Increment base ID.
@@ -821,18 +1207,18 @@ void ACharacter::SetBase( UPrimitiveComponent* NewBaseComponent, const FName InB
 
 		if (CharacterMovement)
 		{
-			const bool bBaseIsSimulating = MovementBaseUtility::IsSimulatedBase(NewBaseComponent);
+			const bool bBaseIsSimulating = MovementBaseUtility::IsSimulatedBase(&BasedMovement.MovementBaseInterfaceData);
 			if (bBaseChanged)
 			{
-				MovementBaseUtility::RemoveTickDependency(CharacterMovement->PrimaryComponentTick, OldBase);
+				MovementBaseUtility::RemoveTickDependency(CharacterMovement->PrimaryComponentTick, &OldMovementBaseInterfaceData);
 				// We use a special post physics function if simulating, otherwise add normal tick prereqs.
 				if (!bBaseIsSimulating)
 				{
-					MovementBaseUtility::AddTickDependency(CharacterMovement->PrimaryComponentTick, NewBaseComponent);
+					MovementBaseUtility::AddTickDependency(CharacterMovement->PrimaryComponentTick, &BasedMovement.MovementBaseInterfaceData);
 				}
 			}
 
-			if (NewBaseComponent)
+			if (BasedMovement.MovementBaseInterfaceData.IsValid())
 			{
 				// Update OldBaseLocation/Rotation as those were referring to a different base
 				// ... but not when handling replication for proxies (since they are going to copy this data from the replicated values anyway)
@@ -854,14 +1240,15 @@ void ACharacter::SetBase( UPrimitiveComponent* NewBaseComponent, const FName InB
 			}
 
 			const ENetRole LocalRole = GetLocalRole();
+			const UObject* MovementBaseSourceObject = BasedMovement.MovementBaseInterfaceData.GetMovementBaseObject();
 			if (LocalRole == ROLE_Authority || LocalRole == ROLE_AutonomousProxy)
 			{
-				BasedMovement.bServerHasBaseComponent = (BasedMovement.MovementBase != nullptr); // Also set on autonomous proxies for nicer debugging.
-				UE_LOG(LogCharacter, Verbose, TEXT("Setting base on %s for '%s' to '%s'"), LocalRole == ROLE_Authority ? TEXT("Server") : TEXT("AutoProxy"), *GetName(), *GetFullNameSafe(NewBaseComponent));
+				BasedMovement.bServerHasBaseComponent = (BasedMovement.MovementBaseInterfaceData.GetBodyInstanceOwner() != nullptr); // Also set on autonomous proxies for nicer debugging.
+				UE_LOG(LogCharacter, Verbose, TEXT("Setting base on %s for '%s' to '%s'"), LocalRole == ROLE_Authority ? TEXT("Server") : TEXT("AutoProxy"), *GetName(), *GetFullNameSafe(MovementBaseSourceObject));
 			}
 			else
 			{
-				UE_LOG(LogCharacter, Verbose, TEXT("Setting base on Client for '%s' to '%s'"), *GetName(), *GetFullNameSafe(NewBaseComponent));
+				UE_LOG(LogCharacter, Verbose, TEXT("Setting base on Client for '%s' to '%s'"), *GetName(), *GetFullNameSafe(MovementBaseSourceObject));
 			}
 
 		}
@@ -1080,13 +1467,13 @@ void ACharacter::DisplayDebug(UCanvas* Canvas, const FDebugDisplayInfo& DebugDis
 		FIndenter PhysicsIndent(Indent);
 
 		FString BaseString;
-		if ( CharacterMovement == nullptr || BasedMovement.MovementBase == nullptr )
+		if ( CharacterMovement == nullptr || !BasedMovement.MovementBaseInterfaceData.IsValid() )
 		{
 			BaseString = "Not Based";
 		}
 		else
 		{
-			BaseString = BasedMovement.MovementBase->IsWorldGeometry() ? "World Geometry" : BasedMovement.MovementBase->GetName();
+			BaseString = BasedMovement.MovementBaseInterfaceData.GetBodyInstanceOwner()->IsPhysicsObjectWorldGeometry() ? "World Geometry" : BasedMovement.MovementBaseInterfaceData.GetMovementBaseObject()->GetName();
 			BaseString = FString::Printf(TEXT("Based On %s"), *BaseString);
 		}
 		
@@ -1285,30 +1672,38 @@ void ACharacter::OnRep_ReplicatedBasedMovement()
 	FGuardValue_Bitfield(bInBaseReplication, true);
 
 	const FBasedMovementInfo& RepBasedMovement = GetReplicatedBasedMovement();
-	const bool bBaseChanged = (BasedMovement.MovementBase != RepBasedMovement.MovementBase || BasedMovement.BoneName != RepBasedMovement.BoneName);
+	FMovementBaseInterfaceData RepMovementBaseInterfaceData = MovementBaseUtility::GetMovementBaseDataFromPhysicsOwner(RepBasedMovement.MovementBaseInterfaceData.PhysicsObjectOwner.Get());
+	
+	const bool bBaseChanged = (BasedMovement.MovementBaseInterfaceData != RepMovementBaseInterfaceData || BasedMovement.BoneName != RepBasedMovement.BoneName);
 	if (bBaseChanged)
 	{
 		// Even though we will copy the replicated based movement info, we need to use SetBase() to set up tick dependencies and trigger notifications.
-		SetBase(RepBasedMovement.MovementBase, RepBasedMovement.BoneName);
+		SetBase(&RepMovementBaseInterfaceData, RepBasedMovement.BoneName);
 	}
 
 	// Make sure to use the values of relative location/rotation etc from the server.
 	BasedMovement = RepBasedMovement;
+	BasedMovement.MovementBaseInterfaceData = RepMovementBaseInterfaceData;
+	ReplicatedBasedMovement.MovementBaseInterfaceData = RepMovementBaseInterfaceData;
 
-	if (RepBasedMovement.HasRelativeLocation())
+	PRAGMA_DISABLE_DEPRECATION_WARNINGS
+	BasedMovement.MovementBase = RepBasedMovement.MovementBase;
+	PRAGMA_ENABLE_DEPRECATION_WARNINGS
+	
+	if (BasedMovement.HasRelativeLocation())
 	{
 		// Update transform relative to movement base
 		const FVector OldLocation = GetActorLocation();
 		const FQuat OldRotation = GetActorQuat();
-		MovementBaseUtility::GetMovementBaseTransform(RepBasedMovement.MovementBase, RepBasedMovement.BoneName, CharacterMovement->OldBaseLocation, CharacterMovement->OldBaseQuat);
+		MovementBaseUtility::GetMovementBaseTransform(&BasedMovement.MovementBaseInterfaceData, BasedMovement.BoneName, CharacterMovement->OldBaseLocation, CharacterMovement->OldBaseQuat);
 		const FTransform BaseTransform(CharacterMovement->OldBaseQuat, CharacterMovement->OldBaseLocation);
-		const FVector NewLocation = BaseTransform.TransformPositionNoScale(RepBasedMovement.Location);
+		const FVector NewLocation = BaseTransform.TransformPositionNoScale(BasedMovement.Location);
 		FRotator NewRotation;
 
-		if (RepBasedMovement.HasRelativeRotation())
+		if (BasedMovement.HasRelativeRotation())
 		{
 			// Relative location, relative rotation
-			NewRotation = (FRotationMatrix(RepBasedMovement.Rotation) * FQuatRotationMatrix(CharacterMovement->OldBaseQuat)).Rotator();
+			NewRotation = (FRotationMatrix(BasedMovement.Rotation) * FQuatRotationMatrix(CharacterMovement->OldBaseQuat)).Rotator();
 			
 			if (CharacterMovement->ShouldRemainVertical())
 			{
@@ -1319,7 +1714,7 @@ void ACharacter::OnRep_ReplicatedBasedMovement()
 		else
 		{
 			// Relative location, absolute rotation
-			NewRotation = RepBasedMovement.Rotation;
+			NewRotation = BasedMovement.Rotation;
 		}
 
 		// When position or base changes, movement mode will need to be updated. This assumes rotation changes don't affect that.
@@ -1527,18 +1922,19 @@ bool ACharacter::CanUseRootMotionRepMove(const FSimulatedRootMotionReplicatedMov
 
 bool ACharacter::RestoreReplicatedMove(const FSimulatedRootMotionReplicatedMove& RootMotionRepMove)
 {
-	UPrimitiveComponent* ServerBase = RootMotionRepMove.RootMotion.MovementBase;
 	const FName ServerBaseBoneName = RootMotionRepMove.RootMotion.MovementBaseBoneName;
 
+	FMovementBaseInterfaceData ServerMovementBaseInterfaceData = MovementBaseUtility::GetMovementBaseDataFromPhysicsOwner(RootMotionRepMove.RootMotion.PhysicsObjectOwner);
+	
 	// Relative Position
 	if( RootMotionRepMove.RootMotion.bRelativePosition )
 	{
 		bool bSuccess = false;
-		if( MovementBaseUtility::UseRelativeLocation(ServerBase) )
+		if( MovementBaseUtility::UseRelativeLocation(&ServerMovementBaseInterfaceData) )
 		{
 			FVector BaseLocation;
 			FQuat BaseRotation;
-			MovementBaseUtility::GetMovementBaseTransform(ServerBase, ServerBaseBoneName, BaseLocation, BaseRotation);
+			MovementBaseUtility::GetMovementBaseTransform(&ServerMovementBaseInterfaceData, ServerBaseBoneName, BaseLocation, BaseRotation);
 			const FTransform BaseTransform(BaseRotation, BaseLocation);
 			
 			const FVector ServerLocation = BaseTransform.TransformPositionNoScale(RootMotionRepMove.RootMotion.Location);
@@ -1571,7 +1967,7 @@ bool ACharacter::RestoreReplicatedMove(const FSimulatedRootMotionReplicatedMove&
 	}
 
 	CharacterMovement->bJustTeleported = true;
-	SetBase( ServerBase, ServerBaseBoneName );
+	SetBase(&ServerMovementBaseInterfaceData, ServerBaseBoneName);
 
 	return true;
 }
@@ -1603,8 +1999,11 @@ void ACharacter::PostNetReceiveLocationAndRotation()
 {
 	if(GetLocalRole() == ROLE_SimulatedProxy)
 	{
+		const FBasedMovementInfo& RepBasedMovement = GetReplicatedBasedMovement();
+		FMovementBaseInterfaceData RepMovementBaseInterfaceData = MovementBaseUtility::GetMovementBaseDataFromPhysicsOwner(RepBasedMovement.MovementBaseInterfaceData.PhysicsObjectOwner.Get());
+		
 		// Don't change transform if using relative position (it should be nearly the same anyway, or base may be slightly out of sync)
-		if (!GetReplicatedBasedMovement().HasRelativeLocation())
+		if (!MovementBaseUtility::UseRelativeLocation(&RepMovementBaseInterfaceData))
 		{
 			const FRepMovement& ConstRepMovement = GetReplicatedMovement();
 			const FVector OldLocation = GetActorLocation();
@@ -1635,7 +2034,10 @@ void ACharacter::PreReplication( IRepChangedPropertyTracker & ChangedPropertyTra
 		RepRootMotionRef.bRelativeRotation = BasedMovement.HasRelativeRotation();
 		RepRootMotionRef.Location			= RepRootMotionRef.bRelativePosition ? BasedMovement.Location : FRepMovement::RebaseOntoZeroOrigin(GetActorLocation(), GetWorld()->OriginLocation);
 		RepRootMotionRef.Rotation			= RepRootMotionRef.bRelativeRotation ? BasedMovement.Rotation : GetActorRotation();
-		RepRootMotionRef.MovementBase		= BasedMovement.MovementBase;
+		RepRootMotionRef.PhysicsObjectOwner = BasedMovement.MovementBaseInterfaceData.PhysicsObjectOwner.Get();
+		PRAGMA_DISABLE_DEPRECATION_WARNINGS
+		RepRootMotionRef.MovementBase = BasedMovement.MovementBase;
+		PRAGMA_ENABLE_DEPRECATION_WARNINGS
 		RepRootMotionRef.MovementBaseBoneName = BasedMovement.BoneName;
 		if (RootMotionMontageInstance)
 		{
@@ -1967,25 +2369,29 @@ void ACharacter::ClientAckGoodMove_Implementation(float TimeStamp)
 // ClientAdjustPosition
 void ACharacter::ClientAdjustPosition_Implementation(float TimeStamp, FVector NewLoc, FVector NewVel, UPrimitiveComponent* NewBase, FName NewBaseBoneName, bool bHasBase, bool bBaseRelativePosition, uint8 ServerMovementMode)
 {
-	GetCharacterMovement()->ClientAdjustPosition_Implementation(TimeStamp, NewLoc, NewVel, NewBase, NewBaseBoneName, bHasBase, bBaseRelativePosition, ServerMovementMode);
+	FMovementBaseInterfaceData NewMovementBaseData(NewBase);
+	GetCharacterMovement()->ClientAdjustPosition_Implementation(TimeStamp, NewLoc, NewVel, &NewMovementBaseData, NewBaseBoneName, bHasBase, bBaseRelativePosition, ServerMovementMode);
 }
 
 // ClientVeryShortAdjustPosition
 void ACharacter::ClientVeryShortAdjustPosition_Implementation(float TimeStamp, FVector NewLoc, UPrimitiveComponent* NewBase, FName NewBaseBoneName, bool bHasBase, bool bBaseRelativePosition, uint8 ServerMovementMode)
 {
-	GetCharacterMovement()->ClientVeryShortAdjustPosition_Implementation(TimeStamp, NewLoc, NewBase, NewBaseBoneName, bHasBase, bBaseRelativePosition, ServerMovementMode);
+	FMovementBaseInterfaceData NewMovementBaseData(NewBase);
+	GetCharacterMovement()->ClientVeryShortAdjustPosition_Implementation(TimeStamp, NewLoc, &NewMovementBaseData, NewBaseBoneName, bHasBase, bBaseRelativePosition, ServerMovementMode);
 }
 
 // ClientAdjustRootMotionPosition
 void ACharacter::ClientAdjustRootMotionPosition_Implementation(float TimeStamp, float ServerMontageTrackPosition, FVector ServerLoc, FVector_NetQuantizeNormal ServerRotation, float ServerVelZ, UPrimitiveComponent* ServerBase, FName ServerBoneName, bool bHasBase, bool bBaseRelativePosition, uint8 ServerMovementMode)
 {
-	GetCharacterMovement()->ClientAdjustRootMotionPosition_Implementation(TimeStamp, ServerMontageTrackPosition, ServerLoc, ServerRotation, ServerVelZ, ServerBase, ServerBoneName, bHasBase, bBaseRelativePosition, ServerMovementMode);
+	FMovementBaseInterfaceData ServerMovementBaseData(ServerBase);
+	GetCharacterMovement()->ClientAdjustRootMotionPosition_Implementation(TimeStamp, ServerMontageTrackPosition, ServerLoc, ServerRotation, ServerVelZ, &ServerMovementBaseData, ServerBoneName, bHasBase, bBaseRelativePosition, ServerMovementMode);
 }
 
 // ClientAdjustRootMotionSourcePosition
 void ACharacter::ClientAdjustRootMotionSourcePosition_Implementation(float TimeStamp, FRootMotionSourceGroup ServerRootMotion, bool bHasAnimRootMotion, float ServerMontageTrackPosition, FVector ServerLoc, FVector_NetQuantizeNormal ServerRotation, float ServerVelZ, UPrimitiveComponent* ServerBase, FName ServerBoneName, bool bHasBase, bool bBaseRelativePosition, uint8 ServerMovementMode)
 {
-	GetCharacterMovement()->ClientAdjustRootMotionSourcePosition_Implementation(TimeStamp, ServerRootMotion, bHasAnimRootMotion, ServerMontageTrackPosition, ServerLoc, ServerRotation, ServerVelZ, ServerBase, ServerBoneName, bHasBase, bBaseRelativePosition, ServerMovementMode);
+	FMovementBaseInterfaceData ServerMovementBaseData(ServerBase);
+	GetCharacterMovement()->ClientAdjustRootMotionSourcePosition_Implementation(TimeStamp, ServerRootMotion, bHasAnimRootMotion, ServerMontageTrackPosition, ServerLoc, ServerRotation, ServerVelZ, &ServerMovementBaseData, ServerBoneName, bHasBase, bBaseRelativePosition, ServerMovementMode);
 }
 
 void ACharacter::FillAsyncInput(FCharacterAsyncInput& Input) const
diff --git a/Engine/Source/Runtime/Engine/Private/Components/CharacterMovementComponent.cpp b/Engine/Source/Runtime/Engine/Private/Components/CharacterMovementComponent.cpp
--- a/Engine/Source/Runtime/Engine/Private/Components/CharacterMovementComponent.cpp
+++ b/Engine/Source/Runtime/Engine/Private/Components/CharacterMovementComponent.cpp
@@ -18,6 +18,7 @@
 #include "GameFramework/PhysicsVolume.h"
 #include "Components/SkeletalMeshComponent.h"
 #include "DrawDebugHelpers.h"
+#include "ObjectCacheContext.h"
 #include "GameFramework/GameNetworkManager.h"
 #include "GameFramework/Character.h"
 #include "Components/CapsuleComponent.h"
@@ -26,11 +27,15 @@
 #include "AI/Navigation/AvoidanceManager.h"
 #include "Components/BrushComponent.h"
 #include "PBDRigidsSolver.h"
+#include "Chaos/ChaosUserEntity.h"
 #include "Engine/NetworkObjectList.h"
 #include "Net/PerfCountersHelpers.h"
 #include "Net/Iris/ReplicationSystem/ReplicationSystemUtil.h"
 #include "Net/Iris/ReplicationSystem/EngineReplicationBridge.h"
 #include "Iris/ReplicationSystem/ReplicationSystem.h"
+#include "PhysicsEngine/PhysicsBodyInstanceOwnerInterface.h"
+#include "PhysicsEngine/PhysicsObjectExternalInterface.h"
+#include "Runtime/Experimental/Chaos/Private/Chaos/PhysicsObjectInternal.h"
 
 #include UE_INLINE_GENERATED_CPP_BY_NAME(CharacterMovementComponent)
 
@@ -482,6 +487,7 @@ namespace CharacterMovementCVars
 #endif // !(UE_BUILD_SHIPPING || UE_BUILD_TEST)
 }
 
+extern ENGINE_API FString GPlayInEditorContextString;
 
 /**
  * Helper to change mesh bone updates within a scope.
@@ -1184,15 +1190,16 @@ FVector UCharacterMovementComponent::GetImpartedMovementBaseVelocity() const
 	FVector Result = FVector::ZeroVector;
 	if (CharacterOwner)
 	{
-		UPrimitiveComponent* MovementBase = CharacterOwner->GetMovementBase();
-		if (MovementBaseUtility::IsDynamicBase(MovementBase))
+		const FMovementBaseInterfaceData* MovementBaseInterfaceData = GetMovementBaseInterfaceData();
+		
+		if (MovementBaseUtility::IsDynamicBase(MovementBaseInterfaceData))
 		{
-			FVector BaseVelocity = MovementBaseUtility::GetMovementBaseVelocity(MovementBase, CharacterOwner->GetBasedMovement().BoneName);
+			FVector BaseVelocity = MovementBaseUtility::GetMovementBaseVelocity(MovementBaseInterfaceData, CharacterOwner->GetBasedMovement().BoneName);
 			
 			if (bImpartBaseAngularVelocity)
 			{
 				const FVector CharacterBasePosition = (UpdatedComponent->GetComponentLocation() + CharacterOwner->GetCapsuleComponent()->GetScaledCapsuleHalfHeight() * GetGravityDirection());
-				const FVector BaseTangentialVel = MovementBaseUtility::GetMovementBaseTangentialVelocity(MovementBase, CharacterOwner->GetBasedMovement().BoneName, CharacterBasePosition);
+				const FVector BaseTangentialVel = MovementBaseUtility::GetMovementBaseTangentialVelocity(MovementBaseInterfaceData, CharacterOwner->GetBasedMovement().BoneName, CharacterBasePosition);
 				BaseVelocity += BaseTangentialVel;
 			}
 
@@ -1317,7 +1324,8 @@ void UCharacterMovementComponent::SetDefaultMovementMode()
 		SetMovementMode(DefaultLandMovementMode);
 
 		// Avoid 1-frame delay if trying to walk but walking fails at this location.
-		if (MovementMode == MOVE_Walking && GetMovementBase() == NULL)
+		const FMovementBaseInterfaceData* MovementBaseData = GetMovementBaseInterfaceData();
+		if (MovementMode == MOVE_Walking && !(MovementBaseData && MovementBaseData->IsValid()))
 		{
 			SetGravitySpaceZ(Velocity, SavedVelocityZ); // Prevent temporary walking state from zeroing Z velocity.
 			SetMovementMode(MOVE_Falling);
@@ -1454,7 +1462,8 @@ void UCharacterMovementComponent::OnMovementModeChanged(EMovementMode PreviousMo
 
 		if (!IsFalling() || bStayBasedInAir == false)
 		{
-			SetBase(NULL);
+			FMovementBaseInterfaceData EmptyMovementBaseData;
+			SetBase(&EmptyMovementBaseData);
 		}
 
 		if (MovementMode == MOVE_None)
@@ -1907,8 +1916,9 @@ void UCharacterMovementComponent::SimulatedTick(float DeltaSeconds)
 				DrawDebugCoordinateSystem(GetWorld(), CharacterOwner->GetMesh()->GetComponentLocation() + FVector(0,0,1), NewRotation, 50.f, false);
 				DrawDebugLine(GetWorld(), OldLocation, NewLocation, FColor::Red, false, 10.f);
 
+				const FMovementBaseInterfaceData* MovementBaseInterfaceData = GetMovementBaseInterfaceData();
 				UE_LOG(LogRootMotion, Log,  TEXT("UCharacterMovementComponent::SimulatedTick DeltaMovement Translation: %s, Rotation: %s, MovementBase: %s"),
-					*(NewLocation - OldLocation).ToCompactString(), *(NewRotation - OldRotation).GetNormalized().ToCompactString(), *GetNameSafe(CharacterOwner->GetMovementBase()) );
+					*(NewLocation - OldLocation).ToCompactString(), *(NewRotation - OldRotation).GetNormalized().ToCompactString(), *GetNameSafe(MovementBaseInterfaceData ? MovementBaseInterfaceData->GetMovementBaseObject() : nullptr));
 			}
 			#endif
 #endif // !(UE_BUILD_SHIPPING)
@@ -2192,8 +2202,11 @@ void UCharacterMovementComponent::SimulateMovement(float DeltaSeconds)
 		return;
 	}
 
+	const FBasedMovementInfo& RepBasedMovement = CharacterOwner->GetReplicatedBasedMovement();
+	FMovementBaseInterfaceData RepMovementBaseInterfaceData = MovementBaseUtility::GetMovementBaseDataFromPhysicsOwner(RepBasedMovement.MovementBaseInterfaceData.PhysicsObjectOwner.Get());
+	
 	// If base is not resolved on the client, we should not try to simulate at all
-	if (CharacterOwner->GetReplicatedBasedMovement().IsBaseUnresolved())
+	if (!RepMovementBaseInterfaceData.IsValid() && CharacterOwner->GetReplicatedBasedMovement().bServerHasBaseComponent)
 	{
 		UE_LOG(LogCharacterMovement, Verbose, TEXT("Base for simulated character '%s' is not resolved on client, skipping SimulateMovement"), *CharacterOwner->GetName());
 		return;
@@ -2298,7 +2311,7 @@ void UCharacterMovementComponent::SimulateMovement(float DeltaSeconds)
 
 				// Possible for dynamic movement bases, particularly those that align to slopes while the character does not, to encroach the character.
 				// Check to see if we can resolve the penetration in those cases, and if so find the floor.
-				if (CurrentFloor.HitResult.bStartPenetrating && MovementBaseUtility::IsDynamicBase(GetMovementBase()))
+				if (CurrentFloor.HitResult.bStartPenetrating && MovementBaseUtility::IsDynamicBase(GetMovementBaseInterfaceData()))
 				{
 					// Follows PhysWalking approach for encroachment on floor tests
 					FHitResult Hit(CurrentFloor.HitResult);
@@ -2383,25 +2396,48 @@ void UCharacterMovementComponent::SimulateMovement(float DeltaSeconds)
 
 UPrimitiveComponent* UCharacterMovementComponent::GetMovementBase() const
 {
-	return CharacterOwner ? CharacterOwner->GetMovementBase() : NULL;
+	const FMovementBaseInterfaceData* MovementBaseInterfaceData = GetMovementBaseInterfaceData();
+	return MovementBaseInterfaceData ? Cast<UPrimitiveComponent>(MovementBaseInterfaceData->GetMovementBaseObject()) : nullptr;
+}
+
+UObject* UCharacterMovementComponent::GetMovementBaseObject() const
+{
+	const FMovementBaseInterfaceData* MovementBaseInterfaceData = GetMovementBaseInterfaceData();
+	return MovementBaseInterfaceData ? MovementBaseInterfaceData->GetMovementBaseObject() : nullptr;
 }
 
-void UCharacterMovementComponent::SetBase( UPrimitiveComponent* NewBase, FName BoneName, bool bNotifyActor )
+const FMovementBaseInterfaceData* UCharacterMovementComponent::GetMovementBaseInterfaceData() const
+{
+	return CharacterOwner ? CharacterOwner->GetMovementBaseInterfaceData() : nullptr;
+}
+
+FMovementBaseInterfaceData* UCharacterMovementComponent::GetMovementBaseInterfaceData_Mutable()
+{
+	return CharacterOwner ? CharacterOwner->GetMovementBaseInterfaceData_Mutable() : nullptr;
+}
+
+void UCharacterMovementComponent::SetBase(FMovementBaseInterfaceData* MovementBaseInterfaceData, const FName BoneName, bool bNotifyActor)
 {
 	// prevent from changing Base while server is NavWalking (no Base in that mode), so both sides are in sync
 	// otherwise it will cause problems with position smoothing
 
 	if (CharacterOwner && !bIsNavWalkingOnServer)
 	{
-		if (bBaseOnAttachmentRoot && NewBase != nullptr && BoneName == NAME_None)
+		if (bBaseOnAttachmentRoot && MovementBaseInterfaceData && MovementBaseInterfaceData->IsValid() && BoneName == NAME_None)
 		{
-			// The attachment root might not be a PrimitiveComponent so check first
-			if (UPrimitiveComponent* NewRootBase = Cast<UPrimitiveComponent>(NewBase->GetAttachmentRoot()))
+			// TODO: Support bBaseOnAttachmentRoot for UObjects
+			if (UPrimitiveComponent* MovementBaseAsPrimitiveComponent = Cast<UPrimitiveComponent>(MovementBaseInterfaceData->PhysicsObjectOwner))
 			{
-				NewBase = NewRootBase;
+				// The attachment root might not be a PrimitiveComponent so check first
+				if (UPrimitiveComponent* NewRootBase = Cast<UPrimitiveComponent>(MovementBaseAsPrimitiveComponent->GetAttachmentRoot()))
+				{
+					FMovementBaseInterfaceData MovementBaseDataAsPrimitive(NewRootBase);
+					MovementBaseInterfaceData->Set(MovementBaseDataAsPrimitive.PhysicsObjectOwner.Get(), MovementBaseDataAsPrimitive.GetBodyInstanceOwner());
+				}
 			}
 		}
-		CharacterOwner->SetBase(NewBase, NewBase ? BoneName : NAME_None, bNotifyActor);
+		
+		CharacterOwner->SetBase(MovementBaseInterfaceData, MovementBaseInterfaceData && MovementBaseInterfaceData->IsValid() ? BoneName : NAME_None, bNotifyActor);
 	}
 }
 
@@ -2410,14 +2446,17 @@ void UCharacterMovementComponent::SetBaseFromFloor(const FFindFloorResult& Floor
 	if (FloorResult.IsWalkableFloor())
 	{
 		const FFindFloorResult InitialFloor = FloorResult;
-		SetBase(FloorResult.HitResult.GetComponent(), FloorResult.HitResult.BoneName);
+
+		FMovementBaseInterfaceData MovementBaseInterfaceData = MovementBaseUtility::GetMovementBaseDataFromHitResult(&FloorResult.HitResult);
+		
+		SetBase(&MovementBaseInterfaceData, FloorResult.HitResult.BoneName);
 
 		// Support valid hit results that don't return a PrimitiveComponent.
 		// SetBase will clear the floor in those cases, but we should keep the valid floor.
 		// SetBase also never keeps the base when bIsNavWalkingOnServer, so don't bother keeping the floor there.
 		if (CharacterMovementCVars::bRestoreFloorWhenMissingBaseComponent)
 		{
-			if (GetMovementBase() == nullptr && !bIsNavWalkingOnServer)
+			if (!MovementBaseUtility::IsMovementBaseDataValid(GetMovementBaseInterfaceData()) && !bIsNavWalkingOnServer)
 			{
 				CurrentFloor = InitialFloor;
 			}
@@ -2425,7 +2464,7 @@ void UCharacterMovementComponent::SetBaseFromFloor(const FFindFloorResult& Floor
 	}
 	else
 	{
-		SetBase(nullptr);
+		SetBase(static_cast<FMovementBaseInterfaceData*>(nullptr));
 	}
 }
 
@@ -2433,11 +2472,11 @@ void UCharacterMovementComponent::MaybeUpdateBasedMovement(float DeltaSeconds)
 {
 	bDeferUpdateBasedMovement = false;
 
-	UPrimitiveComponent* MovementBase = CharacterOwner->GetMovementBase();
-	if (MovementBaseUtility::UseRelativeLocation(MovementBase))
+	FMovementBaseInterfaceData* MovementBaseInterfaceData = GetMovementBaseInterfaceData_Mutable();
+	if (MovementBaseUtility::UseRelativeLocation(MovementBaseInterfaceData))
 	{
 		// Need to see if anything we're on is simulating physics or has a parent that is.		
-		if (!MovementBaseUtility::IsSimulatedBase(MovementBase))
+		if (!MovementBaseUtility::IsSimulatedBase(MovementBaseInterfaceData))
 		{
 			bDeferUpdateBasedMovement = false;
 			UpdateBasedMovement(DeltaSeconds);
@@ -2445,7 +2484,7 @@ void UCharacterMovementComponent::MaybeUpdateBasedMovement(float DeltaSeconds)
 			if (PostPhysicsTickFunction.IsTickFunctionEnabled())
 			{
 				PostPhysicsTickFunction.SetTickFunctionEnable(false);
-				MovementBaseUtility::AddTickDependency(PrimaryComponentTick, MovementBase);
+				MovementBaseUtility::AddTickDependency(PrimaryComponentTick, MovementBaseInterfaceData);
 			}
 		}
 		else
@@ -2456,7 +2495,7 @@ void UCharacterMovementComponent::MaybeUpdateBasedMovement(float DeltaSeconds)
 			if (!PostPhysicsTickFunction.IsTickFunctionEnabled())
 			{
 				PostPhysicsTickFunction.SetTickFunctionEnable(true);
-				MovementBaseUtility::RemoveTickDependency(PrimaryComponentTick, MovementBase);
+				MovementBaseUtility::RemoveTickDependency(PrimaryComponentTick, MovementBaseInterfaceData);
 			}
 
 			if (CharacterMovementCVars::BasedMovementMode == 2)
@@ -2491,15 +2530,17 @@ void UCharacterMovementComponent::UpdateBasedMovement(float DeltaSeconds)
 		return;
 	}
 
-	const UPrimitiveComponent* MovementBase = CharacterOwner->GetMovementBase();
-	if (!MovementBaseUtility::UseRelativeLocation(MovementBase))
+	const FMovementBaseInterfaceData* MovementBaseInterfaceData = GetMovementBaseInterfaceData();
+	
+	if (!MovementBaseUtility::UseRelativeLocation(MovementBaseInterfaceData))
 	{
 		return;
 	}
 
-	if (!IsValid(MovementBase) || !IsValid(MovementBase->GetOwner()))
+	if (!MovementBaseUtility::IsMovementBaseDataValid(MovementBaseInterfaceData))
 	{
-		SetBase(NULL);
+		FMovementBaseInterfaceData EmptyMovementBaseData;
+		SetBase(&EmptyMovementBaseData);
 		return;
 	}
 
@@ -2511,7 +2552,7 @@ void UCharacterMovementComponent::UpdateBasedMovement(float DeltaSeconds)
 
 	FQuat NewBaseQuat;
 	FVector NewBaseLocation;
-	if (!MovementBaseUtility::GetMovementBaseTransform(MovementBase, CharacterOwner->GetBasedMovement().BoneName, NewBaseLocation, NewBaseQuat))
+	if (!MovementBaseUtility::GetMovementBaseTransform(MovementBaseInterfaceData, CharacterOwner->GetBasedMovement().BoneName, NewBaseLocation, NewBaseQuat))
 	{
 		return;
 	}
@@ -2603,7 +2644,7 @@ void UCharacterMovementComponent::UpdateBasedMovement(float DeltaSeconds)
 			AActor* MovementBaseRootActor = nullptr;
 			if (bIgnoreBaseActor)
 			{
-				MovementBaseRootActor = MovementBase->GetAttachmentRootActor();
+				MovementBaseRootActor = Cast<AActor>(MovementBaseInterfaceData->GetBodyInstanceOwner()->GetPhysicsOwnerAttachmentRoot());
 				UpdatedPrimitive->IgnoreActorWhenMoving(MovementBaseRootActor, true);
 				MoveComponentFlags |= MOVECOMP_CheckBlockingRootActorInIgnoreList; // Hit actors during MoveUpdatedComponent will have their root actor compared with the ignored actors array 
 			}
@@ -2633,7 +2674,8 @@ void UCharacterMovementComponent::UpdateBasedMovement(float DeltaSeconds)
 			}
 		}
 
-		if (MovementBase->IsSimulatingPhysics() && CharacterOwner->GetMesh())
+		const FBodyInstance* MovementBaseBodyInstance = MovementBaseInterfaceData && MovementBaseInterfaceData->IsValid() ? MovementBaseInterfaceData->GetBodyInstanceOwner()->GetBodyInstance() : nullptr;
+		if (MovementBaseBodyInstance && MovementBaseBodyInstance->IsInstanceSimulatingPhysics() && CharacterOwner->GetMesh())
 		{
 			CharacterOwner->GetMesh()->ApplyDeltaToAllPhysicsTransforms(DeltaPosition, DeltaQuat);
 		}
@@ -2646,12 +2688,13 @@ void UCharacterMovementComponent::UpdateBasedMovement(float DeltaSeconds)
 		FFindFloorResult OutFloorResult;
 		ComputeFloorDist(PawnLocation, StayBasedInAirHeight, StayBasedInAirHeight, OutFloorResult, CharacterOwner->GetCapsuleComponent()->GetScaledCapsuleRadius(), NULL);
 
-		UPrimitiveComponent* HitComponent = OutFloorResult.HitResult.Component.Get();
-		if (!HitComponent || HitComponent->GetAttachmentRoot() != MovementBase->GetAttachmentRoot())
+		FMovementBaseInterfaceData FoundMovementBaseData = MovementBaseUtility::GetMovementBaseDataFromHitResult(&OutFloorResult.HitResult);
+		if (!FoundMovementBaseData.IsValid() || FoundMovementBaseData.GetBodyInstanceOwner()->GetPhysicsOwnerAttachmentRoot() != MovementBaseInterfaceData->GetBodyInstanceOwner()->GetPhysicsOwnerAttachmentRoot())
 		{
 			// New or no base under the character
 			ApplyImpartedMovementBaseVelocity();
-			SetBase(NULL);
+			FMovementBaseInterfaceData EmptyMovementBaseData;
+			SetBase(&EmptyMovementBaseData);
 			return;
 		}
 	}
@@ -3012,9 +3055,10 @@ void UCharacterMovementComponent::PerformMovement(float DeltaSeconds)
 				// Show resulting delta move.
 				DrawDebugLine(MyWorld, OldLocation, ResultingLocation, FColor::Red, false, 10.f);
 
+				const FMovementBaseInterfaceData* MovementBaseInterfaceData = GetMovementBaseInterfaceData();
 				// Log details.
 				UE_LOG(LogRootMotion, Warning,  TEXT("PerformMovement Resulting DeltaMove Translation: %s, Rotation: %s, MovementBase: %s"),
-					*(ResultingLocation - OldLocation).ToCompactString(), *(ResultingRotation - OldActorRotation).GetNormalized().ToCompactString(), *GetNameSafe(CharacterOwner->GetMovementBase()) );
+					*(ResultingLocation - OldLocation).ToCompactString(), *(ResultingRotation - OldActorRotation).GetNormalized().ToCompactString(), *GetNameSafe(MovementBaseInterfaceData ? MovementBaseInterfaceData->GetMovementBaseObject() : nullptr));
 
 				const FVector RMTranslation = RootMotionParams.GetRootMotionTransform().GetTranslation();
 				const FRotator RMRotation = RootMotionParams.GetRootMotionTransform().GetRotation().Rotator();
@@ -3145,17 +3189,17 @@ void UCharacterMovementComponent::SaveBaseLocation()
 		return;
 	}
 
-	const UPrimitiveComponent* MovementBase = CharacterOwner->GetMovementBase();
-	if (MovementBase)
+	const FMovementBaseInterfaceData* MovementBaseInterfaceData = GetMovementBaseInterfaceData();
+	if (MovementBaseInterfaceData && MovementBaseInterfaceData->IsValid())
 	{
 		// Read transforms into OldBaseLocation, OldBaseQuat. Do this regardless of whether the object is movable, since mobility can change.
-		MovementBaseUtility::GetMovementBaseTransform(MovementBase, CharacterOwner->GetBasedMovement().BoneName, OldBaseLocation, OldBaseQuat);
+		MovementBaseUtility::GetMovementBaseTransform(MovementBaseInterfaceData, CharacterOwner->GetBasedMovement().BoneName, OldBaseLocation, OldBaseQuat);
 
-		if (MovementBaseUtility::UseRelativeLocation(MovementBase))
+		if (MovementBaseUtility::UseRelativeLocation(MovementBaseInterfaceData))
 		{
 			// Relative Location
 			FVector RelativeLocation;
-			MovementBaseUtility::TransformLocationToLocal(MovementBase, CharacterOwner->GetBasedMovement().BoneName, UpdatedComponent->GetComponentLocation(), RelativeLocation);
+			MovementBaseUtility::TransformLocationToLocal(MovementBaseInterfaceData, CharacterOwner->GetBasedMovement().BoneName, UpdatedComponent->GetComponentLocation(), RelativeLocation);
 
 			// Rotation
 			if (bIgnoreBaseRotation)
@@ -3847,7 +3891,7 @@ void UCharacterMovementComponent::CalcVelocity(float DeltaTime, float Friction,
 	const bool bVelocityOverMax = IsExceedingMaxSpeed(MaxSpeed);
 	
 	// Only apply braking if there is no acceleration, or we are over our max speed and need to slow down to it.
-	if ((bZeroAcceleration && bZeroRequestedAcceleration) || bVelocityOverMax)
+	if ((bZeroAcceleration && bZeroRequestedAcceleration) || (bVelocityOverMax && RequestedSpeed < MaxSpeed))
 	{
 		const FVector OldVelocity = Velocity;
 
@@ -5056,10 +5100,11 @@ void UCharacterMovementComponent::PhysFalling(float deltaTime, int32 Iterations)
 				FVector Delta = ComputeSlideVector(Adjusted, 1.f - Hit.Time, OldHitNormal, Hit);
 
 				// Compute velocity after deflection (only gravity component for RootMotion)
-				const UPrimitiveComponent* HitComponent = Hit.GetComponent();
-				if (CharacterMovementCVars::UseTargetVelocityOnImpact && !Velocity.IsNearlyZero() && MovementBaseUtility::IsSimulatedBase(HitComponent))
+				FMovementBaseInterfaceData HitMovementBaseInterfaceData(MovementBaseUtility::GetMovementBaseDataFromHitResult(&Hit));
+				
+				if (CharacterMovementCVars::UseTargetVelocityOnImpact && !Velocity.IsNearlyZero() && MovementBaseUtility::IsSimulatedBase(&HitMovementBaseInterfaceData))
 				{
-					const FVector ContactVelocity = MovementBaseUtility::GetMovementBaseVelocity(HitComponent, NAME_None) + MovementBaseUtility::GetMovementBaseTangentialVelocity(HitComponent, NAME_None, Hit.ImpactPoint);
+					const FVector ContactVelocity = MovementBaseUtility::GetMovementBaseVelocity(&HitMovementBaseInterfaceData, NAME_None) + MovementBaseUtility::GetMovementBaseTangentialVelocity(&HitMovementBaseInterfaceData, NAME_None, Hit.ImpactPoint);
 					const FVector NewVelocity = Velocity - Hit.ImpactNormal * FVector::DotProduct(Velocity - ContactVelocity, Hit.ImpactNormal);
 					Velocity = HasAnimRootMotion() || CurrentRootMotion.HasOverrideVelocityWithIgnoreZAccumulate() ? ProjectToGravityFloor(Velocity) + GetGravitySpaceComponentZ(NewVelocity) : NewVelocity;
 				}
@@ -5347,6 +5392,12 @@ bool UCharacterMovementComponent::CheckFall(const FFindFloorResult& OldFloor, co
 	return false;
 }
 
+void UCharacterMovementComponent::RevertMove(const FVector& OldLocation, UPrimitiveComponent* OldBase, const FVector& InOldBaseLocation, const FFindFloorResult& OldFloor, bool bFailMove)
+{
+	FMovementBaseInterfaceData OldMovementBaseInterfaceData(OldBase);
+	RevertMove(OldLocation, &OldMovementBaseInterfaceData, InOldBaseLocation, OldFloor, bFailMove);
+}
+
 void UCharacterMovementComponent::StartFalling(int32 Iterations, float remainingTime, float timeTick, const FVector& Delta, const FVector& subLoc)
 {
 	// start falling 
@@ -5375,7 +5426,7 @@ void UCharacterMovementComponent::StartFalling(int32 Iterations, float remaining
 }
 
 
-void UCharacterMovementComponent::RevertMove(const FVector& OldLocation, UPrimitiveComponent* OldBase, const FVector& PreviousBaseLocation, const FFindFloorResult& OldFloor, bool bFailMove)
+void UCharacterMovementComponent::RevertMove(const FVector& OldLocation, FMovementBaseInterfaceData* OldMovementBaseInterfaceData, const FVector& PreviousBaseLocation, const FFindFloorResult& OldFloor, bool bFailMove)
 {
 	//UE_LOG(LogCharacterMovement, Log, TEXT("RevertMove from %f %f %f to %f %f %f"), CharacterOwner->Location.X, CharacterOwner->Location.Y, CharacterOwner->Location.Z, OldLocation.X, OldLocation.Y, OldLocation.Z);
 	UpdatedComponent->SetWorldLocation(OldLocation, false, nullptr, GetTeleportType());
@@ -5383,19 +5434,20 @@ void UCharacterMovementComponent::RevertMove(const FVector& OldLocation, UPrimit
 	//UE_LOG(LogCharacterMovement, Log, TEXT("Now at %f %f %f"), CharacterOwner->Location.X, CharacterOwner->Location.Y, CharacterOwner->Location.Z);
 	bJustTeleported = false;
 	// if our previous base couldn't have moved or changed in any physics-affecting way, restore it
-	if (IsValid(OldBase) && 
-		(!MovementBaseUtility::IsDynamicBase(OldBase) ||
-		 (OldBase->Mobility == EComponentMobility::Static) ||
-		 (OldBase->GetComponentLocation() == PreviousBaseLocation)
+	if (MovementBaseUtility::IsMovementBaseDataValid(OldMovementBaseInterfaceData) && 
+		(!MovementBaseUtility::IsDynamicBase(OldMovementBaseInterfaceData) ||
+		 (!OldMovementBaseInterfaceData->GetBodyInstanceOwner()->IsPhysicsOwnerMovable()) ||
+		 (OldMovementBaseInterfaceData->GetBodyInstanceOwner()->GetPhysicsOwnerTransform().GetLocation() == PreviousBaseLocation)
 		)
 	   )
 	{
-		SetBase(OldBase, OldFloor.HitResult.BoneName);
+		SetBase(OldMovementBaseInterfaceData, OldFloor.HitResult.BoneName);
 		CurrentFloor = OldFloor;
 	}
 	else
 	{
-		SetBase(NULL);
+		FMovementBaseInterfaceData EmptyMovementBaseData;
+		SetBase(&EmptyMovementBaseData);
 	}
 
 	if ( bFailMove )
@@ -5514,7 +5566,11 @@ void UCharacterMovementComponent::MoveAlongFloor(const FVector& InVelocity, floa
 
 		if (Hit.IsValidBlockingHit())
 		{
-			if (CanStepUp(Hit) || (CharacterOwner->GetMovementBase() != nullptr && Hit.HitObjectHandle == CharacterOwner->GetMovementBase()->GetOwner()))
+			const FMovementBaseInterfaceData* MovementBaseData = GetMovementBaseInterfaceData();
+			const UObject* MovementBaseObject = MovementBaseData ? GetMovementBaseInterfaceData()->PhysicsObjectOwner.Get() : nullptr;
+			if (CanStepUp(Hit) || 
+				(Hit.PhysicsObjectOwner.Get() == MovementBaseObject) || 
+				(Hit.Component.Get() == MovementBaseObject))
 			{
 				// hit a barrier, try to step up
 				const FVector PreStepUpLocation = UpdatedComponent->GetComponentLocation();
@@ -5615,8 +5671,8 @@ void UCharacterMovementComponent::PhysWalking(float deltaTime, int32 Iterations)
 #endif
 		
 		// Save current values
-		UPrimitiveComponent * const OldBase = GetMovementBase();
-		const FVector PreviousBaseLocation = (OldBase != NULL) ? OldBase->GetComponentLocation() : FVector::ZeroVector;
+		FMovementBaseInterfaceData* OldMovementBaseInterfaceData = GetMovementBaseInterfaceData_Mutable();
+		const FVector PreviousBaseLocation = OldMovementBaseInterfaceData && OldMovementBaseInterfaceData->IsValid() ? OldMovementBaseInterfaceData->GetBodyInstanceOwner()->GetPhysicsOwnerTransform().GetLocation() : FVector::ZeroVector;
 		const FVector OldLocation = UpdatedComponent->GetComponentLocation();
 		const FFindFloorResult OldFloor = CurrentFloor;
 
@@ -5701,7 +5757,7 @@ void UCharacterMovementComponent::PhysWalking(float deltaTime, int32 Iterations)
 			if ( !NewDelta.IsZero() )
 			{
 				// first revert this move
-				RevertMove(OldLocation, OldBase, PreviousBaseLocation, OldFloor, false);
+				RevertMove(OldLocation, OldMovementBaseInterfaceData, PreviousBaseLocation, OldFloor, false);
 
 				// avoid repeated ledge moves if the first one fails
 				bTriedLedgeMove = true;
@@ -5716,7 +5772,7 @@ void UCharacterMovementComponent::PhysWalking(float deltaTime, int32 Iterations)
 			{
 				// see if it is OK to jump
 				// @todo collision : only thing that can be problem is that oldbase has world collision on
-				bool bMustJump = bZeroDelta || (OldBase == NULL || (!OldBase->IsQueryCollisionEnabled() && MovementBaseUtility::IsDynamicBase(OldBase)));
+				bool bMustJump = bZeroDelta || (!MovementBaseUtility::IsMovementBaseDataValid(OldMovementBaseInterfaceData) || (!CollisionEnabledHasQuery(OldMovementBaseInterfaceData->GetBodyInstanceOwner()->GetCollisionEnabled()) && MovementBaseUtility::IsDynamicBase(OldMovementBaseInterfaceData)));
 				if ( (bMustJump || !bCheckedFall) && CheckFall(OldFloor, CurrentFloor.HitResult, Delta, OldLocation, remainingTime, timeTick, Iterations, bMustJump) )
 				{
 					return;
@@ -5724,7 +5780,7 @@ void UCharacterMovementComponent::PhysWalking(float deltaTime, int32 Iterations)
 				bCheckedFall = true;
 
 				// revert this move
-				RevertMove(OldLocation, OldBase, PreviousBaseLocation, OldFloor, true);
+				RevertMove(OldLocation, OldMovementBaseInterfaceData, PreviousBaseLocation, OldFloor, true);
 				remainingTime = 0.f;
 				break;
 			}
@@ -5769,7 +5825,7 @@ void UCharacterMovementComponent::PhysWalking(float deltaTime, int32 Iterations)
 			// See if we need to start falling.
 			if (!CurrentFloor.IsWalkableFloor() && !CurrentFloor.HitResult.bStartPenetrating)
 			{
-				const bool bMustJump = bJustTeleported || bZeroDelta || (OldBase == NULL || (!OldBase->IsQueryCollisionEnabled() && MovementBaseUtility::IsDynamicBase(OldBase)));
+				const bool bMustJump = bJustTeleported || bZeroDelta || (!MovementBaseUtility::IsMovementBaseDataValid(OldMovementBaseInterfaceData) || (!CollisionEnabledHasQuery(OldMovementBaseInterfaceData->GetBodyInstanceOwner()->GetCollisionEnabled()) && MovementBaseUtility::IsDynamicBase(OldMovementBaseInterfaceData)));
 				if ((bMustJump || !bCheckedFall) && CheckFall(OldFloor, CurrentFloor.HitResult, Delta, OldLocation, remainingTime, timeTick, Iterations, bMustJump) )
 				{
 					return;
@@ -6462,15 +6518,14 @@ void UCharacterMovementComponent::OnTeleported()
 	UpdateFloorFromAdjustment();
 
 	// Validate it. We don't want to pop down to walking mode from very high off the ground, but we'd like to keep walking if possible.
-	UPrimitiveComponent* OldBase = CharacterOwner->GetMovementBase();
-	UPrimitiveComponent* NewBase = NULL;
+	const FMovementBaseInterfaceData* OldMovementBaseInterfaceData = GetMovementBaseInterfaceData();
+	const bool bOldMovementBaseValid = OldMovementBaseInterfaceData && OldMovementBaseInterfaceData->IsValid();
 	bool bUseNewBase = false;
-	
-	if ((OldBase || CharacterMovementCVars::bRestoreFloorWhenMissingBaseComponent) && CurrentFloor.IsWalkableFloor() && CurrentFloor.FloorDist <= MAX_FLOOR_DIST && GetGravitySpaceZ(Velocity) <= 0.f)
+	if ((bOldMovementBaseValid || CharacterMovementCVars::bRestoreFloorWhenMissingBaseComponent) && CurrentFloor.IsWalkableFloor() && CurrentFloor.FloorDist <= MAX_FLOOR_DIST && GetGravitySpaceZ(Velocity) <= 0.f)
 	{
 		// Close enough to land or just keep walking.
-		NewBase = CurrentFloor.HitResult.Component.Get();
-		bUseNewBase = (NewBase != nullptr) || CharacterMovementCVars::bRestoreFloorWhenMissingBaseComponent;
+		const FMovementBaseInterfaceData NewMovementBaseInterfaceData = MovementBaseUtility::GetMovementBaseDataFromHitResult(&CurrentFloor.HitResult);
+		bUseNewBase = NewMovementBaseInterfaceData.IsValid() || CharacterMovementCVars::bRestoreFloorWhenMissingBaseComponent;
 	}
 	else
 	{
@@ -6487,7 +6542,7 @@ void UCharacterMovementComponent::OnTeleported()
 			SetMovementMode(DefaultWaterMovementMode);
 		}
 	}
-	else if (!CurrentFloor.IsWalkableFloor() || (OldBase && !bUseNewBase))
+	else if (!CurrentFloor.IsWalkableFloor() || (bOldMovementBaseValid && !bUseNewBase))
 	{
 		if (!bWasFalling && MovementMode != MOVE_Flying && MovementMode != MOVE_Custom)
 		{
@@ -6921,6 +6976,15 @@ bool UCharacterMovementComponent::IsWalkable(const FHitResult& Hit) const
 		const FWalkableSlopeOverride& SlopeOverride = HitComponent->GetWalkableSlopeOverride();
 		TestWalkableZ = SlopeOverride.ModifyWalkableFloorZ(TestWalkableZ);
 	}
+	else
+	{
+		FMovementBaseInterfaceData MovementBaseData = MovementBaseUtility::GetMovementBaseDataFromHitResult(&Hit);
+		if (MovementBaseData.IsValid())
+		{
+			const FWalkableSlopeOverride& SlopeOverride = MovementBaseData.GetBodyInstanceOwner()->GetWalkableSlopeOverride();
+			TestWalkableZ = SlopeOverride.ModifyWalkableFloorZ(TestWalkableZ);
+		}
+	}
 
 	// Can't walk on this surface if it is too steep.
 	if (ImpactNormalZ < TestWalkableZ)
@@ -7148,20 +7212,22 @@ void UCharacterMovementComponent::FindFloor(const FVector& CapsuleLocation, FFin
 		else
 		{
 			// Force floor check if base has collision disabled or if it does not block us.
-			UPrimitiveComponent* MovementBase = CharacterOwner->GetMovementBase();
-			const AActor* BaseActor = MovementBase ? MovementBase->GetOwner() : NULL;
+			const FMovementBaseInterfaceData* MovementBaseInterfaceData = GetMovementBaseInterfaceData();
+			
+			const bool bHasValidMovementBase = MovementBaseInterfaceData && MovementBaseInterfaceData->IsValid();
+			const AActor* BaseActor = bHasValidMovementBase ? Cast<const AActor>(MovementBaseInterfaceData->GetMovementBaseObject()) : nullptr;
 			const ECollisionChannel CollisionChannel = UpdatedComponent->GetCollisionObjectType();
 
-			if (MovementBase != NULL)
+			if (bHasValidMovementBase)
 			{
-				MutableThis->bForceNextFloorCheck = !MovementBase->IsQueryCollisionEnabled()
-				|| MovementBase->GetCollisionResponseToChannel(CollisionChannel) != ECR_Block
-				|| MovementBaseUtility::IsDynamicBase(MovementBase);
+				MutableThis->bForceNextFloorCheck = !CollisionEnabledHasQuery(MovementBaseInterfaceData->GetBodyInstanceOwner()->GetCollisionEnabled())
+				|| MovementBaseInterfaceData->GetBodyInstanceOwner()->GetCollisionResponseToChannel(CollisionChannel) != ECR_Block
+				|| MovementBaseUtility::IsDynamicBase(MovementBaseInterfaceData);
 			}
 
 			const bool IsActorBasePendingKill = BaseActor && !IsValid(BaseActor);
 
-			if ( !bForceNextFloorCheck && !IsActorBasePendingKill && MovementBase )
+			if ( !bForceNextFloorCheck && !IsActorBasePendingKill && bHasValidMovementBase)
 			{
 				//UE_LOG(LogCharacterMovement, Log, TEXT("%s SKIP check for floor"), *CharacterOwner->GetName());
 				OutFloorResult = CurrentFloor;
@@ -7681,6 +7747,12 @@ bool UCharacterMovementComponent::StepUp(const FVector& GravDir, const FVector&
 	return true;
 }
 
+void UCharacterMovementComponent::SetBase(UPrimitiveComponent* NewBase, const FName BoneName, bool bNotifyActor)
+{
+	FMovementBaseInterfaceData MovementBaseInterfaceData(NewBase);
+	SetBase(&MovementBaseInterfaceData, BoneName, bNotifyActor);
+}
+
 void UCharacterMovementComponent::HandleImpact(const FHitResult& Impact, float TimeSlice, const FVector& MoveDelta)
 {
 	SCOPE_CYCLE_COUNTER(STAT_CharHandleImpact);
@@ -7889,11 +7961,12 @@ void UCharacterMovementComponent::DisplayDebug(UCanvas* Canvas, const FDebugDisp
 
 	APhysicsVolume * PhysicsVolume = GetPhysicsVolume();
 
-	const UPrimitiveComponent* BaseComponent = CharacterOwner->GetMovementBase();
-	const AActor* BaseActor = BaseComponent ? BaseComponent->GetOwner() : NULL;
+	const FMovementBaseInterfaceData* MovementBaseInterfaceData = GetMovementBaseInterfaceData();
+	const UObject* BodyInstanceSourceObject = MovementBaseInterfaceData && MovementBaseInterfaceData->GetBodyInstanceOwner() ? MovementBaseInterfaceData->GetBodyInstanceOwner()->GetSourceObject() : nullptr;
+	const UObject* BodyInstanceSourceObjectOwner = MovementBaseInterfaceData && MovementBaseInterfaceData->GetBodyInstanceOwner() ? MovementBaseInterfaceData->GetBodyInstanceOwner()->GetSourceObjectOwner() : nullptr;
 
 	T = FString::Printf(TEXT("%s In physicsvolume %s on base %s component %s gravity %f"), *GetMovementName(), (PhysicsVolume ? *PhysicsVolume->GetName() : TEXT("None")),
-		(BaseActor ? *BaseActor->GetName() : TEXT("None")), (BaseComponent ? *BaseComponent->GetName() : TEXT("None")), GetGravityZ());
+		(BodyInstanceSourceObjectOwner ? *BodyInstanceSourceObjectOwner->GetName() : TEXT("None")), (BodyInstanceSourceObject ? *BodyInstanceSourceObject->GetName() : TEXT("None")), GetGravityZ());
 	DisplayDebugManager.DrawString(T);
 }
 
@@ -8950,7 +9023,8 @@ void UCharacterMovementComponent::ReplicateMoveToServer(float DeltaTime, const F
 
 		UE_CLOG(CharacterOwner && UpdatedComponent, LogNetPlayerMovement, VeryVerbose, TEXT("ClientMove Time %f Acceleration %s Velocity %s Position %s Rotation %s DeltaTime %f Mode %s MovementBase %s.%s (Dynamic:%d) DualMove? %d"),
 			NewMove->TimeStamp, *NewMove->Acceleration.ToString(), *Velocity.ToString(), *UpdatedComponent->GetComponentLocation().ToString(), *UpdatedComponent->GetComponentRotation().ToCompactString(), NewMove->DeltaTime, *GetMovementName(),
-			*GetNameSafe(NewMove->EndBase.Get()), *NewMove->EndBoneName.ToString(), MovementBaseUtility::IsDynamicBase(NewMove->EndBase.Get()) ? 1 : 0, ClientData->PendingMove.IsValid() ? 1 : 0);
+			*GetNameSafe(NewMove->EndMovementBaseInterfaceData.GetMovementBaseObject()), *NewMove->EndBoneName.ToString(),
+			MovementBaseUtility::IsDynamicBase(&NewMove->EndMovementBaseInterfaceData) ? 1 : 0, ClientData->PendingMove.IsValid() ? 1 : 0);
 
 		
 		bool bSendServerMove = true;
@@ -9101,9 +9175,10 @@ void UCharacterMovementComponent::CallServerMove
 	NewMove->GetPackedAngles(ClientYawPitchINT, ClientRollBYTE);
 
 	// Determine if we send absolute or relative location
-	UPrimitiveComponent* ClientMovementBase = NewMove->EndBase.Get();
+	UPrimitiveComponent* ClientMovementBase = Cast<UPrimitiveComponent>(NewMove->EndMovementBaseInterfaceData.PhysicsObjectOwner.Get());
+	
 	const FName ClientBaseBone = NewMove->EndBoneName;
-	const FVector SendLocation = MovementBaseUtility::UseRelativeLocation(ClientMovementBase) ? NewMove->SavedRelativeLocation : FRepMovement::RebaseOntoZeroOrigin(NewMove->SavedLocation, this);
+	const FVector SendLocation = MovementBaseUtility::UseRelativeLocation(&NewMove->EndMovementBaseInterfaceData) ? NewMove->SavedRelativeLocation : FRepMovement::RebaseOntoZeroOrigin(NewMove->SavedLocation, this);
 
 	// send old move if it exists
 	if (OldMove)
@@ -9701,9 +9776,7 @@ void FCharacterNetworkMoveData::ClientFillNetworkMoveData(const FSavedMove_Chara
 	MovementMode = ClientMove.EndPackedMovementMode;
 
 	// Determine if we send absolute or relative location
-	UPrimitiveComponent* ClientMovementBase = ClientMove.EndBase.Get();
-
-	const bool bSendBaseRelativeLocation     = MovementBaseUtility::UseRelativeLocation(ClientMovementBase);
+	const bool bSendBaseRelativeLocation     = MovementBaseUtility::UseRelativeLocation(&ClientMove.EndMovementBaseInterfaceData);
 	const bool bSendBaseRelativeAcceleration = CharacterMovementCVars::NetUseBaseRelativeAcceleration && bSendBaseRelativeLocation;
 
 	const FVector SendLocation     = bSendBaseRelativeLocation ? ClientMove.SavedRelativeLocation : FRepMovement::RebaseOntoZeroOrigin(ClientMove.SavedLocation, ClientMove.CharacterOwner->GetCharacterMovement());
@@ -9714,12 +9787,18 @@ void FCharacterNetworkMoveData::ClientFillNetworkMoveData(const FSavedMove_Chara
 
 	if (bSendBaseRelativeLocation || bSendBaseRelativeAcceleration)
 	{
-		MovementBase = ClientMovementBase;
+		MovementBasePhysicsObjectOwner = ClientMove.EndMovementBaseInterfaceData.PhysicsObjectOwner.Get();
+		PRAGMA_DISABLE_DEPRECATION_WARNINGS
+		MovementBase = ClientMove.EndBase.Get();
+		PRAGMA_ENABLE_DEPRECATION_WARNINGS
 		MovementBaseBoneName = ClientMove.EndBoneName;
 	}
 	else
 	{
+		MovementBasePhysicsObjectOwner = nullptr;
+		PRAGMA_DISABLE_DEPRECATION_WARNINGS
 		MovementBase = nullptr;
+		PRAGMA_ENABLE_DEPRECATION_WARNINGS
 		MovementBaseBoneName = NAME_None;
 	}
 }
@@ -9743,11 +9822,14 @@ bool FCharacterNetworkMoveData::Serialize(UCharacterMovementComponent& Character
 	ControlRotation.NetSerialize(Ar, PackageMap, bLocalSuccess);
 
 	SerializeOptionalValue<uint8>(bIsSaving, Ar, CompressedMoveFlags, 0);
-
-		SerializeOptionalValue<UPrimitiveComponent*>(bIsSaving, Ar, MovementBase, nullptr);
+		SerializeOptionalValue<TObjectPtr<UObject>>(bIsSaving, Ar, MovementBasePhysicsObjectOwner, nullptr);
 		SerializeOptionalValue<FName>(bIsSaving, Ar, MovementBaseBoneName, NAME_None);
 		SerializeOptionalValue<uint8>(bIsSaving, Ar, MovementMode, MOVE_Walking);
 
+	PRAGMA_DISABLE_DEPRECATION_WARNINGS
+	MovementBase = Cast<UPrimitiveComponent>(MovementBasePhysicsObjectOwner.Get());
+	PRAGMA_ENABLE_DEPRECATION_WARNINGS
+	
 	return !Ar.IsError();
 }
 
@@ -9868,10 +9950,12 @@ void UCharacterMovementComponent::ServerMove_PerformMovement(const FCharacterNet
 	
 	FVector ClientAccel = MoveData.Acceleration;
 
+	FMovementBaseInterfaceData MovementBaseInterfaceData = MovementBaseUtility::GetMovementBaseDataFromPhysicsOwner(MoveData.MovementBasePhysicsObjectOwner);
+	
 	// Convert the move's acceleration to worldspace if necessary
-	if (CharacterMovementCVars::NetUseBaseRelativeAcceleration && MovementBaseUtility::IsDynamicBase(MoveData.MovementBase))
+	if (CharacterMovementCVars::NetUseBaseRelativeAcceleration && MovementBaseUtility::IsDynamicBase(&MovementBaseInterfaceData))
 	{
-		MovementBaseUtility::TransformDirectionToWorld(MoveData.MovementBase, MoveData.MovementBaseBoneName, MoveData.Acceleration, ClientAccel);
+		MovementBaseUtility::TransformDirectionToWorld(&MovementBaseInterfaceData, MoveData.MovementBaseBoneName, MoveData.Acceleration, ClientAccel);
 	}
 
 	const uint8 ClientMoveFlags = MoveData.CompressedMoveFlags;
@@ -9943,13 +10027,13 @@ void UCharacterMovementComponent::ServerMove_PerformMovement(const FCharacterNet
 
 		UE_CLOG(CharacterOwner && UpdatedComponent, LogNetPlayerMovement, VeryVerbose, TEXT("ServerMove Time %f Acceleration %s Velocity %s Position %s Rotation %s GravityDirection %s DeltaTime %f Mode %s MovementBase %s.%s (Dynamic:%d)"),
 			ClientTimeStamp, *ClientAccel.ToString(), *Velocity.ToString(), *UpdatedComponent->GetComponentLocation().ToString(), *UpdatedComponent->GetComponentRotation().ToCompactString(), *GravityDirection.ToCompactString(), DeltaTime, *GetMovementName(),
-			*GetNameSafe(GetMovementBase()), *CharacterOwner->GetBasedMovement().BoneName.ToString(), MovementBaseUtility::IsDynamicBase(GetMovementBase()) ? 1 : 0);
+			*GetNameSafe(MovementBaseInterfaceData.GetMovementBaseObject()), *CharacterOwner->GetBasedMovement().BoneName.ToString(), MovementBaseUtility::IsDynamicBase(&MovementBaseInterfaceData) ? 1 : 0);
 	}
 
 	// Validate move only after old and first dual portion, after all moves are completed.
 	if (MoveData.NetworkMoveType == FCharacterNetworkMoveData::ENetworkMoveType::NewMove)
 	{
-		ServerMoveHandleClientError(ClientTimeStamp, DeltaTime, ClientAccel, MoveData.Location, MoveData.MovementBase, MoveData.MovementBaseBoneName, MoveData.MovementMode);
+		ServerMoveHandleClientError(ClientTimeStamp, DeltaTime, ClientAccel, MoveData.Location, &MovementBaseInterfaceData, MoveData.MovementBaseBoneName, MoveData.MovementMode);
 	}
 }
 
@@ -9957,7 +10041,9 @@ void UCharacterMovementComponent::ServerMove_PerformMovement(const FCharacterNet
 ///// DEPRECATED /////
 void UCharacterMovementComponent::ServerMove(float TimeStamp, FVector_NetQuantize10 InAccel, FVector_NetQuantize100 ClientLoc, uint8 CompressedMoveFlags, uint8 ClientRoll, uint32 View, UPrimitiveComponent* ClientMovementBase, FName ClientBaseBoneName, uint8 ClientMovementMode)
 {
-	if (MovementBaseUtility::IsDynamicBase(ClientMovementBase))
+	// Making a new FMovementBaseInterfaceData struct so we can use non-deprecated functions - we could also just disable deprecation warnings here too, since the function is deprecated itself
+	const FMovementBaseInterfaceData MovementBaseInterfaceData(ClientMovementBase);
+	if (MovementBaseUtility::IsDynamicBase(&MovementBaseInterfaceData))
 	{
 		//UE_LOG(LogCharacterMovement, Log, TEXT("ServerMove: base %s"), *ClientMovementBase->GetName());
 		CharacterOwner->ServerMove(TimeStamp, InAccel, ClientLoc, CompressedMoveFlags, ClientRoll, View, ClientMovementBase, ClientBaseBoneName, ClientMovementMode);
@@ -10057,15 +10143,21 @@ void UCharacterMovementComponent::ServerMove_Implementation(
 		MoveAutonomous(TimeStamp, DeltaTime, MoveFlags, Accel);
 	}
 
+	FMovementBaseInterfaceData* MovementBaseInterfaceData = GetMovementBaseInterfaceData_Mutable();
 	UE_CLOG(CharacterOwner && UpdatedComponent, LogNetPlayerMovement, VeryVerbose, TEXT("ServerMove Time %f Acceleration %s Velocity %s Position %s Rotation %s DeltaTime %f Mode %s MovementBase %s.%s (Dynamic:%d)"),
 			TimeStamp, *Accel.ToString(), *Velocity.ToString(), *UpdatedComponent->GetComponentLocation().ToString(), *UpdatedComponent->GetComponentRotation().ToCompactString(), DeltaTime, *GetMovementName(),
-			*GetNameSafe(GetMovementBase()), *CharacterOwner->GetBasedMovement().BoneName.ToString(), MovementBaseUtility::IsDynamicBase(GetMovementBase()) ? 1 : 0);
+			*GetNameSafe(MovementBaseInterfaceData ? MovementBaseInterfaceData->GetMovementBaseObject() : nullptr), *CharacterOwner->GetBasedMovement().BoneName.ToString(), MovementBaseUtility::IsDynamicBase(MovementBaseInterfaceData) ? 1 : 0);
 
-	ServerMoveHandleClientError(TimeStamp, DeltaTime, Accel, ClientLoc, ClientMovementBase, ClientBaseBoneName, ClientMovementMode);
+	ServerMoveHandleClientError(TimeStamp, DeltaTime, Accel, ClientLoc, MovementBaseInterfaceData, ClientBaseBoneName, ClientMovementMode);
 }
 
+void UCharacterMovementComponent::ServerMoveHandleClientError(float ClientTimeStamp, float DeltaTime, const FVector& Accel, const FVector& RelativeClientLocation, UPrimitiveComponent* ClientMovementBase, FName ClientBaseBoneName, uint8 ClientMovementMode)
+{
+	FMovementBaseInterfaceData ClientMovementBaseInterfaceData(ClientMovementBase);
+	ServerMoveHandleClientError(ClientTimeStamp, DeltaTime, Accel, RelativeClientLocation, &ClientMovementBaseInterfaceData, ClientBaseBoneName, ClientMovementMode);
+}
 
-void UCharacterMovementComponent::ServerMoveHandleClientError(float ClientTimeStamp, float DeltaTime, const FVector& Accel, const FVector& RelativeClientLoc, UPrimitiveComponent* ClientMovementBase, FName ClientBaseBoneName, uint8 ClientMovementMode)
+void UCharacterMovementComponent::ServerMoveHandleClientError(float ClientTimeStamp, float DeltaTime, const FVector& Accel, const FVector& RelativeClientLoc, FMovementBaseInterfaceData* ClientMovementBaseInterfaceData, FName ClientBaseBoneName, uint8 ClientMovementMode)
 {
 	if (!ShouldUsePackedMovementRPCs())
 	{
@@ -10092,9 +10184,9 @@ void UCharacterMovementComponent::ServerMoveHandleClientError(float ClientTimeSt
 
 	// Offset may be relative to base component
 	FVector ClientLoc = RelativeClientLoc;
-	if (MovementBaseUtility::UseRelativeLocation(ClientMovementBase))
+	if (MovementBaseUtility::UseRelativeLocation(ClientMovementBaseInterfaceData))
 	{
-		MovementBaseUtility::TransformLocationToWorld(ClientMovementBase, ClientBaseBoneName, RelativeClientLoc, ClientLoc);
+		MovementBaseUtility::TransformLocationToWorld(ClientMovementBaseInterfaceData, ClientBaseBoneName, RelativeClientLoc, ClientLoc);
 	}
 	else
 	{
@@ -10105,7 +10197,7 @@ void UCharacterMovementComponent::ServerMoveHandleClientError(float ClientTimeSt
 
 	// Client may send a null movement base when walking on bases with no relative location (to save bandwidth).
 	// In this case don't check movement base in error conditions, use the server one (which avoids an error based on differing bases). Position will still be validated.
-	if (ClientMovementBase == nullptr)
+	if (!MovementBaseUtility::IsMovementBaseDataValid(ClientMovementBaseInterfaceData))
 	{
 		TEnumAsByte<EMovementMode> NetMovementMode(MOVE_None);
 		TEnumAsByte<EMovementMode> NetGroundMode(MOVE_None);
@@ -10113,15 +10205,15 @@ void UCharacterMovementComponent::ServerMoveHandleClientError(float ClientTimeSt
 		UnpackNetworkMovementMode(ClientMovementMode, NetMovementMode, NetCustomMode, NetGroundMode);
 		if (NetMovementMode == MOVE_Walking)
 		{
-			ClientMovementBase = CharacterOwner->GetBasedMovement().MovementBase;
+			ClientMovementBaseInterfaceData = GetMovementBaseInterfaceData_Mutable();
 			ClientBaseBoneName = CharacterOwner->GetBasedMovement().BoneName;
 		}
 	}
 
 	// If base location is out of sync on server and client, changing base can result in a jarring correction.
 	// So in the case that the base has just changed on server or client, server trusts the client (within a threshold)
-	UPrimitiveComponent* MovementBase = CharacterOwner->GetMovementBase();
 	FName MovementBaseBoneName = CharacterOwner->GetBasedMovement().BoneName;
+	const FMovementBaseInterfaceData* MovementBaseInterfaceData = GetMovementBaseInterfaceData();
 	const bool bServerIsFalling = IsFalling();
 	const bool bClientIsFalling = ClientMovementMode == MOVE_Falling;
 	const bool bServerJustLanded = bLastServerIsFalling && !bServerIsFalling;
@@ -10187,16 +10279,15 @@ void UCharacterMovementComponent::ServerMoveHandleClientError(float ClientTimeSt
 		// Check for lift-off, going from walking to falling.
 		// Note that if bStayBasedInAir is enabled we can't rely on the walking to falling transition, instead run logic on the first tick after clearing the MovementBase
 		const bool bCanLiftOffFromBase = bStayBasedInAir
-			? !MovementBase && LastServerMovementBase.Get() // If we keep the base while in air, consider lift-off if base gets set to null and we had a base last tick
+			? !MovementBaseUtility::IsMovementBaseDataValid(MovementBaseInterfaceData) && LastMovementBaseInterfaceData.IsValid() // If we keep the base while in air, consider lift-off if base gets set to null and we had a base last tick
 			: bLastServerIsWalking; // If walking last tick, we were can consider lift-off logic
 
 		if (bServerIsFalling && bCanLiftOffFromBase && !bTeleportedSinceLastUpdate)
 		{
 			float ClientForwardFactor = 1.f;
-			UPrimitiveComponent* LastServerMovementBasePtr = LastServerMovementBase.Get();
-			if (IsValid(LastServerMovementBasePtr) && MovementBaseUtility::IsDynamicBase(LastServerMovementBasePtr) && MaxWalkSpeed > UE_KINDA_SMALL_NUMBER)
+			if (LastMovementBaseInterfaceData.IsValid() && MovementBaseUtility::IsDynamicBase(&LastMovementBaseInterfaceData) && MaxWalkSpeed > UE_KINDA_SMALL_NUMBER)
 			{
-				const FVector LastBaseVelocity = MovementBaseUtility::GetMovementBaseVelocity(LastServerMovementBasePtr, LastServerMovementBaseBoneName);
+				const FVector LastBaseVelocity = MovementBaseUtility::GetMovementBaseVelocity(&LastMovementBaseInterfaceData, LastServerMovementBaseBoneName);
 				RelativeVelocity = Velocity - LastBaseVelocity;
 				const FVector BaseDirection = ProjectToGravityFloor(LastBaseVelocity).GetSafeNormal();
 				const FVector RelativeDirection = RelativeVelocity * (1.f / MaxWalkSpeed);
@@ -10204,10 +10295,10 @@ void UCharacterMovementComponent::ServerMoveHandleClientError(float ClientTimeSt
 				ClientForwardFactor = FMath::Clamp(FVector::DotProduct(BaseDirection, RelativeDirection), 0.f, 1.f);
 
 				// To improve position syncing, use old base for take-off
-				if (MovementBaseUtility::UseRelativeLocation(LastServerMovementBasePtr))
+				if (MovementBaseUtility::UseRelativeLocation(&LastMovementBaseInterfaceData))
 				{
 					// Relative Location
-					MovementBaseUtility::TransformLocationToLocal(LastServerMovementBasePtr, LastServerMovementBaseBoneName, UpdatedComponent->GetComponentLocation(), RelativeLocation);
+					MovementBaseUtility::TransformLocationToLocal(&LastMovementBaseInterfaceData, LastServerMovementBaseBoneName, UpdatedComponent->GetComponentLocation(), RelativeLocation);
 					bUseLastBase = true;
 				}
 			}
@@ -10279,16 +10370,19 @@ void UCharacterMovementComponent::ServerMoveHandleClientError(float ClientTimeSt
 	// Compute the client error from the server's position
 	// If client has accumulated a noticeable positional error, correct them.
 	bNetworkLargeClientCorrection = ServerData->bForceClientUpdate;
-	if (ServerData->bForceClientUpdate || (!bFallingWithinAcceptableError && ServerCheckClientError(ClientTimeStamp, DeltaTime, Accel, ClientLoc, RelativeClientLoc, ClientMovementBase, ClientBaseBoneName, ClientMovementMode)))
+	if (ServerData->bForceClientUpdate || (!bFallingWithinAcceptableError && ServerCheckClientError(ClientTimeStamp, DeltaTime, Accel, ClientLoc, RelativeClientLoc, ClientMovementBaseInterfaceData, ClientBaseBoneName, ClientMovementMode)))
 	{
 		ServerData->PendingAdjustment.NewVel = Velocity;
-		ServerData->PendingAdjustment.NewBase = MovementBase;
+		ServerData->PendingAdjustment.NewMovementBasePhysicsObjectOwner = MovementBaseInterfaceData ? MovementBaseInterfaceData->PhysicsObjectOwner.Get() : nullptr;
+		PRAGMA_DISABLE_DEPRECATION_WARNINGS
+		ServerData->PendingAdjustment.NewBase = CharacterOwner->GetBasedMovement().MovementBase;
+		PRAGMA_ENABLE_DEPRECATION_WARNINGS
 		ServerData->PendingAdjustment.NewBaseBoneName = MovementBaseBoneName;
 		ServerData->PendingAdjustment.NewLoc = FRepMovement::RebaseOntoZeroOrigin(ServerLoc, this);
 		ServerData->PendingAdjustment.NewRot = UpdatedComponent->GetComponentRotation();
 		ServerData->PendingAdjustment.GravityDirection = GravityDirection;
 
-		ServerData->PendingAdjustment.bBaseRelativePosition = (bDeferServerCorrectionsWhenFalling && bUseLastBase) || MovementBaseUtility::UseRelativeLocation(MovementBase);
+		ServerData->PendingAdjustment.bBaseRelativePosition = (bDeferServerCorrectionsWhenFalling && bUseLastBase) || MovementBaseUtility::UseRelativeLocation(MovementBaseInterfaceData);
 		ServerData->PendingAdjustment.bBaseRelativeVelocity = false;
 		
 		// Relative location?
@@ -10297,7 +10391,10 @@ void UCharacterMovementComponent::ServerMoveHandleClientError(float ClientTimeSt
 			if (bDeferServerCorrectionsWhenFalling && bUseLastBase)
 			{
 				ServerData->PendingAdjustment.NewVel = RelativeVelocity;
+				ServerData->PendingAdjustment.NewMovementBasePhysicsObjectOwner = LastMovementBaseInterfaceData.PhysicsObjectOwner.Get();
+				PRAGMA_DISABLE_DEPRECATION_WARNINGS
 				ServerData->PendingAdjustment.NewBase = LastServerMovementBase.Get();
+				PRAGMA_ENABLE_DEPRECATION_WARNINGS
 				ServerData->PendingAdjustment.NewBaseBoneName = LastServerMovementBaseBoneName;
 				ServerData->PendingAdjustment.NewLoc = RelativeLocation;
 			}
@@ -10309,7 +10406,7 @@ void UCharacterMovementComponent::ServerMoveHandleClientError(float ClientTimeSt
 					// Store world velocity converted to local space of movement base
 					ServerData->PendingAdjustment.bBaseRelativeVelocity = true;
 					const FVector CurrentVelocity = ServerData->PendingAdjustment.NewVel;
-					MovementBaseUtility::TransformDirectionToLocal(MovementBase, MovementBaseBoneName, CurrentVelocity, ServerData->PendingAdjustment.NewVel);
+					MovementBaseUtility::TransformDirectionToLocal(MovementBaseInterfaceData, MovementBaseBoneName, CurrentVelocity, ServerData->PendingAdjustment.NewVel);
 				}
 			}
 			
@@ -10322,7 +10419,8 @@ void UCharacterMovementComponent::ServerMoveHandleClientError(float ClientTimeSt
 		if (CharacterMovementCVars::NetShowCorrections != 0)
 		{
 			const FVector LocDiff = UpdatedComponent->GetComponentLocation() - ClientLoc;
-			const FString BaseString = MovementBase ? MovementBase->GetPathName(MovementBase->GetOutermost()) : TEXT("None");
+			const UObject* MovementBaseObject = MovementBaseInterfaceData ? MovementBaseInterfaceData->GetMovementBaseObject() : nullptr;
+			const FString BaseString = MovementBaseObject ? MovementBaseObject->GetPathName(MovementBaseObject->GetOutermost()) : TEXT("None");
 			UE_LOG(LogNetPlayerMovement, Warning, TEXT("*** Server: Error for %s at Time=%.3f is %3.3f LocDiff(%s) ClientLoc(%s) ServerLoc(%s) Base: %s Bone: %s Accel(%s) Velocity(%s)"),
 				*GetNameSafe(CharacterOwner), ClientTimeStamp, LocDiff.Size(), *LocDiff.ToString(), *ClientLoc.ToString(), *UpdatedComponent->GetComponentLocation().ToString(), *BaseString, *ServerData->PendingAdjustment.NewBaseBoneName.ToString(), *Accel.ToString(), *Velocity.ToString());
 			const float DebugLifetime = CharacterMovementCVars::NetCorrectionLifetime;
@@ -10343,10 +10441,10 @@ void UCharacterMovementComponent::ServerMoveHandleClientError(float ClientTimeSt
 	}
 	else
 	{
-		if (ServerShouldUseAuthoritativePosition(ClientTimeStamp, DeltaTime, Accel, ClientLoc, RelativeClientLoc, ClientMovementBase, ClientBaseBoneName, ClientMovementMode))
+		if (ServerShouldUseAuthoritativePosition(ClientTimeStamp, DeltaTime, Accel, ClientLoc, RelativeClientLoc, ClientMovementBaseInterfaceData, ClientBaseBoneName, ClientMovementMode))
 		{
 			const FVector LocDiff = UpdatedComponent->GetComponentLocation() - ClientLoc;
-			if (!LocDiff.IsZero() || ClientMovementMode != PackNetworkMovementMode() || GetMovementBase() != ClientMovementBase || (CharacterOwner && CharacterOwner->GetBasedMovement().BoneName != ClientBaseBoneName))
+			if (!LocDiff.IsZero() || ClientMovementMode != PackNetworkMovementMode() || GetMovementBaseInterfaceData() != ClientMovementBaseInterfaceData || (CharacterOwner && CharacterOwner->GetBasedMovement().BoneName != ClientBaseBoneName))
 			{
 				// Just set the position. On subsequent moves we will resolve initially overlapping conditions.
 				UpdatedComponent->SetWorldLocation(ClientLoc, false);
@@ -10355,7 +10453,7 @@ void UCharacterMovementComponent::ServerMoveHandleClientError(float ClientTimeSt
 				ApplyNetworkMovementMode(ClientMovementMode);
 
 				// Update base and floor at new location.
-				SetBase(ClientMovementBase, ClientBaseBoneName);
+				SetBase(ClientMovementBaseInterfaceData, ClientBaseBoneName);
 				UpdateFloorFromAdjustment();
 
 				// Even if base has not changed, we need to recompute the relative offsets (since we've moved).
@@ -10378,7 +10476,10 @@ void UCharacterMovementComponent::ServerMoveHandleClientError(float ClientTimeSt
 
 	ServerData->bForceClientUpdate = false;
 
-	LastServerMovementBase = MovementBase;
+	LastMovementBaseInterfaceData = MovementBaseInterfaceData;
+	PRAGMA_DISABLE_DEPRECATION_WARNINGS
+	LastServerMovementBase = CharacterOwner->GetBasedMovement().MovementBase;
+	PRAGMA_ENABLE_DEPRECATION_WARNINGS
 	LastServerMovementBaseBoneName = MovementBaseBoneName;
 	bLastClientIsFalling = bClientIsFalling;
 	bLastServerIsFalling = bServerIsFalling;
@@ -10386,6 +10487,12 @@ void UCharacterMovementComponent::ServerMoveHandleClientError(float ClientTimeSt
 }
 
 bool UCharacterMovementComponent::ServerCheckClientError(float ClientTimeStamp, float DeltaTime, const FVector& Accel, const FVector& ClientWorldLocation, const FVector& RelativeClientLocation, UPrimitiveComponent* ClientMovementBase, FName ClientBaseBoneName, uint8 ClientMovementMode)
+{
+	FMovementBaseInterfaceData ClientMovementBaseInterfaceData(ClientMovementBase);
+	return ServerCheckClientError(ClientTimeStamp, DeltaTime, Accel, ClientWorldLocation, RelativeClientLocation, &ClientMovementBaseInterfaceData, ClientBaseBoneName, ClientMovementMode);
+}
+
+bool UCharacterMovementComponent::ServerCheckClientError(float ClientTimeStamp, float DeltaTime, const FVector& Accel, const FVector& ClientWorldLocation, const FVector& RelativeClientLocation, FMovementBaseInterfaceData* ClientMovementBaseInterfaceData, FName ClientBaseBoneName, uint8 ClientMovementMode)
 {
 	// Check location difference against global setting
 	if (!bIgnoreClientMovementErrorChecksAndCorrection)
@@ -10400,7 +10507,7 @@ bool UCharacterMovementComponent::ServerCheckClientError(float ClientTimeStamp,
 		}
 #endif
 
-		if (ServerExceedsAllowablePositionError(ClientTimeStamp, DeltaTime, Accel, ClientWorldLocation, RelativeClientLocation, ClientMovementBase, ClientBaseBoneName, ClientMovementMode))
+		if (ServerExceedsAllowablePositionError(ClientTimeStamp, DeltaTime, Accel, ClientWorldLocation, RelativeClientLocation, ClientMovementBaseInterfaceData, ClientBaseBoneName, ClientMovementMode))
 		{
 			return true;
 		}
@@ -10429,8 +10536,14 @@ bool UCharacterMovementComponent::ServerCheckClientError(float ClientTimeStamp,
 	return false;
 }
 
-
 bool UCharacterMovementComponent::ServerExceedsAllowablePositionError(float ClientTimeStamp, float DeltaTime, const FVector& Accel, const FVector& ClientWorldLocation, const FVector& RelativeClientLocation, UPrimitiveComponent* ClientMovementBase, FName ClientBaseBoneName, uint8 ClientMovementMode)
+{
+	FMovementBaseInterfaceData ClientMovementBaseInterfaceData(ClientMovementBase);
+	return ServerExceedsAllowablePositionError(ClientTimeStamp, DeltaTime, Accel, ClientWorldLocation, RelativeClientLocation, &ClientMovementBaseInterfaceData, ClientBaseBoneName, ClientMovementMode);
+}
+
+
+bool UCharacterMovementComponent::ServerExceedsAllowablePositionError(float ClientTimeStamp, float DeltaTime, const FVector& Accel, const FVector& ClientWorldLocation, const FVector& RelativeClientLocation, FMovementBaseInterfaceData* ClientMovementBaseInterfaceData, FName ClientBaseBoneName, uint8 ClientMovementMode)
 {
 	// Check for disagreement in movement mode
 	const uint8 CurrentPackedMovementMode = PackNetworkMovementMode();
@@ -10453,6 +10566,12 @@ bool UCharacterMovementComponent::ServerExceedsAllowablePositionError(float Clie
 }
 
 bool UCharacterMovementComponent::ServerShouldUseAuthoritativePosition(float ClientTimeStamp, float DeltaTime, const FVector& Accel, const FVector& ClientWorldLocation, const FVector& RelativeClientLocation, UPrimitiveComponent* ClientMovementBase, FName ClientBaseBoneName, uint8 ClientMovementMode)
+{
+	FMovementBaseInterfaceData ClientMovementBaseInterfaceData(ClientMovementBase);
+	return ServerShouldUseAuthoritativePosition(ClientTimeStamp, DeltaTime, Accel, ClientWorldLocation, RelativeClientLocation, &ClientMovementBaseInterfaceData, ClientBaseBoneName, ClientMovementMode);
+}
+
+bool UCharacterMovementComponent::ServerShouldUseAuthoritativePosition(float ClientTimeStamp, float DeltaTime, const FVector& Accel, const FVector& ClientWorldLocation, const FVector& RelativeClientLocation, FMovementBaseInterfaceData* ClientMovementBaseInterfaceData, FName ClientBaseBoneName, uint8 ClientMovementMode)
 {
 	if (bServerAcceptClientAuthoritativePosition)
 	{
@@ -10475,7 +10594,9 @@ bool UCharacterMovementComponent::ServerMove_Validate(float TimeStamp, FVector_N
 
 void UCharacterMovementComponent::ServerMoveDual(float TimeStamp0, FVector_NetQuantize10 InAccel0, uint8 PendingFlags, uint32 View0, float TimeStamp, FVector_NetQuantize10 InAccel, FVector_NetQuantize100 ClientLoc, uint8 NewFlags, uint8 ClientRoll, uint32 View, UPrimitiveComponent* ClientMovementBase, FName ClientBaseBoneName, uint8 ClientMovementMode)
 {
-	if (MovementBaseUtility::IsDynamicBase(ClientMovementBase))
+	// Making a new FMovementBaseInterfaceData struct so we can use non-deprecated functions - we could also just disable deprecation warnings here too, since the function is deprecated itself
+	const FMovementBaseInterfaceData MovementBaseInterfaceData(ClientMovementBase);
+	if (MovementBaseUtility::IsDynamicBase(&MovementBaseInterfaceData))
 	{
 		//UE_LOG(LogCharacterMovement, Log, TEXT("ServerMoveDual: base %s"), *ClientMovementBase->GetName());
 		CharacterOwner->ServerMoveDual(TimeStamp0, InAccel0, PendingFlags, View0, TimeStamp, InAccel, ClientLoc, NewFlags, ClientRoll, View, ClientMovementBase, ClientBaseBoneName, ClientMovementMode);
@@ -10595,7 +10716,7 @@ void UCharacterMovementComponent::UpdateFloorFromAdjustment()
 	// If base is now NULL, presumably we are no longer walking. If we had a valid floor but don't find one now, we'll likely start falling.
 	// Try to keep the floor valid if we are walking without a movement base (but not NavWalking, see SetBase).
 	const bool bFindFloorIfOnGround = (CharacterMovementCVars::bRestoreFloorWhenMissingBaseComponent && IsMovingOnGround() && !bIsNavWalkingOnServer);
-	if (CharacterOwner->GetMovementBase() || bFindFloorIfOnGround)
+	if (MovementBaseUtility::IsMovementBaseDataValid(GetMovementBaseInterfaceData()) || bFindFloorIfOnGround)
 	{
 		FindFloor(UpdatedComponent->GetComponentLocation(), CurrentFloor, false);
 	}
@@ -10734,7 +10855,7 @@ void FCharacterMoveResponseDataContainer::ServerFillResponseData(const UCharacte
 	if (!PendingAdjustment.bAckGoodMove)
 	{
 		bHasRotation = CharacterMovement.ShouldCorrectRotation();
-		bHasBase = (PendingAdjustment.NewBase != nullptr);
+		bHasBase = (PendingAdjustment.NewMovementBasePhysicsObjectOwner != nullptr);
 		if (const ACharacter* CharacterOwner = CharacterMovement.GetCharacterOwner())
 		{
 			bRootMotionMontageCorrection = CharacterOwner->IsPlayingNetworkedRootMotionMontage();
@@ -10780,13 +10901,18 @@ bool FCharacterMoveResponseDataContainer::Serialize(UCharacterMovementComponent&
 			ClientAdjustment.NewRot = FRotator::ZeroRotator;
 		}
 
+		SerializeOptionalValue<TObjectPtr<UObject>>(bIsSaving, Ar, ClientAdjustment.NewMovementBasePhysicsObjectOwner, nullptr);
+		SerializeOptionalValue<FName>(bIsSaving, Ar, ClientAdjustment.NewBaseBoneName, NAME_None);
+		SerializeOptionalValue<uint8>(bIsSaving, Ar, ClientAdjustment.MovementMode, MOVE_Walking);
+		
+		PRAGMA_DISABLE_DEPRECATION_WARNINGS
 #if UE_WITH_REMOTE_OBJECT_HANDLE
-		SerializeOptionalValue<TObjectPtr<UPrimitiveComponent>>(bIsSaving, Ar, ClientAdjustment.NewBase, nullptr);
+		ClientAdjustment.NewBase = Cast<UPrimitiveComponent>(ClientAdjustment.NewMovementBasePhysicsObjectOwner.Get());
 #else
-		SerializeOptionalValue<UPrimitiveComponent*>(bIsSaving, Ar, ClientAdjustment.NewBase, nullptr);
+		ClientAdjustment.NewBase = Cast<UPrimitiveComponent>(ClientAdjustment.NewMovementBasePhysicsObjectOwner.Get());
 #endif
-		SerializeOptionalValue<FName>(bIsSaving, Ar, ClientAdjustment.NewBaseBoneName, NAME_None);
-		SerializeOptionalValue<uint8>(bIsSaving, Ar, ClientAdjustment.MovementMode, MOVE_Walking);
+		PRAGMA_ENABLE_DEPRECATION_WARNINGS
+
 		Ar.SerializeBits(&ClientAdjustment.bBaseRelativePosition, 1);
 		Ar.SerializeBits(&ClientAdjustment.bBaseRelativeVelocity, 1);
 
@@ -10824,6 +10950,8 @@ void UCharacterMovementComponent::ClientHandleMoveResponse(const FCharacterMoveR
 	}
 	else
 	{
+		FMovementBaseInterfaceData MovementBaseInterfaceDataFromClientAdjustment = MovementBaseUtility::GetMovementBaseDataFromPhysicsOwner(MoveResponse.ClientAdjustment.NewMovementBasePhysicsObjectOwner);
+		
 		// Wrappers to old RPC handlers, to maintain compatibility. If overrides need additional serialized data, they can access GetMoveResponseDataContainer()
 		if (MoveResponse.bRootMotionSourceCorrection)
 		{
@@ -10837,7 +10965,7 @@ void UCharacterMovementComponent::ClientHandleMoveResponse(const FCharacterMoveR
 					MoveResponse.ClientAdjustment.NewLoc,
 					MoveResponse.RootMotionRotation,
 					GetGravitySpaceZ(MoveResponse.ClientAdjustment.NewVel),
-					MoveResponse.ClientAdjustment.NewBase,
+					&MovementBaseInterfaceDataFromClientAdjustment,
 					MoveResponse.ClientAdjustment.NewBaseBoneName,
 					MoveResponse.bHasBase,
 					MoveResponse.ClientAdjustment.bBaseRelativePosition,
@@ -10852,7 +10980,7 @@ void UCharacterMovementComponent::ClientHandleMoveResponse(const FCharacterMoveR
 				MoveResponse.ClientAdjustment.NewLoc,
 				MoveResponse.RootMotionRotation,
 				GetGravitySpaceZ(MoveResponse.ClientAdjustment.NewVel),
-				MoveResponse.ClientAdjustment.NewBase,
+				&MovementBaseInterfaceDataFromClientAdjustment,
 				MoveResponse.ClientAdjustment.NewBaseBoneName,
 				MoveResponse.bHasBase,
 				MoveResponse.ClientAdjustment.bBaseRelativePosition,
@@ -10864,7 +10992,7 @@ void UCharacterMovementComponent::ClientHandleMoveResponse(const FCharacterMoveR
 				MoveResponse.ClientAdjustment.TimeStamp,
 				MoveResponse.ClientAdjustment.NewLoc,
 				MoveResponse.ClientAdjustment.NewVel,
-				MoveResponse.ClientAdjustment.NewBase,
+				&MovementBaseInterfaceDataFromClientAdjustment,
 				MoveResponse.ClientAdjustment.NewBaseBoneName,
 				MoveResponse.bHasBase,
 				MoveResponse.ClientAdjustment.bBaseRelativePosition,
@@ -10947,7 +11075,11 @@ void UCharacterMovementComponent::SendClientAdjustment()
 			}
 			else
 			{
+				// Filling out a UPrimitiveComponent from the movement base struct for old functionality
 				const bool bIsPlayingNetworkedRootMotionMontage = CharacterOwner->IsPlayingNetworkedRootMotionMontage();
+				UPrimitiveComponent* MovementBase = Cast<UPrimitiveComponent>(ServerData->PendingAdjustment.NewMovementBasePhysicsObjectOwner.Get());
+				const bool bHasMovementBase = MovementBase != nullptr;
+				
 				if (CurrentRootMotion.HasActiveRootMotionSources())
 				{
 					FRotator Rotation = ServerData->PendingAdjustment.NewRot.GetNormalized();
@@ -10961,9 +11093,9 @@ void UCharacterMovementComponent::SendClientAdjustment()
 						ServerData->PendingAdjustment.NewLoc,
 						CompressedRotation,
 						GetGravitySpaceZ(ServerData->PendingAdjustment.NewVel),
-						ServerData->PendingAdjustment.NewBase,
+						MovementBase,
 						ServerData->PendingAdjustment.NewBaseBoneName,
-						ServerData->PendingAdjustment.NewBase != NULL,
+						bHasMovementBase,
 						ServerData->PendingAdjustment.bBaseRelativePosition,
 						ServerData->PendingAdjustment.MovementMode
 					);
@@ -10979,9 +11111,9 @@ void UCharacterMovementComponent::SendClientAdjustment()
 						ServerData->PendingAdjustment.NewLoc,
 						CompressedRotation,
 						GetGravitySpaceZ(ServerData->PendingAdjustment.NewVel),
-						ServerData->PendingAdjustment.NewBase,
+						MovementBase,
 						ServerData->PendingAdjustment.NewBaseBoneName,
-						ServerData->PendingAdjustment.NewBase != NULL,
+						bHasMovementBase,
 						ServerData->PendingAdjustment.bBaseRelativePosition,
 						ServerData->PendingAdjustment.MovementMode
 					);
@@ -10992,9 +11124,9 @@ void UCharacterMovementComponent::SendClientAdjustment()
 					(
 						ServerData->PendingAdjustment.TimeStamp,
 						ServerData->PendingAdjustment.NewLoc,
-						ServerData->PendingAdjustment.NewBase,
+						MovementBase,
 						ServerData->PendingAdjustment.NewBaseBoneName,
-						ServerData->PendingAdjustment.NewBase != NULL,
+						bHasMovementBase,
 						ServerData->PendingAdjustment.bBaseRelativePosition,
 						ServerData->PendingAdjustment.MovementMode
 					);
@@ -11006,9 +11138,9 @@ void UCharacterMovementComponent::SendClientAdjustment()
 						ServerData->PendingAdjustment.TimeStamp,
 						ServerData->PendingAdjustment.NewLoc,
 						ServerData->PendingAdjustment.NewVel,
-						ServerData->PendingAdjustment.NewBase,
+						MovementBase,
 						ServerData->PendingAdjustment.NewBaseBoneName,
-						ServerData->PendingAdjustment.NewBase != NULL,
+						bHasMovementBase,
 						ServerData->PendingAdjustment.bBaseRelativePosition,
 						ServerData->PendingAdjustment.MovementMode
 					);
@@ -11028,11 +11160,18 @@ void UCharacterMovementComponent::ClientVeryShortAdjustPosition(float TimeStamp,
 	CharacterOwner->ClientVeryShortAdjustPosition(TimeStamp, NewLoc, NewBase, NewBaseBoneName, bHasBase, bBaseRelativePosition, ServerMovementMode);
 }
 
+
+void UCharacterMovementComponent::ClientVeryShortAdjustPosition_Implementation(float TimeStamp, FVector NewLoc, UPrimitiveComponent* NewBase, FName NewBaseBoneName, bool bHasBase, bool bBaseRelativePosition, uint8 ServerMovementMode)
+{
+	FMovementBaseInterfaceData NewMovementBaseInterfaceData(NewBase);
+	ClientVeryShortAdjustPosition_Implementation(TimeStamp, NewLoc, &NewMovementBaseInterfaceData, NewBaseBoneName, bHasBase, bBaseRelativePosition, ServerMovementMode);
+}
+
 void UCharacterMovementComponent::ClientVeryShortAdjustPosition_Implementation
 	(
 	float TimeStamp,
 	FVector NewLoc,
-	UPrimitiveComponent* NewBase,
+	FMovementBaseInterfaceData* NewMovementBaseInterfaceData,
 	FName NewBaseBoneName,
 	bool bHasBase,
 	bool bBaseRelativePosition,
@@ -11041,7 +11180,8 @@ void UCharacterMovementComponent::ClientVeryShortAdjustPosition_Implementation
 {
 	if (HasValidData())
 	{
-		ClientAdjustPosition(TimeStamp, NewLoc, FVector::ZeroVector, NewBase, NewBaseBoneName, bHasBase, bBaseRelativePosition, ServerMovementMode);
+		UPrimitiveComponent* MovementBase = NewMovementBaseInterfaceData ? Cast<UPrimitiveComponent>(NewMovementBaseInterfaceData->PhysicsObjectOwner.Get()) : nullptr;
+		ClientAdjustPosition(TimeStamp, NewLoc, FVector::ZeroVector, MovementBase, NewBaseBoneName, bHasBase, bBaseRelativePosition, ServerMovementMode);
 	}
 }
 
@@ -11051,12 +11191,19 @@ void UCharacterMovementComponent::ClientAdjustPosition(float TimeStamp, FVector
 	CharacterOwner->ClientAdjustPosition(TimeStamp, NewLoc, NewVel, NewBase, NewBaseBoneName, bHasBase, bBaseRelativePosition, ServerMovementMode);
 }
 
+
+void UCharacterMovementComponent::ClientAdjustPosition_Implementation(float TimeStamp, FVector NewLoc, FVector NewVel, UPrimitiveComponent* NewBase, FName NewBaseBoneName, bool bHasBase, bool bBaseRelativePosition, uint8 ServerMovementMode, TOptional<FRotator> OptionalRotation)
+{
+	FMovementBaseInterfaceData NewMovementBaseInterfaceData(NewBase);
+	ClientAdjustPosition_Implementation(TimeStamp, NewLoc, NewVel, &NewMovementBaseInterfaceData, NewBaseBoneName, bHasBase, bBaseRelativePosition, ServerMovementMode, OptionalRotation);
+}
+
 void UCharacterMovementComponent::ClientAdjustPosition_Implementation
 	(
 	float TimeStamp,
 	FVector NewLocation,
 	FVector NewVelocity,
-	UPrimitiveComponent* NewBase,
+	FMovementBaseInterfaceData* NewMovementBaseInterfaceData,
 	FName NewBaseBoneName,
 	bool bHasBase,
 	bool bBaseRelativePosition,
@@ -11074,13 +11221,13 @@ void UCharacterMovementComponent::ClientAdjustPosition_Implementation
 	check(ClientData);
 	
 	// Make sure the base actor exists on this client.
-	const bool bUnresolvedBase = bHasBase && (NewBase == NULL);
+	const bool bUnresolvedBase = bHasBase && !MovementBaseUtility::IsMovementBaseDataValid(NewMovementBaseInterfaceData);
 	if (bUnresolvedBase)
 	{
 		if (bBaseRelativePosition)
 		{
 			UE_LOG(LogNetPlayerMovement, Warning, TEXT("ClientAdjustPosition_Implementation could not resolve the new relative movement base actor, ignoring server correction! Client currently at world location %s on base %s"),
-				*UpdatedComponent->GetComponentLocation().ToString(), *GetNameSafe(GetMovementBase()));
+				*UpdatedComponent->GetComponentLocation().ToString(), *GetNameSafe(NewMovementBaseInterfaceData ? NewMovementBaseInterfaceData->GetMovementBaseObject() : nullptr));
 			return;
 		}
 		else
@@ -11106,7 +11253,7 @@ void UCharacterMovementComponent::ClientAdjustPosition_Implementation
 	//  Received Location is relative to dynamic base
 	if (bBaseRelativePosition)
 	{
-		MovementBaseUtility::TransformLocationToWorld(NewBase, NewBaseBoneName, NewLocation, WorldShiftedNewLocation); // TODO: error handling if returns false	
+		MovementBaseUtility::TransformLocationToWorld(NewMovementBaseInterfaceData, NewBaseBoneName, NewLocation, WorldShiftedNewLocation); // TODO: error handling if returns false	
 	}
 	else
 	{
@@ -11119,7 +11266,7 @@ void UCharacterMovementComponent::ClientAdjustPosition_Implementation
 	{
 		// Convert Relative Velocity -> World Velocity
 		const FVector CurrentVelocity = NewVelocity;
-		MovementBaseUtility::TransformDirectionToWorld(NewBase, NewBaseBoneName, CurrentVelocity, NewVelocity);
+		MovementBaseUtility::TransformDirectionToWorld(NewMovementBaseInterfaceData, NewBaseBoneName, CurrentVelocity, NewVelocity);
 	}
 
 	// Fall back to the last-known good rotation if the server didn't send one
@@ -11131,7 +11278,7 @@ void UCharacterMovementComponent::ClientAdjustPosition_Implementation
 	}
 
 	// Trigger event
-	OnClientCorrectionReceived(*ClientData, TimeStamp, WorldShiftedNewLocation, NewVelocity, NewBase, NewBaseBoneName, bHasBase, bBaseRelativePosition, ServerMovementMode, DefaultGravityDirection);
+	OnClientCorrectionReceived(*ClientData, TimeStamp, WorldShiftedNewLocation, NewVelocity, NewMovementBaseInterfaceData, NewBaseBoneName, bHasBase, bBaseRelativePosition, ServerMovementMode, DefaultGravityDirection);
 
 	// Trust the server's positioning.
 	if (UpdatedComponent)
@@ -11149,35 +11296,34 @@ void UCharacterMovementComponent::ClientAdjustPosition_Implementation
 	Velocity = NewVelocity;
 
 	// Trust the server's movement mode
-	UPrimitiveComponent* PreviousBase = CharacterOwner->GetMovementBase();
 	ApplyNetworkMovementMode(ServerMovementMode);
 
 	// Set base component
-	UPrimitiveComponent* FinalBase = NewBase;
+	FMovementBaseInterfaceData FinalMovementBaseInterfaceData = NewMovementBaseInterfaceData;
 	FName FinalBaseBoneName = NewBaseBoneName;
 	if (bUnresolvedBase)
 	{
-		check(NewBase == NULL);
+		check(!FinalMovementBaseInterfaceData.IsValid());
 		check(!bBaseRelativePosition);
 		
 		// We had an unresolved base from the server
 		// If walking, we'd like to continue walking if possible, to avoid falling for a frame, so try to find a base where we moved to.
-		if (PreviousBase && UpdatedComponent)
+		if (MovementBaseUtility::IsMovementBaseDataValid(GetMovementBaseInterfaceData()) && UpdatedComponent)
 		{
 			FindFloor(UpdatedComponent->GetComponentLocation(), CurrentFloor, false);
 			if (CurrentFloor.IsWalkableFloor())
 			{
-				FinalBase = CurrentFloor.HitResult.Component.Get();
+				FinalMovementBaseInterfaceData = MovementBaseUtility::GetMovementBaseDataFromHitResult(&CurrentFloor.HitResult);
 				FinalBaseBoneName = CurrentFloor.HitResult.BoneName;
 			}
 			else
 			{
-				FinalBase = nullptr;
+				FinalMovementBaseInterfaceData.Clear();
 				FinalBaseBoneName = NAME_None;
 			}
 		}
 	}
-	SetBase(FinalBase, FinalBaseBoneName);
+	SetBase(&FinalMovementBaseInterfaceData, FinalBaseBoneName);
 
 	// Update floor at new location
 	UpdateFloorFromAdjustment();
@@ -11201,6 +11347,12 @@ void UCharacterMovementComponent::ClientAdjustRootMotionPosition(float TimeStamp
 }
 
 void UCharacterMovementComponent::OnClientCorrectionReceived(FNetworkPredictionData_Client_Character& ClientData, float TimeStamp, FVector NewLocation, FVector NewVelocity, UPrimitiveComponent* NewBase, FName NewBaseBoneName, bool bHasBase, bool bBaseRelativePosition, uint8 ServerMovementMode, FVector ServerGravityDirection)
+{
+	FMovementBaseInterfaceData NewMovementBaseInterfaceData(NewBase);
+	OnClientCorrectionReceived(ClientData, TimeStamp, NewLocation, NewVelocity, &NewMovementBaseInterfaceData, NewBaseBoneName, bHasBase, bBaseRelativePosition, ServerMovementMode, ServerGravityDirection);
+}
+
+void UCharacterMovementComponent::OnClientCorrectionReceived(FNetworkPredictionData_Client_Character& ClientData, float TimeStamp, FVector NewLocation, FVector NewVelocity, FMovementBaseInterfaceData* NewMovementBaseInterfaceData, FName NewBaseBoneName, bool bHasBase, bool bBaseRelativePosition, uint8 ServerMovementMode, FVector ServerGravityDirection)
 {
 #if CSV_PROFILER_STATS
 	const bool bIsRecordingDetailedCsvStats = FCsvProfiler::Get()->IsCategoryEnabled(CSV_CATEGORY_INDEX(CharacterMovementDetailed));
@@ -11221,7 +11373,8 @@ void UCharacterMovementComponent::OnClientCorrectionReceived(FNetworkPredictionD
 	{
 		const FVector ClientLocAtCorrectedMove = ClientData.LastAckedMove.IsValid() ? ClientData.LastAckedMove->SavedLocation : UpdatedComponent->GetComponentLocation();
 		const FVector LocDiff = ClientLocAtCorrectedMove - NewLocation;
-		const FString NewBaseString = NewBase ? NewBase->GetPathName(NewBase->GetOutermost()) : TEXT("None");
+		const UObject* MovementBaseObject = NewMovementBaseInterfaceData ? NewMovementBaseInterfaceData->GetMovementBaseObject() : nullptr;
+		const FString NewBaseString = MovementBaseUtility::IsMovementBaseDataValid(NewMovementBaseInterfaceData) ? MovementBaseObject->GetPathName(MovementBaseObject->GetOutermost()) : TEXT("None");
 		UE_LOG(LogNetPlayerMovement, Warning, TEXT("*** Client: Error for %s at Time=%.3f is %3.3f LocDiff(%s) ClientLoc(%s) ServerLoc(%s) NewBase: %s NewBone: %s ClientVel(%s) ServerVel(%s) SavedMoves %d"),
 			   *GetNameSafe(CharacterOwner), TimeStamp, LocDiff.Size(), *LocDiff.ToString(), *ClientLocAtCorrectedMove.ToString(), *NewLocation.ToString(), *NewBaseString, *NewBaseBoneName.ToString(), *Velocity.ToString(), *NewVelocity.ToString(), ClientData.SavedMoves.Num());
 		const float DebugLifetime = CharacterMovementCVars::NetCorrectionLifetime;
@@ -11252,6 +11405,11 @@ void UCharacterMovementComponent::OnClientCorrectionReceived(FNetworkPredictionD
 #endif //ROOT_MOTION_DEBUG
 }
 
+void UCharacterMovementComponent::ClientAdjustRootMotionPosition_Implementation(float TimeStamp, float ServerMontageTrackPosition, FVector ServerLoc, FVector_NetQuantizeNormal ServerRotation, float ServerVelZ, UPrimitiveComponent* ServerBase, FName ServerBoneName, bool bHasBase, bool bBaseRelativePosition, uint8 ServerMovementMode)
+{
+	FMovementBaseInterfaceData ServerMovementBaseInterfaceData(ServerBase);
+	ClientAdjustRootMotionPosition_Implementation(TimeStamp, ServerMontageTrackPosition, ServerLoc, ServerRotation, ServerVelZ, &ServerMovementBaseInterfaceData, ServerBoneName, bHasBase, bBaseRelativePosition, ServerMovementMode);
+}
 
 void UCharacterMovementComponent::ClientAdjustRootMotionPosition_Implementation(
 	float TimeStamp,
@@ -11259,7 +11417,7 @@ void UCharacterMovementComponent::ClientAdjustRootMotionPosition_Implementation(
 	FVector ServerLoc,
 	FVector_NetQuantizeNormal ServerRotation,
 	float ServerVelZ,
-	UPrimitiveComponent * ServerBase,
+	FMovementBaseInterfaceData* ServerMovementBaseInterfaceData,
 	FName ServerBaseBoneName,
 	bool bHasBase,
 	bool bBaseRelativePosition,
@@ -11271,7 +11429,7 @@ void UCharacterMovementComponent::ClientAdjustRootMotionPosition_Implementation(
 	}
 
 	// Call ClientAdjustPosition first. This will Ack the move if it's not outdated.
-	ClientAdjustPosition_Implementation(TimeStamp, ServerLoc, ServerVelZ * -GetGravityDirection(), ServerBase, ServerBaseBoneName, bHasBase, bBaseRelativePosition, ServerMovementMode);
+	ClientAdjustPosition_Implementation(TimeStamp, ServerLoc, ServerVelZ * -GetGravityDirection(), ServerMovementBaseInterfaceData, ServerBaseBoneName, bHasBase, bBaseRelativePosition, ServerMovementMode);
 	
 	FNetworkPredictionData_Client_Character* ClientData = GetPredictionData_Client_Character();
 	check(ClientData);
@@ -11287,7 +11445,7 @@ void UCharacterMovementComponent::ClientAdjustRootMotionPosition_Implementation(
 	CharacterOwner->SetActorRotation(DecompressedRot);
 	const FVector ServerLocation(FRepMovement::RebaseOntoLocalOrigin(ServerLoc, UpdatedComponent));
 	UE_LOG(LogRootMotion, Log,  TEXT("ClientAdjustRootMotionPosition_Implementation TimeStamp: %f, ServerMontageTrackPosition: %f, ServerLocation: %s, ServerRotation: %s, ServerVelZ: %f, ServerBase: %s"),
-		TimeStamp, ServerMontageTrackPosition, *ServerLocation.ToCompactString(), *DecompressedRot.ToCompactString(), ServerVelZ, *GetNameSafe(ServerBase) );
+		TimeStamp, ServerMontageTrackPosition, *ServerLocation.ToCompactString(), *DecompressedRot.ToCompactString(), ServerVelZ, *GetNameSafe(ServerMovementBaseInterfaceData ? ServerMovementBaseInterfaceData->GetMovementBaseObject() : nullptr));
 
 #if 0
 	// DEBUG - get some insight on where errors came from
@@ -11325,6 +11483,12 @@ void UCharacterMovementComponent::ClientAdjustRootMotionSourcePosition(float Tim
 	CharacterOwner->ClientAdjustRootMotionSourcePosition(TimeStamp, ServerRootMotion, bHasAnimRootMotion, ServerMontageTrackPosition, ServerLoc, ServerRotation, ServerVelZ, ServerBase, ServerBoneName, bHasBase, bBaseRelativePosition, ServerMovementMode);
 }
 
+void UCharacterMovementComponent::ClientAdjustRootMotionSourcePosition_Implementation(float TimeStamp, FRootMotionSourceGroup ServerRootMotion, bool bHasAnimRootMotion, float ServerMontageTrackPosition, FVector ServerLoc, FVector_NetQuantizeNormal ServerRotation, float ServerVelZ, UPrimitiveComponent* ServerBase, FName ServerBoneName, bool bHasBase, bool bBaseRelativePosition, uint8 ServerMovementMode)
+{
+	FMovementBaseInterfaceData ServerMovementBaseInterfaceData(ServerBase);
+	ClientAdjustRootMotionSourcePosition_Implementation(TimeStamp, ServerRootMotion, bHasAnimRootMotion, ServerMontageTrackPosition, ServerLoc, ServerRotation, ServerVelZ, &ServerMovementBaseInterfaceData, ServerBoneName, bHasBase, bBaseRelativePosition, ServerMovementMode);
+}
+
 void UCharacterMovementComponent::ClientAdjustRootMotionSourcePosition_Implementation(
 	float TimeStamp,
 	FRootMotionSourceGroup ServerRootMotion,
@@ -11333,7 +11497,7 @@ void UCharacterMovementComponent::ClientAdjustRootMotionSourcePosition_Implement
 	FVector ServerLoc,
 	FVector_NetQuantizeNormal ServerRotation,
 	float ServerVelZ,
-	UPrimitiveComponent * ServerBase,
+	FMovementBaseInterfaceData* ServerMovementBaseInterfaceData,
 	FName ServerBaseBoneName,
 	bool bHasBase,
 	bool bBaseRelativePosition,
@@ -11354,7 +11518,7 @@ void UCharacterMovementComponent::ClientAdjustRootMotionSourcePosition_Implement
 #endif
 
 	// Call ClientAdjustPosition first. This will Ack the move if it's not outdated.
-	ClientAdjustPosition_Implementation(TimeStamp, ServerLoc, ServerVelZ * -GetGravityDirection(), ServerBase, ServerBaseBoneName, bHasBase, bBaseRelativePosition, ServerMovementMode);
+	ClientAdjustPosition_Implementation(TimeStamp, ServerLoc, ServerVelZ * -GetGravityDirection(), ServerMovementBaseInterfaceData, ServerBaseBoneName, bHasBase, bBaseRelativePosition, ServerMovementMode);
 	
 	FNetworkPredictionData_Client_Character* ClientData = GetPredictionData_Client_Character();
 	check(ClientData);
@@ -11370,7 +11534,7 @@ void UCharacterMovementComponent::ClientAdjustRootMotionSourcePosition_Implement
 	CharacterOwner->SetActorRotation(DecompressedRot);
 	const FVector ServerLocation(FRepMovement::RebaseOntoLocalOrigin(ServerLoc, UpdatedComponent));
 	UE_LOG(LogRootMotion, Log,  TEXT("ClientAdjustRootMotionSourcePosition_Implementation TimeStamp: %f, NumRootMotionSources: %d, ServerLocation: %s, ServerRotation: %s, ServerVelZ: %f, ServerBase: %s"),
-		TimeStamp, ServerRootMotion.RootMotionSources.Num(), *ServerLocation.ToCompactString(), *DecompressedRot.ToCompactString(), ServerVelZ, *GetNameSafe(ServerBase) );
+		TimeStamp, ServerRootMotion.RootMotionSources.Num(), *ServerLocation.ToCompactString(), *DecompressedRot.ToCompactString(), ServerVelZ, *GetNameSafe(ServerMovementBaseInterfaceData ? ServerMovementBaseInterfaceData->GetMovementBaseObject() : nullptr));
 
 	// Handle AnimRootMotion correction
 	if (bHasAnimRootMotion)
@@ -12502,7 +12666,10 @@ void FSavedMove_Character::Clear()
 	StartBaseRotation = FQuat::Identity;
 	StartCapsuleRadius = 0.f;
 	StartCapsuleHalfHeight = 0.f;
+	StartMovementBaseInterfaceData.Clear();
+	PRAGMA_DISABLE_DEPRECATION_WARNINGS
 	StartBase = nullptr;
+	PRAGMA_ENABLE_DEPRECATION_WARNINGS
 	StartBoneName = NAME_None;
 	StartActorOverlapCounter = 0;
 	StartComponentOverlapCounter = 0;
@@ -12522,7 +12689,10 @@ void FSavedMove_Character::Clear()
 	AccelMag = 0.0f;
 	AccelNormal = FVector::ZeroVector;
 
+	EndMovementBaseInterfaceData.Clear();
+	PRAGMA_DISABLE_DEPRECATION_WARNINGS
 	EndBase = nullptr;
+	PRAGMA_ENABLE_DEPRECATION_WARNINGS
 	EndBoneName = NAME_None;
 	EndActorOverlapCounter = 0;
 	EndComponentOverlapCounter = 0;
@@ -12597,8 +12767,10 @@ void FSavedMove_Character::SetInitialPosition(ACharacter* Character)
 	StartLocation = Character->GetActorLocation();
 	StartRotation = Character->GetActorRotation();
 	StartVelocity = Character->GetCharacterMovement()->Velocity;
-	UPrimitiveComponent* const MovementBase = Character->GetMovementBase();
-	StartBase = MovementBase;
+	StartMovementBaseInterfaceData = Character->GetBasedMovement().MovementBaseInterfaceData;
+	PRAGMA_DISABLE_DEPRECATION_WARNINGS
+	StartBase = Character->GetBasedMovement().MovementBase;
+	PRAGMA_ENABLE_DEPRECATION_WARNINGS
 	StartBaseRotation = FQuat::Identity;
 	StartFloor = Character->GetCharacterMovement()->CurrentFloor;
 	CustomTimeDilation = Character->CustomTimeDilation;
@@ -12606,11 +12778,11 @@ void FSavedMove_Character::SetInitialPosition(ACharacter* Character)
 	StartActorOverlapCounter = Character->NumActorOverlapEventsCounter;
 	StartComponentOverlapCounter = UPrimitiveComponent::GlobalOverlapEventsCounter;
 
-	if (MovementBaseUtility::UseRelativeLocation(MovementBase))
+	if (MovementBaseUtility::UseRelativeLocation(&StartMovementBaseInterfaceData))
 	{
 		StartRelativeLocation = Character->GetBasedMovement().Location;
 		FVector StartBaseLocation_Unused;
-		MovementBaseUtility::GetMovementBaseTransform(MovementBase, StartBoneName, StartBaseLocation_Unused, StartBaseRotation);
+		MovementBaseUtility::GetMovementBaseTransform(&StartMovementBaseInterfaceData, StartBoneName, StartBaseLocation_Unused, StartBaseRotation);
 	}
 
 	// Attachment state
@@ -12657,18 +12829,20 @@ void FSavedMove_Character::PostUpdate(ACharacter* Character, FSavedMove_Characte
 		}
 #endif
 		// Movement base and base-relative movement
-		UPrimitiveComponent* const MovementBase = Character->GetMovementBase();
-		EndBase = MovementBase;
+		EndMovementBaseInterfaceData = Character->GetBasedMovement().MovementBaseInterfaceData;
+		PRAGMA_DISABLE_DEPRECATION_WARNINGS
+		EndBase = Character->GetBasedMovement().MovementBase;
+		PRAGMA_ENABLE_DEPRECATION_WARNINGS
 		EndBoneName = Character->GetBasedMovement().BoneName;
-		if (MovementBaseUtility::UseRelativeLocation(MovementBase))
+		if (MovementBaseUtility::UseRelativeLocation(&EndMovementBaseInterfaceData))
 		{
 			SavedRelativeLocation = Character->GetBasedMovement().Location;
 		}
 
 		// Save off movement base-relative acceleration if needed
-		if (CharacterMovementCVars::NetUseBaseRelativeAcceleration && MovementBaseUtility::IsDynamicBase(MovementBase))
+		if (CharacterMovementCVars::NetUseBaseRelativeAcceleration && MovementBaseUtility::IsDynamicBase(&EndMovementBaseInterfaceData))
 		{
-			MovementBaseUtility::TransformDirectionToLocal(EndBase.Get(), EndBoneName, Acceleration, SavedRelativeAcceleration);
+			MovementBaseUtility::TransformDirectionToLocal(&EndMovementBaseInterfaceData, EndBoneName, Acceleration, SavedRelativeAcceleration);
 		}
 
 		// Attachment state
@@ -12792,11 +12966,10 @@ FVector FSavedMove_Character::GetRevertedLocation() const
 		return AttachParent->GetSocketTransform(StartAttachSocketName).TransformPosition(StartAttachRelativeLocation);
 	}
 
-	const UPrimitiveComponent* MovementBase = StartBase.Get();
-	if (MovementBaseUtility::UseRelativeLocation(MovementBase))
+	if (MovementBaseUtility::UseRelativeLocation(&StartMovementBaseInterfaceData))
 	{
 		FVector WorldSpacePosition;
-		MovementBaseUtility::TransformLocationToWorld(MovementBase, StartBoneName, StartRelativeLocation, WorldSpacePosition);
+		MovementBaseUtility::TransformLocationToWorld(&StartMovementBaseInterfaceData, StartBoneName, StartRelativeLocation, WorldSpacePosition);
 		return WorldSpacePosition;
 	}
 
@@ -12945,10 +13118,10 @@ bool FSavedMove_Character::CanCombineWith(const FSavedMovePtr& NewMovePtr, AChar
 		return false;
 	}
 
-	const UPrimitiveComponent* OldBasePtr = StartBase.Get();
-	const UPrimitiveComponent* NewBasePtr = NewMove->StartBase.Get();
-	const bool bDynamicBaseOld = MovementBaseUtility::IsDynamicBase(OldBasePtr);
-	const bool bDynamicBaseNew = MovementBaseUtility::IsDynamicBase(NewBasePtr);
+	const FMovementBaseInterfaceData* OldMovementBaseInterfaceData = &StartMovementBaseInterfaceData;
+	const FMovementBaseInterfaceData* NewMovementBaseInterfaceData = &NewMove->StartMovementBaseInterfaceData;
+	const bool bDynamicBaseOld = MovementBaseUtility::IsDynamicBase(OldMovementBaseInterfaceData);
+	const bool bDynamicBaseNew = MovementBaseUtility::IsDynamicBase(NewMovementBaseInterfaceData);
 
 	// Change between static/dynamic requires separate moves (position sent as world vs relative)
 	if (bDynamicBaseOld != bDynamicBaseNew)
@@ -12960,7 +13133,7 @@ bool FSavedMove_Character::CanCombineWith(const FSavedMovePtr& NewMovePtr, AChar
 	const bool bPreventOnStaticBaseChange = (CharacterMovementCVars::NetEnableMoveCombiningOnStaticBaseChange == 0);
 	if (bPreventOnStaticBaseChange || (bDynamicBaseOld || bDynamicBaseNew))
 	{
-		if (OldBasePtr != NewBasePtr)
+		if (OldMovementBaseInterfaceData != NewMovementBaseInterfaceData)
 		{
 			return false;
 		}
@@ -13063,7 +13236,8 @@ void FSavedMove_Character::CombineWith(const FSavedMove_Character* OldMove, ACha
 	
 	CharMovement->Velocity = OldMove->StartVelocity;
 
-	CharMovement->SetBase(OldMove->StartBase.Get(), OldMove->StartBoneName);
+	FMovementBaseInterfaceData OldMovementBaseInterfaceData = OldMove->StartMovementBaseInterfaceData;
+	CharMovement->SetBase(&OldMovementBaseInterfaceData, OldMove->StartBoneName);
 	CharMovement->CurrentFloor = OldMove->StartFloor;
 
 	// Now that we have reverted to the old position, prepare a new move from that position,
@@ -13228,7 +13402,7 @@ void UCharacterMovementComponent::FlushServerMoves()
 
 			UE_CLOG(CharacterOwner && UpdatedComponent, LogNetPlayerMovement, Verbose, TEXT("ClientMove (Flush) Time %f Acceleration %s Velocity %s Position %s DeltaTime %f Mode %s MovementBase %s.%s (Dynamic:%d) DualMove? %d"),
 				NewMove->TimeStamp, *NewMove->Acceleration.ToString(), *Velocity.ToString(), *UpdatedComponent->GetComponentLocation().ToString(), NewMove->DeltaTime, *GetMovementName(),
-				*GetNameSafe(NewMove->EndBase.Get()), *NewMove->EndBoneName.ToString(), MovementBaseUtility::IsDynamicBase(NewMove->EndBase.Get()) ? 1 : 0, ClientData->PendingMove.IsValid() ? 1 : 0);
+				*GetNameSafe(NewMove->EndMovementBaseInterfaceData.GetMovementBaseObject()), *NewMove->EndBoneName.ToString(), MovementBaseUtility::IsDynamicBase(&NewMove->EndMovementBaseInterfaceData) ? 1 : 0, ClientData->PendingMove.IsValid() ? 1 : 0);
 
 			if (ShouldUsePackedMovementRPCs())
 			{
@@ -13489,7 +13663,7 @@ void UCharacterMovementComponent::FillAsyncInput(const FVector& InputVector, FCh
 
 	FCachedMovementBaseAsyncData& MovementBaseData = AsyncInput.MovementBaseAsyncData;
 	
-	UPrimitiveComponent* MovementBase = GetMovementBase();
+	UPrimitiveComponent* MovementBase = Cast<UPrimitiveComponent>(GetMovementBaseObject());
 	MovementBaseData.CachedMovementBase = MovementBase;
 	MovementBaseData.OldBaseQuat = OldBaseQuat;
 	MovementBaseData.OldBaseLocation = OldBaseLocation;
@@ -13497,10 +13671,13 @@ void UCharacterMovementComponent::FillAsyncInput(const FVector& InputVector, FCh
 	{
 		MovementBaseData.bMovementBaseIsValidCached = true;
 		MovementBaseData.bMovementBaseOwnerIsValidCached = true;
-		MovementBaseData.bMovementBaseUsesRelativeLocationCached = MovementBaseUtility::UseRelativeLocation(MovementBase);
-		MovementBaseData.bMovementBaseIsSimulatedCached = MovementBaseUtility::IsSimulatedBase(MovementBase);
-		MovementBaseData.bMovementBaseIsDynamicCached = MovementBaseUtility::IsDynamicBase(MovementBase);
-		MovementBaseData.bIsBaseTransformValid = MovementBaseUtility::GetMovementBaseTransform(MovementBase, CharacterOwner->GetBasedMovement().BoneName, MovementBaseData.BaseLocation, MovementBaseData.BaseQuat);
+		
+		// Todo: eventually switch async code to use new FMovementBaseInterfaceData struct properly rather than converting here
+		const FMovementBaseInterfaceData MovementBaseInterfaceData(MovementBase);
+		MovementBaseData.bMovementBaseUsesRelativeLocationCached = MovementBaseUtility::UseRelativeLocation(&MovementBaseInterfaceData);
+		MovementBaseData.bMovementBaseIsSimulatedCached = MovementBaseUtility::IsSimulatedBase(&MovementBaseInterfaceData);
+		MovementBaseData.bMovementBaseIsDynamicCached = MovementBaseUtility::IsDynamicBase(&MovementBaseInterfaceData);
+		MovementBaseData.bIsBaseTransformValid = MovementBaseUtility::GetMovementBaseTransform(&MovementBaseInterfaceData, CharacterOwner->GetBasedMovement().BoneName, MovementBaseData.BaseLocation, MovementBaseData.BaseQuat);
 	}
 	else
 	{
diff --git a/Engine/Source/Runtime/Engine/Private/GameViewportClient.cpp b/Engine/Source/Runtime/Engine/Private/GameViewportClient.cpp
--- a/Engine/Source/Runtime/Engine/Private/GameViewportClient.cpp
+++ b/Engine/Source/Runtime/Engine/Private/GameViewportClient.cpp
@@ -3719,9 +3719,9 @@ void UGameViewportClient::ToggleShowCollision()
 			if (PC != NULL && PC->GetPawn() != NULL)
 			{
 				PC->ClientMessage(FString::Printf(TEXT("!!!! Player Pawn %s Collision Info !!!!"), *PC->GetPawn()->GetName()));
-				if (PC->GetPawn()->GetMovementBase())
+				if (const UObject* MovementBaseObject = PC->GetPawn()->GetMovementBaseObject())
 				{
-					PC->ClientMessage(FString::Printf(TEXT("Base %s"), *PC->GetPawn()->GetMovementBase()->GetName()));
+					PC->ClientMessage(FString::Printf(TEXT("Base %s"), *MovementBaseObject->GetName()));
 				}
 				TSet<AActor*> TouchingActors;
 				PC->GetPawn()->GetOverlappingActors(TouchingActors);
diff --git a/Engine/Source/Runtime/Engine/Private/Pawn.cpp b/Engine/Source/Runtime/Engine/Private/Pawn.cpp
--- a/Engine/Source/Runtime/Engine/Private/Pawn.cpp
+++ b/Engine/Source/Runtime/Engine/Private/Pawn.cpp
@@ -22,6 +22,7 @@
 #include "DisplayDebugHelpers.h"
 #include "Engine/InputDelegateBinding.h"
 #include "Interfaces/NetworkPredictionInterface.h"
+#include "Interfaces/MovementBaseInterface.h"
 #include "GameFramework/PlayerState.h"
 #include "Components/PawnNoiseEmitterComponent.h"
 #include "GameFramework/GameNetworkManager.h"
@@ -222,8 +223,18 @@ void APawn::PawnStartFire(uint8 FireModeNum) {}
 
 AActor* APawn::GetMovementBaseActor(const APawn* Pawn)
 {
-	UPrimitiveComponent* MovementBase = (Pawn ? Pawn->GetMovementBase() : nullptr);
-	return (MovementBase ? MovementBase->GetOwner() : nullptr);
+	return Cast<AActor>(GetMovementBaseOwner(Pawn));
+}
+
+UObject* APawn::GetMovementBaseOwner(const APawn* Pawn)
+{
+	const FMovementBaseInterfaceData* MovementBaseInterfaceData = (Pawn ? Pawn->GetMovementBaseInterfaceData() : nullptr);
+	if (MovementBaseInterfaceData && MovementBaseInterfaceData->IsValid())
+	{
+		return MovementBaseInterfaceData->GetMovementBaseObjectOwner();
+	}
+	
+	return nullptr;
 }
 
 bool APawn::CanBeBaseForCharacter(class APawn* APawn) const
@@ -1236,12 +1247,14 @@ void APawn::PostNetReceiveLocationAndRotation()
 
 bool APawn::IsBasedOnActor(const AActor* Other) const
 {
-	UPrimitiveComponent * MovementBase = GetMovementBase();
-	AActor* MovementBaseActor = MovementBase ? MovementBase->GetOwner() : nullptr;
-
-	if (MovementBaseActor && MovementBaseActor == Other)
+	const FMovementBaseInterfaceData* MovementBaseInterfaceData = GetMovementBaseInterfaceData();
+	if (MovementBaseInterfaceData && MovementBaseInterfaceData->IsValid())
 	{
-		return true;
+		const AActor* MovementBaseActor = Cast<AActor>(MovementBaseInterfaceData->GetBodyInstanceOwner()->GetSourceObjectOwner());
+		if (MovementBaseActor && MovementBaseActor == Other)
+		{
+			return true;
+		}
 	}
 
 	return Super::IsBasedOnActor(Other);
@@ -1262,11 +1275,14 @@ bool APawn::IsNetRelevantFor(const AActor* RealViewer, const AActor* ViewTarget,
 	}
 	else
 	{
-		UPrimitiveComponent* MovementBase = GetMovementBase();
-		AActor* BaseActor = MovementBase ? MovementBase->GetOwner() : nullptr;
-		if ( MovementBase && BaseActor && GetMovementComponent() && ((Cast<const USkeletalMeshComponent>(MovementBase)) || (BaseActor == GetOwner())) )
+		const FMovementBaseInterfaceData* MovementBaseInterfaceData = GetMovementBaseInterfaceData();
+		if (MovementBaseInterfaceData && MovementBaseInterfaceData->IsValid())
 		{
-			return BaseActor->IsNetRelevantFor(RealViewer, ViewTarget, SrcLocation);
+			AActor* BaseActor = Cast<AActor>(MovementBaseInterfaceData->GetMovementBaseObjectOwner());
+			if (BaseActor && GetMovementComponent() && ((Cast<const USkeletalMeshComponent>(MovementBaseInterfaceData->GetMovementBaseObject())) || (BaseActor == GetOwner())) )
+			{
+				return BaseActor->IsNetRelevantFor(RealViewer, ViewTarget, SrcLocation);
+			}
 		}
 	}
 
diff --git a/Engine/Source/Runtime/Engine/Private/PrimitiveComponentPhysics.cpp b/Engine/Source/Runtime/Engine/Private/PrimitiveComponentPhysics.cpp
--- a/Engine/Source/Runtime/Engine/Private/PrimitiveComponentPhysics.cpp
+++ b/Engine/Source/Runtime/Engine/Private/PrimitiveComponentPhysics.cpp
@@ -1621,6 +1621,16 @@ Chaos::FPhysicsObjectId UPrimitiveComponent::GetIdFromGTParticle(Chaos::FGeometr
 	return 0;
 }
 
+IPhysicsBodyInstanceOwner* UPrimitiveComponent::ResolvePhysicsBodyInstanceOwner(Chaos::FConstPhysicsObjectHandle PhysicsObject)
+{
+	return this;
+}
+
+UObject* UPrimitiveComponent::GetSourceObjectOwner() const
+{
+	return GetOwner();
+}
+
 UPhysicalMaterial* UPrimitiveComponent::GetPhysicsMaterialOverride() const
 {
 	return BodyInstance.GetPhysMaterialOverride();
@@ -1631,4 +1641,47 @@ UMaterialInterface* UPrimitiveComponent::GetPhysicsMaterialBase() const
 	return IsA<UMeshComponent>() ? GetMaterial(0) : nullptr;
 }
 
+bool UPrimitiveComponent::IsPhysicsOwnerMovable() const
+{
+	return Mobility == EComponentMobility::Movable;
+}
+
+bool UPrimitiveComponent::IsPhysicsOwnerSimulatingPhysics() const
+{
+	bool bBaseIsSimulatingPhysics = false;
+	const USceneComponent* AttachParentComponent = this;
+	while (!bBaseIsSimulatingPhysics && AttachParentComponent)
+	{
+		bBaseIsSimulatingPhysics = AttachParentComponent->IsSimulatingPhysics();
+		AttachParentComponent = AttachParentComponent->GetAttachParent();
+	}
+	
+	return bBaseIsSimulatingPhysics;
+}
+
+FVector UPrimitiveComponent::GetPhysicsOwnerVelocity() const
+{
+	return GetComponentVelocity();
+}
+
+UObject* UPrimitiveComponent::GetPhysicsOwnerAttachmentRoot() const
+{
+	return GetAttachmentRoot();
+}
+
+bool UPrimitiveComponent::IsPhysicsObjectWorldGeometry() const
+{
+	return IsWorldGeometry();
+}
+
+FTransform UPrimitiveComponent::GetPhysicsOwnerSocketTransform(FName InSocketName) const
+{
+	return GetSocketTransform(InSocketName);
+}
+
+bool UPrimitiveComponent::DoesSocketExistOnPhysicsOwner(FName InSocketName) const
+{
+	return DoesSocketExist(InSocketName);
+}
+
 #undef LOCTEXT_NAMESPACE
diff --git a/Engine/Source/Runtime/Engine/Public/CharacterMovementComponentAsync.h b/Engine/Source/Runtime/Engine/Public/CharacterMovementComponentAsync.h
--- a/Engine/Source/Runtime/Engine/Public/CharacterMovementComponentAsync.h
+++ b/Engine/Source/Runtime/Engine/Public/CharacterMovementComponentAsync.h
@@ -242,6 +242,7 @@ public:
 // Don't read into this part too much it needs to be changed.
 struct FCachedMovementBaseAsyncData
 {
+	// TODO: Convert this to use new FMovementBaseInterfaceData struct
 	// data derived from movement base
 	UPrimitiveComponent* CachedMovementBase; // Do not access, this was input movement base, only here so I could ensure when it changed.
 	
