VulkanRHI: Fix source/target component counts in SpecularProfile shaders.  Fix validation errors in VulkanRHI.

#rnx
#Synced-CL 51670748

diff --git a/Engine/Shaders/Private/SpecularProfile.usf b/Engine/Shaders/Private/SpecularProfile.usf
--- a/Engine/Shaders/Private/SpecularProfile.usf
+++ b/Engine/Shaders/Private/SpecularProfile.usf
@@ -6,6 +6,18 @@
 
 ///////////////////////////////////////////////////////////////////////////////////////////////////
 
+#if defined(SHADER_COPY) || defined(SHADER_CONVOLVE)
+	#if PERMUTATION_COMPONENT_COUNT == 3
+		#define SPECULAR_PROFILE_TYPE float3
+		#define SPECULAR_PROFILE_WRITE(RGBValue) RGBValue
+	#elif PERMUTATION_COMPONENT_COUNT == 4
+		#define SPECULAR_PROFILE_TYPE float4
+		#define SPECULAR_PROFILE_WRITE(RGBValue) float4(RGBValue,1)
+	#else
+		#error Only component counts of 3 or 4 are currently supported
+	#endif
+#endif
+
 #ifdef SHADER_COPY
 
 int2 SourceResolution;
@@ -13,8 +25,8 @@ uint SourceMipCount;
 int2 TargetResolution;
 uint TargetIndex;
 SamplerState SourceSampler;
-Texture2D<float4> SourceTexture;
-RWTexture2DArray<float4> TargetTexture;
+Texture2D<SPECULAR_PROFILE_TYPE> SourceTexture;
+RWTexture2DArray<SPECULAR_PROFILE_TYPE> TargetTexture;
 
 Buffer<float4> ViewColorBuffer;
 Buffer<float4> LightColorBuffer;
@@ -30,13 +42,14 @@ void MainCS(uint3 DispatchThreadId : SV_DispatchThreadID)
 		{
 			const float3 C0 = ViewColorBuffer[PixelPos.x].xyz;
 			const float3 C1 = LightColorBuffer[PixelPos.y].xyz;
-			TargetTexture[uint3(PixelPos,TargetIndex * SUBSTRATE_SPECULAR_PROFILE_ENTRY_COUNT + SUBSTRATE_SPECULAR_PROFILE_ENTRY_LIGHT)] = float4(C0 * C1, 1);
+			const float3 C0C1 = C0 * C1;
+			TargetTexture[uint3(PixelPos,TargetIndex * SUBSTRATE_SPECULAR_PROFILE_ENTRY_COUNT + SUBSTRATE_SPECULAR_PROFILE_ENTRY_LIGHT)] = SPECULAR_PROFILE_WRITE(C0C1);
 		}
 		#else
 		{
 			const float SourceMip = clamp(log2(SourceResolution.x) - log2(TargetResolution.x), 0.f, SourceMipCount-1);
-			const float4 Value = SourceTexture.SampleLevel(SourceSampler, TexCoord, SourceMip);
-			TargetTexture[uint3(PixelPos,TargetIndex * SUBSTRATE_SPECULAR_PROFILE_ENTRY_COUNT + SUBSTRATE_SPECULAR_PROFILE_ENTRY_LIGHT)] = Value;
+			const float3 Value = SourceTexture.SampleLevel(SourceSampler, TexCoord, SourceMip).xyz;
+			TargetTexture[uint3(PixelPos,TargetIndex * SUBSTRATE_SPECULAR_PROFILE_ENTRY_COUNT + SUBSTRATE_SPECULAR_PROFILE_ENTRY_LIGHT)] = SPECULAR_PROFILE_WRITE(Value);
 		}
 		#endif
 	}
@@ -54,9 +67,9 @@ uint SampleCount;
 int2 Resolution;
 uint ProfileCount;
 #if PERMUTATION_READ_WRITE_ON_ARRAY == 0
-Texture2DArray<float4> SourceTexture;
+Texture2DArray<SPECULAR_PROFILE_TYPE> SourceTexture;
 #endif
-RWTexture2DArray<float4> TargetTexture;
+RWTexture2DArray<SPECULAR_PROFILE_TYPE> TargetTexture;
 StructuredBuffer<float> ProfileSignBuffer;
 
 float3 EvaluateSpecularProfile(float NoV, float NoL, float VoH, float NoH, float SpecularProfileId)
@@ -123,7 +136,7 @@ void MainCS(uint3 DispatchThreadId : SV_DispatchThreadID)
 		const float3 ConvolvedColor = Convolve(V, N, Roughness, SpecularProfileId);
 
 		const uint3 OutCoord = uint3(PixelPos, ProfileIndex * SUBSTRATE_SPECULAR_PROFILE_ENTRY_COUNT + SUBSTRATE_SPECULAR_PROFILE_ENTRY_ENV);
-		TargetTexture[OutCoord] = float4(ConvolvedColor, 1);
+		TargetTexture[OutCoord] = SPECULAR_PROFILE_WRITE(ConvolvedColor);
 	}
 }
 #endif
diff --git a/Engine/Source/Runtime/Engine/Private/Rendering/SpecularProfile.cpp b/Engine/Source/Runtime/Engine/Private/Rendering/SpecularProfile.cpp
--- a/Engine/Source/Runtime/Engine/Private/Rendering/SpecularProfile.cpp
+++ b/Engine/Source/Runtime/Engine/Private/Rendering/SpecularProfile.cpp
@@ -60,7 +60,8 @@ public:
 	SHADER_USE_PARAMETER_STRUCT(FSpecularProfileCopyCS, FGlobalShader);
 
 	class FProcedural : SHADER_PERMUTATION_BOOL("PERMUTATION_PROCEDURAL");
-	using FPermutationDomain = TShaderPermutationDomain<FProcedural>;
+	class FComponentCount : SHADER_PERMUTATION_SPARSE_INT("PERMUTATION_COMPONENT_COUNT", 3, 4);
+	using FPermutationDomain = TShaderPermutationDomain<FProcedural, FComponentCount>;
 
 	BEGIN_SHADER_PARAMETER_STRUCT(FParameters, )
 		SHADER_PARAMETER_TEXTURE(Texture2D, SourceTexture)
@@ -96,7 +97,8 @@ public:
 	SHADER_USE_PARAMETER_STRUCT(FSpecularProfileConvolveCS, FGlobalShader);
 
 	class FReadWriteOnArray : SHADER_PERMUTATION_BOOL("PERMUTATION_READ_WRITE_ON_ARRAY");
-	using FPermutationDomain = TShaderPermutationDomain<FReadWriteOnArray>;
+	class FComponentCount : SHADER_PERMUTATION_SPARSE_INT("PERMUTATION_COMPONENT_COUNT", 3, 4);
+	using FPermutationDomain = TShaderPermutationDomain<FReadWriteOnArray, FComponentCount>;
 
 	BEGIN_SHADER_PARAMETER_STRUCT(FParameters, )
 		SHADER_PARAMETER(uint32, SampleCount)
@@ -109,7 +111,12 @@ public:
 
 	static bool ShouldCompilePermutation(const FGlobalShaderPermutationParameters& Parameters)
 	{
-		return IsSpecularProfileSupport(Parameters.Platform);
+		if (IsSpecularProfileSupport(Parameters.Platform))
+		{
+			const FPermutationDomain PermutationVector(Parameters.PermutationId);
+			return !IsOpenGLPlatform(Parameters.Platform) || !PermutationVector.Get<FReadWriteOnArray>();
+		}
+		return false;
 	}
 	static void ModifyCompilationEnvironment(const FGlobalShaderPermutationParameters& Parameters, FShaderCompilerEnvironment& OutEnvironment)
 	{
@@ -375,6 +382,8 @@ IPooledRenderTarget* FSpecularProfileTextureManager::GetAtlasTexture(FRDGBuilder
 			(bR11G11B10TypedLoadSupported ? PF_FloatR11G11B10 : PF_FloatRGBA) :
 			(bR11G11B10Supported ? PF_FloatR11G11B10 : PF_FloatRGBA);
 
+		const int32 ComponentCount = GPixelFormats[Format].NumComponents;
+
 		FPooledRenderTargetDesc Desc(FPooledRenderTargetDesc::Create2DArrayDesc(FIntPoint(Resolution, Resolution), Format, FClearValueBinding::None, TexCreate_None, TexCreate_UAV | TexCreate_ShaderResource, false, LayerCount * SUBSTRATE_SPECULAR_PROFILE_ENTRY_COUNT));
 		GRenderTargetPool.FindFreeElement(RHICmdList, Desc, GSpecularProfileTextureAtlas, TEXT("SpecularProfileTexture"));
 		const auto GlobalShaderMap = GetGlobalShaderMap(ShaderPlatform);
@@ -432,6 +441,7 @@ IPooledRenderTarget* FSpecularProfileTextureManager::GetAtlasTexture(FRDGBuilder
 
 			FSpecularProfileCopyCS::FPermutationDomain PermutationVector;
 			PermutationVector.Set<FSpecularProfileCopyCS::FProcedural>(Data.IsProcedural());
+			PermutationVector.Set<FSpecularProfileCopyCS::FComponentCount>(ComponentCount);
 			TShaderMapRef<FSpecularProfileCopyCS> Shader(GlobalShaderMap, PermutationVector);
 			FComputeShaderUtils::AddPass(GraphBuilder, 
 				RDG_EVENT_NAME("SpecularProfile::CopyTexture"), 
@@ -472,6 +482,7 @@ IPooledRenderTarget* FSpecularProfileTextureManager::GetAtlasTexture(FRDGBuilder
 
 			FSpecularProfileConvolveCS::FPermutationDomain PermutationVector;
 			PermutationVector.Set<FSpecularProfileConvolveCS::FReadWriteOnArray>(PassParameters->SourceTexture == nullptr);
+			PermutationVector.Set<FSpecularProfileConvolveCS::FComponentCount>(ComponentCount);
 			TShaderMapRef<FSpecularProfileConvolveCS> Shader(GlobalShaderMap, PermutationVector);
 			FComputeShaderUtils::AddPass(GraphBuilder, 
 				RDG_EVENT_NAME("SpecularProfile::Convolve"), 
