[Nanite] Fix asserts that can happen with Nanite static meshes when the locally-built Nanite resources has no streaming pages, but the one in the DDC does.

#rnx
#Synced-CL 51639297

diff --git a/Engine/Source/Runtime/Engine/Private/StaticMesh.cpp b/Engine/Source/Runtime/Engine/Private/StaticMesh.cpp
--- a/Engine/Source/Runtime/Engine/Private/StaticMesh.cpp
+++ b/Engine/Source/Runtime/Engine/Private/StaticMesh.cpp
@@ -4463,9 +4463,9 @@ void FStaticMeshRenderData::Cache(const ITargetPlatform* TargetPlatform, UStatic
 								EstimatedMeshDataCompressedSize = uint64(CompressedBuffer.GetCompressedSize() * DDCSizeToEstimateFactor);
 							}
 
+							NaniteStreamingDataHash = ResponseNaniteStreamingValue.GetRawHash();
 							if (ResponseNaniteStreamingValue != PutRequest.Record.GetValue(NaniteStreamingDataId))
 							{
-								NaniteStreamingDataHash = ResponseNaniteStreamingValue.GetRawHash();
 								bInvalidateNaniteStreaming = true;
 							}
 						}
@@ -4499,6 +4499,18 @@ void FStaticMeshRenderData::Cache(const ITargetPlatform* TargetPlatform, UStatic
 
 						FMemoryReaderView ReadAr(MeshDataBuffer.GetView(), /*bIsPersistent=*/ true);
 						Serialize(ReadAr, Owner, /*bCooked=*/ false);
+
+						// Always reset the DDC hashes in case the presence of streaming data changed in re-serialization
+						if (NaniteResources.HasStreamingData())
+						{
+							NaniteResources.DDCKeyHash = CacheKey.Hash;
+							NaniteResources.DDCRawHash = NaniteStreamingDataHash;
+						}
+						else
+						{
+							NaniteResources.DDCKeyHash.Reset();
+							NaniteResources.DDCRawHash.Reset();
+						}
 					}
 				}
 
