Fixed streaming cell not marked client-only when all DataLayers are client-only except the EDL.
- A world partition streaming cell was not properly marked as bClientOnlyVisible when its data layer list contained an External Data Layer (EDL) alongside client-only data layers.
- Because the cell wasn't flagged client-only, GetCellEffectiveWantedState returned false on the server (where client-only DataLayers don't exist) but true on the client.
- This caused the client to wait for a server acknowledgment to make the level visible, which would never come.
- The existing check called IsClientOnly() on all entries including the EDL, whose result could incorrectly prevent the cell from being flagged as client-only. The fix skips the EDL's IsClientOnly() in mixed lists, but still rejects if the EDL is server-only.
- Fixed AreDataLayersLoadFilterValid validation to include the EDL in the load filter mismatch check when its filter is not None (previously the EDL was always excluded from validation).
- Extracted shared logic into FDataLayerUtils::IsClientOnlyFromDataLayers().

#Synced-CL 51570755

diff --git a/Engine/Source/Runtime/Engine/Private/WorldPartition/DataLayer/DataLayerUtils.cpp b/Engine/Source/Runtime/Engine/Private/WorldPartition/DataLayer/DataLayerUtils.cpp
--- a/Engine/Source/Runtime/Engine/Private/WorldPartition/DataLayer/DataLayerUtils.cpp
+++ b/Engine/Source/Runtime/Engine/Private/WorldPartition/DataLayer/DataLayerUtils.cpp
@@ -8,6 +8,7 @@
 #include "WorldPartition/DataLayer/WorldDataLayers.h"
 #include "WorldPartition/DataLayer/WorldDataLayersActorDesc.h"
 #include "WorldPartition/DataLayer/DataLayerInstanceWithAsset.h"
+#include "WorldPartition/DataLayer/ExternalDataLayerInstance.h"
 #include "WorldPartition/DataLayer/DataLayerManager.h"
 #include "Algo/AllOf.h"
 
@@ -253,6 +254,25 @@ bool FDataLayerUtils::AreDataLayerTypesCompatible(EDataLayerType ParentDataLayer
 	return true;
 }
 
+bool FDataLayerUtils::IsClientOnlyFromDataLayers(const TArray<const UDataLayerInstance*>& InCellDataLayerInstances)
+{
+	if (!InCellDataLayerInstances.Num())
+	{
+		return false;
+	}
+
+	return Algo::AllOf(InCellDataLayerInstances, [bHasMultipleDataLayers = InCellDataLayerInstances.Num() > 1](const UDataLayerInstance* DataLayerInstance)
+	{
+		// When an EDL is mixed with other data layers, skip its IsClientOnly() check
+		// as long as it is not server-only (i.e. its load filter is None or ClientOnly).
+		if (bHasMultipleDataLayers && DataLayerInstance->IsA<UExternalDataLayerInstance>())
+		{
+			return !DataLayerInstance->IsServerOnly();
+		}
+		return DataLayerInstance->IsClientOnly();
+	});
+}
+
 #undef LOCTEXT_NAMESPACE
 
 #endif
\ No newline at end of file
diff --git a/Engine/Source/Runtime/Engine/Private/WorldPartition/RuntimeHashSet/RuntimePartition.cpp b/Engine/Source/Runtime/Engine/Private/WorldPartition/RuntimeHashSet/RuntimePartition.cpp
--- a/Engine/Source/Runtime/Engine/Private/WorldPartition/RuntimeHashSet/RuntimePartition.cpp
+++ b/Engine/Source/Runtime/Engine/Private/WorldPartition/RuntimeHashSet/RuntimePartition.cpp
@@ -2,8 +2,7 @@
 
 #include "WorldPartition/RuntimeHashSet/RuntimePartition.h"
 #include "WorldPartition/RuntimeHashSet/WorldPartitionRuntimeHashSet.h"
-#include "WorldPartition/DataLayer/DataLayerInstance.h"
-#include "Algo/AllOf.h"
+#include "WorldPartition/DataLayer/DataLayerUtils.h"
 
 #include UE_INLINE_GENERATED_CPP_BY_NAME(RuntimePartition)
 
@@ -87,8 +86,7 @@ URuntimePartition::FCellDescInstance::FCellDescInstance(const FCellDesc& InCellD
 	, DataLayerInstances(InDataLayerInstances)
 	, ContentBundleID(InContentBundleID)
 {
-	const bool bAreClientOnlyDataLayers = DataLayerInstances.Num() && Algo::AllOf(DataLayerInstances, [](const UDataLayerInstance* DataLayerInstance) { return DataLayerInstance->IsClientOnly(); });
-	bClientOnlyVisible = InCellDesc.bClientOnlyVisible || bAreClientOnlyDataLayers;
+	bClientOnlyVisible = InCellDesc.bClientOnlyVisible || FDataLayerUtils::IsClientOnlyFromDataLayers(DataLayerInstances);
 }
 
 #endif
diff --git a/Engine/Source/Runtime/Engine/Private/WorldPartition/WorldPartitionRuntimeSpatialHash.cpp b/Engine/Source/Runtime/Engine/Private/WorldPartition/WorldPartitionRuntimeSpatialHash.cpp
--- a/Engine/Source/Runtime/Engine/Private/WorldPartition/WorldPartitionRuntimeSpatialHash.cpp
+++ b/Engine/Source/Runtime/Engine/Private/WorldPartition/WorldPartitionRuntimeSpatialHash.cpp
@@ -22,6 +22,7 @@
 #include "WorldPartition/ContentBundle/ContentBundleBase.h"
 #include "WorldPartition/DataLayer/DataLayerManager.h"
 #include "WorldPartition/DataLayer/DataLayerInstance.h"
+#include "WorldPartition/DataLayer/DataLayerUtils.h"
 #include "WorldPartition/DataLayer/ExternalDataLayerAsset.h"
 #include "WorldPartition/DataLayer/ExternalDataLayerInstance.h"
 #include "WorldPartition/HLOD/HLODLayer.h"
@@ -1677,14 +1678,14 @@ bool UWorldPartitionRuntimeSpatialHash::CreateStreamingGrid(const FSpatialHashRu
 					UWorldPartitionRuntimeCellDataSpatialHash* CellDataSpatialHash = CastChecked<UWorldPartitionRuntimeCellDataSpatialHash>(StreamingCell->RuntimeCellData);
 
 					const TArray<const UDataLayerInstance*>& DataLayers = GridCellDataChunk.GetDataLayers();
-					const bool bAreClientOnlyDataLayers = DataLayers.Num() && Algo::AllOf(DataLayers, [](const UDataLayerInstance* DataLayerInstance) { return DataLayerInstance->IsClientOnly(); });
+					const bool bIsClientOnlyFromDataLayers = FDataLayerUtils::IsClientOnlyFromDataLayers(DataLayers);
 					const bool bIsHLOD = RuntimeGrid.HLODLayer ? true : false;
 					const bool bBlockOnSlowStreaming = ResolveBlockOnSlowStreamingForCell(CurrentStreamingGrid.bBlockOnSlowStreaming, bIsHLOD, DataLayers);
 					StreamingCell->SetIsAlwaysLoaded(bIsCellAlwaysLoaded);
 					StreamingCell->SetIsSpatiallyLoaded(bIsCellSpatiallyLoaded);
 					StreamingCell->SetDataLayers(DataLayers);
 					StreamingCell->SetContentBundleUID(GridCellDataChunk.GetContentBundleID());
-					StreamingCell->SetClientOnlyVisible(CurrentStreamingGrid.bClientOnlyVisible || bAreClientOnlyDataLayers);
+					StreamingCell->SetClientOnlyVisible(CurrentStreamingGrid.bClientOnlyVisible || bIsClientOnlyFromDataLayers);
 					StreamingCell->SetBlockOnSlowLoading(bBlockOnSlowStreaming);
 					StreamingCell->SetIsHLOD(bIsHLOD);
 					StreamingCell->SetGuid(CellGuid);
diff --git a/Engine/Source/Runtime/Engine/Private/WorldPartition/WorldPartitionStreamingGeneration.cpp b/Engine/Source/Runtime/Engine/Private/WorldPartition/WorldPartitionStreamingGeneration.cpp
--- a/Engine/Source/Runtime/Engine/Private/WorldPartition/WorldPartitionStreamingGeneration.cpp
+++ b/Engine/Source/Runtime/Engine/Private/WorldPartition/WorldPartitionStreamingGeneration.cpp
@@ -28,6 +28,7 @@
 #include "WorldPartition/DataLayer/DataLayerManager.h"
 #include "WorldPartition/DataLayer/DataLayerUtils.h"
 #include "WorldPartition/DataLayer/ExternalDataLayerHelper.h"
+#include "WorldPartition/DataLayer/ExternalDataLayerInstance.h"
 #include "WorldPartition/DataLayer/ExternalDataLayerManager.h"
 #include "WorldPartition/DataLayer/DataLayerInstanceNames.h"
 #include "WorldPartition/DataLayer/WorldDataLayersActorDesc.h"
@@ -1480,20 +1481,27 @@ class FWorldPartitionStreamingGenerator
 				// Validate that all runtime data layers have the same DataLayerLoadFilter
 				auto AreDataLayersLoadFilterValid = [this, GetDataLayerLoadFilter](const FStreamingGenerationActorDescView& ActorDescView)
 				{
-					TArrayView<const FName> DataLayers = ActorDescView.GetRuntimeDataLayerInstanceNames().GetNonExternalDataLayers();
-					if (DataLayers.Num() > 1)
+					const FDataLayerInstanceNames& RuntimeDLNames = ActorDescView.GetRuntimeDataLayerInstanceNames();
+					TArray<const UDataLayerInstance*> RuntimeDataLayerInstances = GetRuntimeDataLayerInstances(RuntimeDLNames.ToArray());
+
+					if (RuntimeDataLayerInstances.Num() > 1)
 					{
-						TArray<const UDataLayerInstance*> RuntimeDataLayerInstances = GetRuntimeDataLayerInstances(TArray<FName>(DataLayers));
-						if (RuntimeDataLayerInstances.Num() > 1)
+						TOptional<EDataLayerLoadFilter> LoadFilter;
+						for (const UDataLayerInstance* DataLayerInstance : RuntimeDataLayerInstances)
 						{
-							EDataLayerLoadFilter LoadFilter = GetDataLayerLoadFilter(RuntimeDataLayerInstances[0]);
-							for (int32 Index = 1; Index < RuntimeDataLayerInstances.Num(); ++Index)
+							EDataLayerLoadFilter Current = GetDataLayerLoadFilter(DataLayerInstance);
+							// Skip EDL if its load filter is None (doesn't constrain filtering)
+							if (DataLayerInstance->IsA<UExternalDataLayerInstance>() && Current == EDataLayerLoadFilter::None)
 							{
-								EDataLayerLoadFilter Current = GetDataLayerLoadFilter(RuntimeDataLayerInstances[Index]);
-								if (LoadFilter != Current)
-								{
-									return false;
-								}
+								continue;
+							}
+							if (!LoadFilter.IsSet())
+							{
+								LoadFilter = Current;
+							}
+							else if (LoadFilter.GetValue() != Current)
+							{
+								return false;
 							}
 						}
 					}
diff --git a/Engine/Source/Runtime/Engine/Public/WorldPartition/DataLayer/DataLayerUtils.h b/Engine/Source/Runtime/Engine/Public/WorldPartition/DataLayer/DataLayerUtils.h
--- a/Engine/Source/Runtime/Engine/Public/WorldPartition/DataLayer/DataLayerUtils.h
+++ b/Engine/Source/Runtime/Engine/Public/WorldPartition/DataLayer/DataLayerUtils.h
@@ -71,6 +71,16 @@ public:
 	static ENGINE_API bool FindDataLayerByShortName(const UDataLayerManager* InDataLayerManager, const FString& InShortName, TSet<UDataLayerInstance*>& OutDataLayerInstances);
 
 	static ENGINE_API bool AreDataLayerTypesCompatible(EDataLayerType ParentDataLayerType, EDataLayerType ChildDataLayerType, bool bIsParentExternalDataLayer, FText* OutReason = nullptr);
+
+	/**
+	 * Returns true if a cell should be considered client-only based on its data layer instances.
+	 * When the list contains a mix of External Data Layer (EDL) and regular data layers,
+	 * the EDL is skipped as long as it is not server-only (i.e. its load filter is None or ClientOnly).
+	 * Only the regular data layers are tested via IsClientOnly().
+	 * When the list contains a single entry (EDL or regular), its IsClientOnly() value is used directly.
+	 * At most one EDL is expected in the list.
+	 */
+	static ENGINE_API bool IsClientOnlyFromDataLayers(const TArray<const UDataLayerInstance*>& InCellDataLayerInstances);
 #endif
 
 	static FString GetSanitizedDataLayerShortName(FString InShortName)
