Refactor: Removed FBodyInstance::GetPhysicsCreationInputs and updated all remaining call sites to use GetAsyncPhysicsCreationInputs.

#Synced-CL 51624647

diff --git a/Engine/Plugins/Experimental/FastGeoStreaming/Source/FastGeoStreaming/Private/FastGeoStaticMeshComponent.cpp b/Engine/Plugins/Experimental/FastGeoStreaming/Source/FastGeoStreaming/Private/FastGeoStaticMeshComponent.cpp
--- a/Engine/Plugins/Experimental/FastGeoStreaming/Source/FastGeoStreaming/Private/FastGeoStaticMeshComponent.cpp
+++ b/Engine/Plugins/Experimental/FastGeoStreaming/Source/FastGeoStreaming/Private/FastGeoStaticMeshComponent.cpp
@@ -132,12 +132,9 @@ const Nanite::FResources* FFastGeoStaticMeshComponentBase::GetNaniteResources()
 
 UBodySetup* FFastGeoStaticMeshComponentBase::GetBodySetup() const
 {
-	if (!IsInGameThread())
+	if (const FBodyInstancePhysicsCreationInputs* Cached = BodyInstance.GetAsyncPhysicsCreationInputs())
 	{
-		if (BodyInstance.GetPhysicsCreationInputs().IsValid())
-		{
-			return BodyInstance.GetPhysicsCreationInputs()->ResolvedBodySetup;
-		}
+		return Cached->ResolvedBodySetup;
 	}
 
 	return GetStaticMesh() ? GetStaticMesh()->GetBodySetup() : nullptr;
diff --git a/Engine/Source/Runtime/Engine/Classes/PhysicsEngine/BodyInstance.h b/Engine/Source/Runtime/Engine/Classes/PhysicsEngine/BodyInstance.h
--- a/Engine/Source/Runtime/Engine/Classes/PhysicsEngine/BodyInstance.h
+++ b/Engine/Source/Runtime/Engine/Classes/PhysicsEngine/BodyInstance.h
@@ -1555,17 +1555,18 @@ private:
 
 
 public:
-	// Resolved Inputs used by physics state creation (currently only used by async mode)
+	// Cache resolved Inputs used by physics state creation (currently only used by async mode)
 	ENGINE_API const FBodyInstancePhysicsCreationInputsPtr& CachePhysicsCreationInputs(UPrimitiveComponent* InOwnerComponent, IPhysicsBodyInstanceOwner* InBodyInstanceOwner = nullptr);
+	
+	// Clear cached physics creation inputs
 	void ClearCachedPhysicsCreationInputs() { PhysicsCreationInputs.Reset(); }
-	const FBodyInstancePhysicsCreationInputsPtr& GetPhysicsCreationInputs() const { return PhysicsCreationInputs; }
 
 	/**
 	 * Returns the cached physics creation inputs when called from a worker thread (not game thread, not parallel game thread).
 	 * Returns nullptr on the game thread or if no cache is available.
 	 * Use this to safely branch between cached data (off-GT) and live data (on-GT).
 	 */
-	const FBodyInstancePhysicsCreationInputs* GetAsyncPhysicsCreationInputs() const;
+	ENGINE_API const FBodyInstancePhysicsCreationInputs* GetAsyncPhysicsCreationInputs() const;
 
 private:
 
diff --git a/Engine/Source/Runtime/Engine/Private/Components/StaticMeshComponent.cpp b/Engine/Source/Runtime/Engine/Private/Components/StaticMeshComponent.cpp
--- a/Engine/Source/Runtime/Engine/Private/Components/StaticMeshComponent.cpp
+++ b/Engine/Source/Runtime/Engine/Private/Components/StaticMeshComponent.cpp
@@ -1308,11 +1308,12 @@ void UStaticMeshComponent::GetUsedTextures(TArray<UTexture*>& OutTextures, EMate
 
 UBodySetup* UStaticMeshComponent::GetBodySetup()
 {
-	if (!IsInGameThread() && IsAsyncCreatePhysicsStateRunning())
+	if (!IsInGameThread() && !IsInParallelGameThread() && IsAsyncCreatePhysicsStateRunning())
 	{
-		if (ensureAlwaysMsgf(BodyInstance.GetPhysicsCreationInputs().IsValid(), TEXT("UStaticMeshComponent::GetBodySetup called outside of game thread, expected a valid cache for %s"), *GetFullName()))
+		const FBodyInstancePhysicsCreationInputs* Cached = BodyInstance.GetAsyncPhysicsCreationInputs();
+		if (ensureAlwaysMsgf(Cached, TEXT("UStaticMeshComponent::GetBodySetup called outside of game thread, expected a valid cache for %s"), *GetFullName()))
 		{
-			return BodyInstance.GetPhysicsCreationInputs()->ResolvedBodySetup;
+			return Cached->ResolvedBodySetup;
 		}
 	}
 
diff --git a/Engine/Source/Runtime/Engine/Private/PhysicsEngine/BodyInstance.cpp b/Engine/Source/Runtime/Engine/Private/PhysicsEngine/BodyInstance.cpp
--- a/Engine/Source/Runtime/Engine/Private/PhysicsEngine/BodyInstance.cpp
+++ b/Engine/Source/Runtime/Engine/Private/PhysicsEngine/BodyInstance.cpp
@@ -3172,10 +3172,14 @@ namespace UE::Private
 
 UPhysicalMaterial* FBodyInstance::GetSimplePhysicalMaterial(const FBodyInstance* BodyInstance, TWeakObjectPtr<UPrimitiveComponent> OwnerComp, TWeakObjectPtr<UBodySetup> BodySetupPtr, IPhysicsBodyInstanceOwner* BodyInstanceOwner)
 {
-	const FBodyInstancePhysicsCreationInputsPtr& PhysicsCreationInputs = BodyInstance->GetPhysicsCreationInputs();
-	if (!IsInGameThread() && !IsInParallelGameThread() && PhysicsCreationInputs.IsValid())
+	if (!ensure(BodyInstance))
 	{
-		return PhysicsCreationInputs->ResolvedSimpleMaterial;
+		return nullptr;
+	}
+
+	if (const FBodyInstancePhysicsCreationInputs* Cached = BodyInstance->GetAsyncPhysicsCreationInputs())
+	{
+		return Cached->ResolvedSimpleMaterial;
 	}
 
 	if(!GEngine || !GEngine->DefaultPhysMaterial)
@@ -3268,18 +3272,19 @@ void FBodyInstance::GetComplexPhysicalMaterials(TArray<UPhysicalMaterial*>& OutP
 
 void FBodyInstance::GetComplexPhysicalMaterials(const FBodyInstance* BodyInstance, TWeakObjectPtr<UPrimitiveComponent> OwnerComp, TArray<UPhysicalMaterial*>& OutPhysMaterials, TArray<FPhysicalMaterialMaskParams>* OutPhysMaterialMasks, IPhysicsBodyInstanceOwner* InBodyInstanceOwner)
 {
-	if (BodyInstance)
+	if (!ensure(BodyInstance))
 	{
-		const FBodyInstancePhysicsCreationInputsPtr& PhysicsCreationInputs = BodyInstance->GetPhysicsCreationInputs();
-		if (!IsInGameThread() && !IsInParallelGameThread() && PhysicsCreationInputs.IsValid())
+		return;
+	}
+
+	if (const FBodyInstancePhysicsCreationInputs* Cached = BodyInstance->GetAsyncPhysicsCreationInputs())
+	{
+		if (OutPhysMaterialMasks)
 		{
-			if (OutPhysMaterialMasks)
-			{
-				*OutPhysMaterialMasks = PhysicsCreationInputs->ResolvedComplexMasks;
-			}
-			OutPhysMaterials = PhysicsCreationInputs->ResolvedComplexMaterials;
-			return;
+			*OutPhysMaterialMasks = Cached->ResolvedComplexMasks;
 		}
+		OutPhysMaterials = Cached->ResolvedComplexMaterials;
+		return;
 	}
 
 	if(!GEngine || !GEngine->DefaultPhysMaterial)
