[SkeletalMeshUpdater] Make FSkeletalMeshUpdateChannel thread-safe for generic worker producers

- FSkeletalMeshUpdateChannel::Create/Update/Release were restricted to GT/ParallelGT via
SKELETAL_MESH_UPDATER_CHECK_THREAD_ACCESS. The underlying operations (MPSC queue, mutex-protected
index allocator, atomics) are already thread-safe. The restriction existed because PopFromQueue
did a non-atomic MoveTemp on the OpQueue pointer, racing with concurrent producers.

- FastGeo: This enables FastGeo async render state creation for skinned meshes from generic worker threads
without requiring FTaskTagScope(ParallelGameThread). Also simplifies FastGeo DestroyMeshObject:
ReleaseResources now runs immediately from any thread, only BeginCleanup is deferred to GT via AsyncTask.

#Synced-CL 51624647

diff --git a/Engine/Plugins/Experimental/FastGeoStreaming/Source/FastGeoStreaming/Private/FastGeoSkinnedMeshComponent.cpp b/Engine/Plugins/Experimental/FastGeoStreaming/Source/FastGeoStreaming/Private/FastGeoSkinnedMeshComponent.cpp
--- a/Engine/Plugins/Experimental/FastGeoStreaming/Source/FastGeoStreaming/Private/FastGeoSkinnedMeshComponent.cpp
+++ b/Engine/Plugins/Experimental/FastGeoStreaming/Source/FastGeoStreaming/Private/FastGeoSkinnedMeshComponent.cpp
@@ -178,18 +178,20 @@ void FFastGeoSkinnedMeshComponentBase::DestroyMeshObject()
 {
 	if (MeshObject)
 	{
+		// ReleaseResources is safe from any thread (MPSC queues, ENQUEUE_RENDER_COMMAND).
+		MeshObject->ReleaseResources();
+
+		// BeginCleanup must be called from the game thread or parallel game thread.
 		if (IsInGameThread() || IsInParallelGameThread())
 		{
-			MeshObject->ReleaseResources();
 			BeginCleanup(MeshObject);
 		}
 		else
 		{
-			// MeshObject->ReleaseResources() needs to be called from the game thread
-			AsyncTask(ENamedThreads::GameThread, [MeshObject = MeshObject]()
+			FSkeletalMeshObject* MeshObjectToCleanup = MeshObject;
+			AsyncTask(ENamedThreads::GameThread, [MeshObjectToCleanup]()
 			{
-				MeshObject->ReleaseResources();
-				BeginCleanup(MeshObject);
+				BeginCleanup(MeshObjectToCleanup);
 			});
 		}
 
diff --git a/Engine/Source/Runtime/Engine/Internal/SkeletalMeshUpdater.h b/Engine/Source/Runtime/Engine/Internal/SkeletalMeshUpdater.h
--- a/Engine/Source/Runtime/Engine/Internal/SkeletalMeshUpdater.h
+++ b/Engine/Source/Runtime/Engine/Internal/SkeletalMeshUpdater.h
@@ -733,6 +733,14 @@ private:
 
 	FIndexAllocator IndexAllocator;
 	TUniquePtr<FOpQueue> OpQueue;
+
+	// Protects OpQueue pointer lifetime for concurrent producer access.
+	// Producers (Create/Update/Release) take a shared lock to ensure the queue
+	// is not swapped out mid-Enqueue. PopFromQueue takes an exclusive lock to
+	// swap the queue. The MPSC queue handles concurrent enqueues; this lock
+	// only protects the pointer, not the individual enqueue operations.
+	FRWLock OpQueueLock;
+
 	FOpStream OpStream;
 	FSlotRegistry SlotRegistry;
 	FBackend* Backend = nullptr;
diff --git a/Engine/Source/Runtime/Engine/Private/SkeletalMeshUpdater.cpp b/Engine/Source/Runtime/Engine/Private/SkeletalMeshUpdater.cpp
--- a/Engine/Source/Runtime/Engine/Private/SkeletalMeshUpdater.cpp
+++ b/Engine/Source/Runtime/Engine/Private/SkeletalMeshUpdater.cpp
@@ -9,7 +9,10 @@
 #if WITH_STATE_STREAM
 #define SKELETAL_MESH_UPDATER_CHECK_THREAD_ACCESS check(IsInParallelRenderingThread());
 #else
-#define SKELETAL_MESH_UPDATER_CHECK_THREAD_ACCESS check(IsInGameThread() || IsInParallelGameThread());
+// Producers (Create/Update/Release) and PopFromQueue are synchronized via OpQueueLock (FRWLock).
+// Producers take a shared lock; PopFromQueue takes an exclusive lock.
+// No thread restriction needed - generic workers can safely produce ops.
+#define SKELETAL_MESH_UPDATER_CHECK_THREAD_ACCESS
 #endif
 
 
@@ -216,6 +219,9 @@ FSkeletalMeshUpdateHandle FSkeletalMeshUpdateChannel::Create(FSkeletalMeshObject
 	Op.Type = FOp::EType::Add;
 	Op.Data_Add.MeshObject = MeshObject;
 
+	// Shared lock: prevents PopFromQueue from swapping the OpQueue pointer while we enqueue.
+	// Multiple producers can enqueue concurrently (MPSC queue + atomics are thread-safe).
+	FReadScopeLock ReadLock(OpQueueLock);
 	check(OpQueue);
 	OpQueue->Queue.Enqueue(Op);
 	OpQueue->NumAdds.fetch_add(1, std::memory_order_relaxed);
@@ -233,10 +239,14 @@ void FSkeletalMeshUpdateChannel::Release(FSkeletalMeshUpdateHandle&& Handle)
 	Op.HandleIndex = Handle.Index;
 	Op.Type = FOp::EType::Remove;
 
-	check(OpQueue);
-	OpQueue->Queue.Enqueue(Op);
-	OpQueue->NumRemoves.fetch_add(1, std::memory_order_relaxed);
-	OpQueue->Num.fetch_add(1, std::memory_order_relaxed);
+	{
+		// Shared lock: prevents PopFromQueue from swapping the OpQueue pointer while we enqueue.
+		FReadScopeLock ReadLock(OpQueueLock);
+		check(OpQueue);
+		OpQueue->Queue.Enqueue(Op);
+		OpQueue->NumRemoves.fetch_add(1, std::memory_order_relaxed);
+		OpQueue->Num.fetch_add(1, std::memory_order_relaxed);
+	}
 
 	IndexAllocator.Free(Handle.Index);
 
@@ -255,6 +265,8 @@ void FSkeletalMeshUpdateChannel::Update(const FSkeletalMeshUpdateHandle& Handle,
 	Op.Type = FOp::EType::Update;
 	Op.Data_Update.MeshDynamicData = MeshDynamicData;
 
+	// Shared lock: prevents PopFromQueue from swapping the OpQueue pointer while we enqueue.
+	FReadScopeLock ReadLock(OpQueueLock);
 	check(OpQueue);
 	OpQueue->Queue.Enqueue(Op);
 	OpQueue->NumUpdates.fetch_add(1, std::memory_order_relaxed);
@@ -263,13 +275,20 @@ void FSkeletalMeshUpdateChannel::Update(const FSkeletalMeshUpdateHandle& Handle,
 
 void FSkeletalMeshUpdateChannel::Shutdown()
 {
+	// Exclusive lock: waits for all producers to finish before nulling the queue.
+	FWriteScopeLock WriteLock(OpQueueLock);
 	OpQueue = nullptr;
 }
 
 TUniquePtr<FSkeletalMeshUpdateChannel::FOpQueue> FSkeletalMeshUpdateChannel::PopFromQueue()
 {
 	SKELETAL_MESH_UPDATER_CHECK_THREAD_ACCESS
-	if (OpQueue->Num.load(std::memory_order_relaxed) == 0)
+
+	// Exclusive lock: waits for all producers to finish enqueuing before swapping the queue.
+	// Without this, a producer could be mid-Enqueue on the old queue when it's passed to
+	// the render thread for Close(), losing the op.
+	FWriteScopeLock WriteLock(OpQueueLock);
+	if (!OpQueue || OpQueue->Num.load(std::memory_order_relaxed) == 0)
 	{
 		return {};
 	}
