Removing unneeded comments related to fixing up some static analysis warnings in character movement

#Synced-CL 51689274

diff --git a/Engine/Source/Runtime/Engine/Private/Components/CharacterMovementComponent.cpp b/Engine/Source/Runtime/Engine/Private/Components/CharacterMovementComponent.cpp
--- a/Engine/Source/Runtime/Engine/Private/Components/CharacterMovementComponent.cpp
+++ b/Engine/Source/Runtime/Engine/Private/Components/CharacterMovementComponent.cpp
@@ -2525,18 +2525,18 @@ void UCharacterMovementComponent::MaybeSaveBaseLocation()
 // @todo handle lift moving up and down through encroachment
 void UCharacterMovementComponent::UpdateBasedMovement(float DeltaSeconds)
 {
-	if (!HasValidData()) // warning C6011: Dereferencing NULL pointer 'MovementBaseInterfaceData'
+	if (!HasValidData())
 	{
 		return;
 	}
 
-	const FMovementBaseInterfaceData* MovementBaseInterfaceData = GetMovementBaseInterfaceData(); // warning C6011: Dereferencing NULL pointer 'MovementBaseInterfaceData'
-	if (!MovementBaseInterfaceData || !MovementBaseInterfaceData->IsValid() || !MovementBaseUtility::UseRelativeLocation(MovementBaseInterfaceData)) // warning C6011: Dereferencing NULL pointer 'MovementBaseInterfaceData'
+	const FMovementBaseInterfaceData* MovementBaseInterfaceData = GetMovementBaseInterfaceData();
+	if (!MovementBaseInterfaceData || !MovementBaseInterfaceData->IsValid() || !MovementBaseUtility::UseRelativeLocation(MovementBaseInterfaceData))
 	{
 		return;
 	}
 
-	if (!MovementBaseUtility::IsMovementBaseDataValid(MovementBaseInterfaceData)) // warning C6011: Dereferencing NULL pointer 'MovementBaseInterfaceData'
+	if (!MovementBaseUtility::IsMovementBaseDataValid(MovementBaseInterfaceData))
 	{
 		FMovementBaseInterfaceData EmptyMovementBaseData;
 		SetBase(&EmptyMovementBaseData);
@@ -2544,35 +2544,35 @@ void UCharacterMovementComponent::UpdateBasedMovement(float DeltaSeconds)
 	}
 
 	// Ignore collision with bases during these movements.
-	TGuardValue<EMoveComponentFlags> ScopedFlagRestore(MoveComponentFlags, MoveComponentFlags | MOVECOMP_IgnoreBases); // warning C6011: Dereferencing NULL pointer 'MovementBaseInterfaceData'
+	TGuardValue<EMoveComponentFlags> ScopedFlagRestore(MoveComponentFlags, MoveComponentFlags | MOVECOMP_IgnoreBases);
 
-	FQuat DeltaQuat = FQuat::Identity; // warning C6011: Dereferencing NULL pointer 'MovementBaseInterfaceData'
-	FVector DeltaPosition = FVector::ZeroVector; // warning C6011: Dereferencing NULL pointer 'MovementBaseInterfaceData'
+	FQuat DeltaQuat = FQuat::Identity;
+	FVector DeltaPosition = FVector::ZeroVector;
 
-	FQuat NewBaseQuat; // warning C6011: Dereferencing NULL pointer 'MovementBaseInterfaceData'
-	FVector NewBaseLocation; // warning C6011: Dereferencing NULL pointer 'MovementBaseInterfaceData'
-	if (!MovementBaseUtility::GetMovementBaseTransform(MovementBaseInterfaceData, CharacterOwner->GetBasedMovement().BoneName, NewBaseLocation, NewBaseQuat)) // warning C6011: Dereferencing NULL pointer 'MovementBaseInterfaceData'
+	FQuat NewBaseQuat;
+	FVector NewBaseLocation;
+	if (!MovementBaseUtility::GetMovementBaseTransform(MovementBaseInterfaceData, CharacterOwner->GetBasedMovement().BoneName, NewBaseLocation, NewBaseQuat))
 	{
 		return;
 	}
 
 	// Find change in rotation
-	const bool bRotationChanged = !OldBaseQuat.Equals(NewBaseQuat, 1e-8f); // warning C6011: Dereferencing NULL pointer 'MovementBaseInterfaceData'
-	if (bRotationChanged) // warning C6011: Dereferencing NULL pointer 'MovementBaseInterfaceData'
+	const bool bRotationChanged = !OldBaseQuat.Equals(NewBaseQuat, 1e-8f);
+	if (bRotationChanged)
 	{
 		DeltaQuat = NewBaseQuat * OldBaseQuat.Inverse();
 	}
 
 	// only if base moved
-	if (bRotationChanged || (OldBaseLocation != NewBaseLocation)) // warning C6011: Dereferencing NULL pointer 'MovementBaseInterfaceData'
+	if (bRotationChanged || (OldBaseLocation != NewBaseLocation))
 	{
 		// Calculate new transform matrix of base actor (ignoring scale).
-		const FQuatRotationTranslationMatrix OldLocalToWorld(OldBaseQuat, OldBaseLocation); // warning C6011: Dereferencing NULL pointer 'MovementBaseInterfaceData'
-		const FQuatRotationTranslationMatrix NewLocalToWorld(NewBaseQuat, NewBaseLocation); // warning C6011: Dereferencing NULL pointer 'MovementBaseInterfaceData'
+		const FQuatRotationTranslationMatrix OldLocalToWorld(OldBaseQuat, OldBaseLocation);
+		const FQuatRotationTranslationMatrix NewLocalToWorld(NewBaseQuat, NewBaseLocation);
 
-		FQuat FinalQuat = UpdatedComponent->GetComponentQuat(); // warning C6011: Dereferencing NULL pointer 'MovementBaseInterfaceData'
+		FQuat FinalQuat = UpdatedComponent->GetComponentQuat();
 			
-		if (bRotationChanged && !bIgnoreBaseRotation) // warning C6011: Dereferencing NULL pointer 'MovementBaseInterfaceData'
+		if (bRotationChanged && !bIgnoreBaseRotation)
 		{
 			// Apply change in rotation and pipe through FaceRotation to maintain axis restrictions
 			const FQuat PawnOldQuat = UpdatedComponent->GetComponentQuat();
@@ -2608,8 +2608,8 @@ void UCharacterMovementComponent::UpdateBasedMovement(float DeltaSeconds)
 			}
 		}
 
-		FVector NewWorldPos; // warning C6011: Dereferencing NULL pointer 'MovementBaseInterfaceData'
-		if (HasCustomGravity()) // warning C6011: Dereferencing NULL pointer 'MovementBaseInterfaceData'
+		FVector NewWorldPos;
+		if (HasCustomGravity())
 		{
 			const FVector RotationRadius = UpdatedComponent->GetComponentLocation() - NewBaseLocation;
 			const FVector RotationDelta = DeltaQuat.RotateVector(RotationRadius) - RotationRadius;
@@ -2619,18 +2619,18 @@ void UCharacterMovementComponent::UpdateBasedMovement(float DeltaSeconds)
 		else
 		{
 			// We need to offset the base of the character here, not its origin, so offset by half height
-			float HalfHeight, Radius; // warning C6011: Dereferencing NULL pointer 'MovementBaseInterfaceData'
-			CharacterOwner->GetCapsuleComponent()->GetScaledCapsuleSize(Radius, HalfHeight); // warning C6011: Dereferencing NULL pointer 'MovementBaseInterfaceData'
+			float HalfHeight, Radius;
+			CharacterOwner->GetCapsuleComponent()->GetScaledCapsuleSize(Radius, HalfHeight);
 
-			FVector const BaseOffset = FVector(0.0f, 0.0f, HalfHeight); // warning C6011: Dereferencing NULL pointer 'MovementBaseInterfaceData'
-			FVector const LocalBasePos = OldLocalToWorld.InverseTransformPosition(UpdatedComponent->GetComponentLocation() - BaseOffset); // warning C6011: Dereferencing NULL pointer 'MovementBaseInterfaceData'
-			NewWorldPos = ConstrainLocationToPlane(NewLocalToWorld.TransformPosition(LocalBasePos) + BaseOffset); // warning C6011: Dereferencing NULL pointer 'MovementBaseInterfaceData'
+			FVector const BaseOffset = FVector(0.0f, 0.0f, HalfHeight);
+			FVector const LocalBasePos = OldLocalToWorld.InverseTransformPosition(UpdatedComponent->GetComponentLocation() - BaseOffset);
+			NewWorldPos = ConstrainLocationToPlane(NewLocalToWorld.TransformPosition(LocalBasePos) + BaseOffset);
 		}
 		
-		DeltaPosition = ConstrainDirectionToPlane(NewWorldPos - UpdatedComponent->GetComponentLocation()); // warning C6011: Dereferencing NULL pointer 'MovementBaseInterfaceData'
+		DeltaPosition = ConstrainDirectionToPlane(NewWorldPos - UpdatedComponent->GetComponentLocation());
 
 		// move attached actor
-		if (bFastAttachedMove) // warning C6011: Dereferencing NULL pointer 'MovementBaseInterfaceData'
+		if (bFastAttachedMove)
 		{
 			// we're trusting no other obstacle can prevent the move here
 			UpdatedComponent->SetWorldLocationAndRotation(NewWorldPos, FinalQuat, false);
@@ -2639,9 +2639,9 @@ void UCharacterMovementComponent::UpdateBasedMovement(float DeltaSeconds)
 		{
 			// Set MovementBase's root actor as ignored when moving the character primitive component, 
 			// only perform if bDeferUpdateBasedMovement is true since this means the MovementBase is simulating physics
-			const bool bIgnoreBaseActor = bDeferUpdateBasedMovement && bBasedMovementIgnorePhysicsBase; // warning C6011: Dereferencing NULL pointer 'MovementBaseInterfaceData'
-			AActor* MovementBaseRootActor = nullptr; // warning C6011: Dereferencing NULL pointer 'MovementBaseInterfaceData'
-			if (bIgnoreBaseActor) // warning C6011: Dereferencing NULL pointer 'MovementBaseInterfaceData'
+			const bool bIgnoreBaseActor = bDeferUpdateBasedMovement && bBasedMovementIgnorePhysicsBase;
+			AActor* MovementBaseRootActor = nullptr;
+			if (bIgnoreBaseActor)
 			{
 				if (MovementBaseInterfaceData && MovementBaseInterfaceData->IsValid())
 				{
@@ -2652,31 +2652,31 @@ void UCharacterMovementComponent::UpdateBasedMovement(float DeltaSeconds)
 			}
 
 			// hack - transforms between local and world space introducing slight error FIXMESTEVE - discuss with engine team: just skip the transforms if no rotation?
-			FVector BaseMoveDelta = NewBaseLocation - OldBaseLocation; // warning C6011: Dereferencing NULL pointer 'MovementBaseInterfaceData'
-			if (!bRotationChanged && (BaseMoveDelta.X == 0.f) && (BaseMoveDelta.Y == 0.f)) // warning C6011: Dereferencing NULL pointer 'MovementBaseInterfaceData'
+			FVector BaseMoveDelta = NewBaseLocation - OldBaseLocation;
+			if (!bRotationChanged && (BaseMoveDelta.X == 0.f) && (BaseMoveDelta.Y == 0.f))
 			{
 				DeltaPosition.X = 0.f;
 				DeltaPosition.Y = 0.f;
 			}
 
-			FHitResult MoveOnBaseHit(1.f); // warning C6011: Dereferencing NULL pointer 'MovementBaseInterfaceData'
-			const FVector OldLocation = UpdatedComponent->GetComponentLocation(); // warning C6011: Dereferencing NULL pointer 'MovementBaseInterfaceData'
-			MoveUpdatedComponent(DeltaPosition, FinalQuat, true, &MoveOnBaseHit); // warning C6011: Dereferencing NULL pointer 'MovementBaseInterfaceData'
+			FHitResult MoveOnBaseHit(1.f);
+			const FVector OldLocation = UpdatedComponent->GetComponentLocation();
+			MoveUpdatedComponent(DeltaPosition, FinalQuat, true, &MoveOnBaseHit);
 
-			if ((UpdatedComponent->GetComponentLocation() - (OldLocation + DeltaPosition)).IsNearlyZero() == false) // warning C6011: Dereferencing NULL pointer 'MovementBaseInterfaceData'
+			if ((UpdatedComponent->GetComponentLocation() - (OldLocation + DeltaPosition)).IsNearlyZero() == false)
 			{
 				OnUnableToFollowBaseMove(DeltaPosition, OldLocation, MoveOnBaseHit);
 			}
 
 			// Reset base actor ignore state
-			if (bIgnoreBaseActor) // warning C6011: Dereferencing NULL pointer 'MovementBaseInterfaceData'
+			if (bIgnoreBaseActor)
 			{
 				MoveComponentFlags &= ~MOVECOMP_CheckBlockingRootActorInIgnoreList;
 				UpdatedPrimitive->IgnoreActorWhenMoving(MovementBaseRootActor, false);
 			}
 		}
 
-		if (MovementBaseInterfaceData && MovementBaseInterfaceData->IsValid()) // warning C6011: Dereferencing NULL pointer 'MovementBaseInterfaceData'
+		if (MovementBaseInterfaceData && MovementBaseInterfaceData->IsValid())
 		{
 			const FBodyInstance* MovementBaseBodyInstance = MovementBaseInterfaceData->GetBodyInstanceOwner()->GetBodyInstance();
 			if (MovementBaseBodyInstance && MovementBaseBodyInstance->IsInstanceSimulatingPhysics() && CharacterOwner->GetMesh())
@@ -2687,14 +2687,14 @@ void UCharacterMovementComponent::UpdateBasedMovement(float DeltaSeconds)
 	}
 
 	// Check if falling above current base
-	if (IsFalling() && bStayBasedInAir) // warning C6011: Dereferencing NULL pointer 'MovementBaseInterfaceData'
+	if (IsFalling() && bStayBasedInAir)
 	{
-		const FVector PawnLocation = UpdatedComponent->GetComponentLocation(); // warning C6011: Dereferencing NULL pointer 'MovementBaseInterfaceData'
-		FFindFloorResult OutFloorResult; // warning C6011: Dereferencing NULL pointer 'MovementBaseInterfaceData'
-		ComputeFloorDist(PawnLocation, StayBasedInAirHeight, StayBasedInAirHeight, OutFloorResult, CharacterOwner->GetCapsuleComponent()->GetScaledCapsuleRadius(), NULL); // warning C6011: Dereferencing NULL pointer 'MovementBaseInterfaceData'
+		const FVector PawnLocation = UpdatedComponent->GetComponentLocation();
+		FFindFloorResult OutFloorResult;
+		ComputeFloorDist(PawnLocation, StayBasedInAirHeight, StayBasedInAirHeight, OutFloorResult, CharacterOwner->GetCapsuleComponent()->GetScaledCapsuleRadius(), NULL);
 
-		FMovementBaseInterfaceData FoundMovementBaseData = MovementBaseUtility::GetMovementBaseDataFromHitResult(&OutFloorResult.HitResult); // warning C6011: Dereferencing NULL pointer 'MovementBaseInterfaceData'
-		if (!FoundMovementBaseData.IsValid() || FoundMovementBaseData.GetBodyInstanceOwner()->GetPhysicsOwnerAttachmentRoot() != MovementBaseInterfaceData->GetBodyInstanceOwner()->GetPhysicsOwnerAttachmentRoot()) // warning C6011: Dereferencing NULL pointer 'MovementBaseInterfaceData'
+		FMovementBaseInterfaceData FoundMovementBaseData = MovementBaseUtility::GetMovementBaseDataFromHitResult(&OutFloorResult.HitResult);
+		if (!FoundMovementBaseData.IsValid() || FoundMovementBaseData.GetBodyInstanceOwner()->GetPhysicsOwnerAttachmentRoot() != MovementBaseInterfaceData->GetBodyInstanceOwner()->GetPhysicsOwnerAttachmentRoot())
 		{
 			// New or no base under the character
 			ApplyImpartedMovementBaseVelocity();
