[HWRT] Support "additional geometries" in ray tracing geometry groups / fix crash with geometry collections using cached instances.

Context:
- The ray tracing geometry streaming/eviction -> invalidation of primitive's cached ray tracing state flow depends on ray tracing geometries being assigned to a group and the primitives registering with the manager using the corresponding group handle.
- Geometry collections support cached ray tracing instances but didn't use ray tracing geometry groups so didn't get invalidation requests when geometry was streamed in or evicted.
- Geometry collections are also a special case where we have multiple ray tracing geometries (one BLAS per part) per proxy and each don't currently support LODs. This didn't really map the the existing ray tracing geometry group usage.

Changes:
- added support for "additional geometries" in ray tracing geometry groups which are handled outside the existing LOD chain.
- this avoids the need to have one group per collection part which would add extra overhead and require broader changes.
- also added some checks in FPrimitiveSceneInfo::CacheRayTracingPrimitives(...) to detect missing/incorrect group assignment for cached ray tracing instances.

#Synced-CL 51690135

diff --git a/Engine/Source/Runtime/Engine/Private/Rendering/RayTracingGeometryManager.cpp b/Engine/Source/Runtime/Engine/Private/Rendering/RayTracingGeometryManager.cpp
--- a/Engine/Source/Runtime/Engine/Private/Rendering/RayTracingGeometryManager.cpp
+++ b/Engine/Source/Runtime/Engine/Private/Rendering/RayTracingGeometryManager.cpp
@@ -188,6 +188,11 @@ DECLARE_DWORD_ACCUMULATOR_STAT(TEXT("In-flight Streaming Requests"), STAT_RayTra
 
 CSV_DEFINE_CATEGORY(RayTracingGeometry, true);
 
+static bool IsAdditionalRayTracingGeometry(const FRayTracingGeometry* InGeometry)
+{
+	return InGeometry->LODIndex == FRayTracingGeometry::AdditionalGeometryLODIndex;
+}
+
 struct FReferenceBasedResidencyState
 {
 	struct FGeometryMetadata
@@ -205,6 +210,7 @@ struct FReferenceBasedResidencyState
 	struct FGeometryGroupMetadata
 	{
 		TArray<RayTracing::FGeometryHandle, TInlineAllocator<2>> GeometryHandles;
+		TSet<RayTracing::FGeometryHandle> AdditionalGeometryHandles;
 		int32 CurrentFirstLODIdx = INDEX_NONE;
 	};
 
@@ -345,6 +351,30 @@ void FReferenceBasedResidencyState::Update()
 				}
 			}
 		}
+
+		for (RayTracing::FGeometryHandle GeometryHandle : GroupMetadata.AdditionalGeometryHandles)
+		{
+			FGeometryMetadata& GeometryMetadata = GetGeometryMetadata(GeometryHandle);
+			GeometryMetadata.LastReferencedFrame = FrameCounter;
+
+			ReferencedSize += GeometryMetadata.Size;
+
+			checkf(!GeometryMetadata.bAlwaysResident, TEXT("Additional geometries can't be always-resident."));
+
+			if (GeometryMetadata.bEvicted || GeometryMetadata.bBVHStreamedOut || (GeometryMetadata.bBuffersStreamedOut && ReferencedGeometryGroupsForDynamicUpdate.Contains(GeometryMetadata.GroupHandle)))
+			{
+				GeometriesToMakeResidentAndStreamIn.Add(GeometryHandle);
+			}
+
+			NotReferencedResidentGeometries.Remove(GeometryHandle);
+
+			RequestedSize += GeometryMetadata.Size;
+
+			if (GeometryMetadata.bEvicted)
+			{
+				RequestedButEvictedSize += GeometryMetadata.Size;
+			}
+		}
 	}
 
 #if DO_CHECK
@@ -564,7 +594,7 @@ void FRayTracingGeometryManager::RemoveBuildRequest(RayTracing::FBuildRequestHan
 
 bool FRayTracingGeometryManager::IsAlwaysResidentGeometry(const FRayTracingGeometry* InGeometry, const FRayTracingGeometryGroup& Group)
 {
-	return InGeometry->LODIndex >= Group.GeometryHandles.Num() - GRayTracingNumAlwaysResidentLODs;
+	return !IsAdditionalRayTracingGeometry(InGeometry) && InGeometry->LODIndex >= Group.GeometryHandles.Num() - GRayTracingNumAlwaysResidentLODs;
 }
 
 RayTracing::FGeometryGroupHandle FRayTracingGeometryManager::RegisterRayTracingGeometryGroup(uint32 NumLODs, uint32 CurrentFirstLODIdx)
@@ -611,6 +641,8 @@ void FRayTracingGeometryManager::ReleaseRayTracingGeometryGroupReference(RayTrac
 			checkf(GeometryHandle.IsNull(), TEXT("All FRayTracingGeometry in a group must be unregistered before releasing the group."));
 		}
 
+		checkf(Group.AdditionalGeometryHandles.IsEmpty(), TEXT("All Additional FRayTracingGeometry in a group must be unregistered before releasing the group."));
+
 		check(Group.ProxiesWithCachedRayTracingState.IsEmpty());
 
 		const uint32 GroupIndex = GetHandleIndex(Handle);
@@ -655,10 +687,21 @@ RayTracing::FGeometryHandle FRayTracingGeometryManager::RegisterRayTracingGeomet
 	{
 		FRayTracingGeometryGroup& Group = GetRegisteredGroup(GroupHandle);
 
-		checkf(LODIndex >= 0 && LODIndex < Group.GeometryHandles.Num(), TEXT("FRayTracingGeometry assigned to a group must have a valid LODIndex"));
-		checkf(Group.GeometryHandles[LODIndex].IsNull(), TEXT("Each LOD inside a FRayTracingGeometryGroup can only be associated with a single FRayTracingGeometry"));
+		const bool bIsAdditionalRayTracingGeometry = IsAdditionalRayTracingGeometry(InGeometry);
+
+		if (!bIsAdditionalRayTracingGeometry)
+		{
+			checkf(LODIndex >= 0 && LODIndex < Group.GeometryHandles.Num(), TEXT("FRayTracingGeometry assigned to a group must have a valid LODIndex"));
+			checkf(Group.GeometryHandles[LODIndex].IsNull(), TEXT("Each LOD inside a FRayTracingGeometryGroup can only be associated with a single FRayTracingGeometry"));
+
+			Group.GeometryHandles[LODIndex] = Handle;
+		}
+		else
+		{
+			check(!Group.AdditionalGeometryHandles.Contains(Handle));
+			Group.AdditionalGeometryHandles.Add(Handle);
+		}
 
-		Group.GeometryHandles[LODIndex] = Handle;
 		++Group.NumReferences;
 
 		ModifiedGroups.Add(GroupHandle);
@@ -670,8 +713,9 @@ RayTracing::FGeometryHandle FRayTracingGeometryManager::RegisterRayTracingGeomet
 			AlwaysResidentGeometries.Add(Handle);
 		}
 
-		if (IsRayTracingEnabled() && LODIndex >= Group.CurrentFirstLODIdx && (!IsRayTracingUsingReferenceBasedResidency() || RegisteredGeometry.bAlwaysResident))
+		if (IsRayTracingEnabled() && (bIsAdditionalRayTracingGeometry || LODIndex >= Group.CurrentFirstLODIdx) && (!IsRayTracingUsingReferenceBasedResidency() || RegisteredGeometry.bAlwaysResident))
 		{
+			// TODO: Should also do this if the geometry is not assigned to a group?
 			PendingStreamingRequests.Add(Handle);
 			INC_DWORD_STAT(STAT_RayTracingPendingStreamingRequests);
 		}
@@ -699,10 +743,18 @@ void FRayTracingGeometryManager::ReleaseRayTracingGeometryHandle(RayTracing::FGe
 
 		FRayTracingGeometryGroup& Group = GetRegisteredGroup(RegisteredGeometry.Geometry->GroupHandle);
 
-		checkf(RegisteredGeometry.Geometry->LODIndex >= 0 && RegisteredGeometry.Geometry->LODIndex < Group.GeometryHandles.Num(), TEXT("FRayTracingGeometry assigned to a group must have a valid LODIndex"));
-		checkf(Group.GeometryHandles[RegisteredGeometry.Geometry->LODIndex] == Handle, TEXT("Unexpected mismatch of FRayTracingGeometry in FRayTracingGeometryGroup"));
+		if (!IsAdditionalRayTracingGeometry(RegisteredGeometry.Geometry))
+		{
+			checkf(RegisteredGeometry.Geometry->LODIndex >= 0 && RegisteredGeometry.Geometry->LODIndex < Group.GeometryHandles.Num(), TEXT("FRayTracingGeometry assigned to a group must have a valid LODIndex"));
+			checkf(Group.GeometryHandles[RegisteredGeometry.Geometry->LODIndex] == Handle, TEXT("Unexpected mismatch of FRayTracingGeometry in FRayTracingGeometryGroup"));
 
-		Group.GeometryHandles[RegisteredGeometry.Geometry->LODIndex] = {};
+			Group.GeometryHandles[RegisteredGeometry.Geometry->LODIndex] = {};
+		}
+		else
+		{
+			check(Group.AdditionalGeometryHandles.Contains(Handle));
+			Group.AdditionalGeometryHandles.Remove(Handle);
+		}
 
 		ReleaseRayTracingGeometryGroupReference(RegisteredGeometry.Geometry->GroupHandle);
 	}
@@ -1133,6 +1185,7 @@ void FRayTracingGeometryManager::UpdateReferenceBasedResidencyState()
 
 		FReferenceBasedResidencyState::FGeometryGroupMetadata& Metadata = ReferenceBasedResidencyState->GroupMetadatas[GroupIndex];
 		Metadata.GeometryHandles = RegisteredGroup.GeometryHandles;
+		Metadata.AdditionalGeometryHandles = RegisteredGroup.AdditionalGeometryHandles;
 		Metadata.CurrentFirstLODIdx = RegisteredGroup.CurrentFirstLODIdx;
 
 	}
@@ -1456,7 +1509,7 @@ void FRayTracingGeometryManager::DoesRayTracingGeometryRequireStreaming(const FR
 	{
 		const FRayTracingGeometryGroup& Group = GetRegisteredGroup(RegisteredGeometry.Geometry->GroupHandle);
 
-		if (RegisteredGeometry.Geometry->LODIndex < Group.CurrentFirstLODIdx)
+		if (!IsAdditionalRayTracingGeometry(RegisteredGeometry.Geometry) && RegisteredGeometry.Geometry->LODIndex < Group.CurrentFirstLODIdx)
 		{
 			bOutStreamBVH = false;
 			bOutStreamBuffers = false;
@@ -1583,7 +1636,7 @@ void FRayTracingGeometryManager::ProcessCompletedStreamingRequests(FRHICommandLi
 
 			const FRayTracingGeometryGroup& Group = GetRegisteredGroup(RegisteredGeometry.Geometry->GroupHandle);
 
-			if (RegisteredGeometry.Geometry->IsEvicted() || RegisteredGeometry.Geometry->LODIndex < Group.CurrentFirstLODIdx)
+			if (RegisteredGeometry.Geometry->IsEvicted() || (!IsAdditionalRayTracingGeometry(RegisteredGeometry.Geometry) && RegisteredGeometry.Geometry->LODIndex < Group.CurrentFirstLODIdx))
 			{
 				// do nothing since geometry was evicted while streaming request was being processed
 			}
diff --git a/Engine/Source/Runtime/Engine/Public/Rendering/RayTracingGeometryManager.h b/Engine/Source/Runtime/Engine/Public/Rendering/RayTracingGeometryManager.h
--- a/Engine/Source/Runtime/Engine/Public/Rendering/RayTracingGeometryManager.h
+++ b/Engine/Source/Runtime/Engine/Public/Rendering/RayTracingGeometryManager.h
@@ -101,6 +101,7 @@ private:
 	struct FRayTracingGeometryGroup
 	{
 		TArray<RayTracing::FGeometryHandle, TInlineAllocator<2>> GeometryHandles;
+		TSet<RayTracing::FGeometryHandle> AdditionalGeometryHandles;
 
 		TSet<FPrimitiveSceneProxy*> ProxiesWithCachedRayTracingState;
 
diff --git a/Engine/Source/Runtime/Experimental/GeometryCollectionEngine/Private/GeometryCollection/GeometryCollectionRenderData.cpp b/Engine/Source/Runtime/Experimental/GeometryCollectionEngine/Private/GeometryCollection/GeometryCollectionRenderData.cpp
--- a/Engine/Source/Runtime/Experimental/GeometryCollectionEngine/Private/GeometryCollection/GeometryCollectionRenderData.cpp
+++ b/Engine/Source/Runtime/Experimental/GeometryCollectionEngine/Private/GeometryCollection/GeometryCollectionRenderData.cpp
@@ -7,6 +7,8 @@
 #include "GeometryCollection/GeometryCollectionObject.h"
 #include "GeometryCollection/ManagedArray.h"
 
+#include "Rendering/RayTracingGeometryManager.h"
+
 #if WITH_EDITOR
 #include "NaniteBuilder.h"
 #endif
@@ -801,6 +803,11 @@ void FGeometryCollectionRenderData::InitResources(UGeometryCollection const& Own
 			ENQUEUE_RENDER_COMMAND(InitGeometryCollectiomRayTracingGeometry)(
 				[this, OwnerName](FRHICommandListImmediate& RHICmdList)
 				{
+					// Create a single ray tracing geometry group to be used by parts
+
+					check(RayTracingGeometryGroupHandle.IsNull());
+					RayTracingGeometryGroupHandle = GRayTracingGeometryManager->RegisterRayTracingGeometryGroup(/*NumLODs*/ 0, /*CurrentFirstLODIdx*/ 0);
+
 					// TODO: Combine sections using the same transform into a single BLAS
 
 					for (int32 SectionIndex = 0; SectionIndex < MeshDescription.SubSections.Num(); ++SectionIndex)
@@ -827,6 +834,8 @@ void FGeometryCollectionRenderData::InitResources(UGeometryCollection const& Own
 							Initializer.TotalPrimitiveCount += Section.TriangleCount;
 						}
 
+						PartRayTracingGeometries[SectionIndex].GroupHandle = RayTracingGeometryGroupHandle;
+						PartRayTracingGeometries[SectionIndex].LODIndex = FRayTracingGeometry::AdditionalGeometryLODIndex;
 						PartRayTracingGeometries[SectionIndex].SetInitializer(MoveTemp(Initializer));
 						PartRayTracingGeometries[SectionIndex].InitResource(RHICmdList);
 					}
@@ -859,6 +868,21 @@ void FGeometryCollectionRenderData::ReleaseResources()
 		}
 		PartRayTracingGeometries.Empty();
 
+#if RHI_RAYTRACING
+		if (IsRayTracingAllowed())
+		{
+			ENQUEUE_RENDER_COMMAND(CmdReleaseRayTracingGeometryGroup)(
+				[this](FRHICommandListImmediate&)
+				{
+					if (!RayTracingGeometryGroupHandle.IsNull())
+					{
+						GRayTracingGeometryManager->ReleaseRayTracingGeometryGroup(RayTracingGeometryGroupHandle);
+						RayTracingGeometryGroupHandle = {};
+					}
+				});
+		}
+#endif
+
 		MeshResource.ReleaseResources();
 	}
 
diff --git a/Engine/Source/Runtime/Experimental/GeometryCollectionEngine/Private/GeometryCollection/GeometryCollectionRenderData.h b/Engine/Source/Runtime/Experimental/GeometryCollectionEngine/Private/GeometryCollection/GeometryCollectionRenderData.h
--- a/Engine/Source/Runtime/Experimental/GeometryCollectionEngine/Private/GeometryCollection/GeometryCollectionRenderData.h
+++ b/Engine/Source/Runtime/Experimental/GeometryCollectionEngine/Private/GeometryCollection/GeometryCollectionRenderData.h
@@ -157,6 +157,9 @@ public:
 	/** Object local bounds before any animation. */
 	FBoxSphereBounds PreSkinnedBounds = FBoxSphereBounds(ForceInitToZero);
 
+#if RHI_RAYTRACING
+	RayTracing::FGeometryGroupHandle RayTracingGeometryGroupHandle;
+#endif
 	TArray<FRayTracingGeometry> PartRayTracingGeometries;
 
 private:
diff --git a/Engine/Source/Runtime/Experimental/GeometryCollectionEngine/Private/GeometryCollection/GeometryCollectionSceneProxy.cpp b/Engine/Source/Runtime/Experimental/GeometryCollectionEngine/Private/GeometryCollection/GeometryCollectionSceneProxy.cpp
--- a/Engine/Source/Runtime/Experimental/GeometryCollectionEngine/Private/GeometryCollection/GeometryCollectionSceneProxy.cpp
+++ b/Engine/Source/Runtime/Experimental/GeometryCollectionEngine/Private/GeometryCollection/GeometryCollectionSceneProxy.cpp
@@ -165,7 +165,7 @@ FGeometryCollectionSceneProxyBase::FGeometryCollectionSceneProxyBase(UGeometryCo
 	, bSupportsTripleBufferVertexUpload(GRHISupportsMapWriteNoOverwrite)
 #if RHI_RAYTRACING
 	, bHasRayTracingRepresentation(GRayTracingGeometryCollection && IsRayTracingEnabled() && Component->GetRestCollection()->RenderData->PartRayTracingGeometries.Num() > 0)
-
+	, RayTracingGeometryGroupHandle(Component->GetRestCollection()->RenderData->RayTracingGeometryGroupHandle)
 #endif
 {
 	if (!bIsNanite || bHasRayTracingRepresentation)
@@ -733,11 +733,6 @@ void FGeometryCollectionSceneProxyBase::GetDynamicRayTracingInstances(FRayTracin
 
 void FGeometryCollectionSceneProxyBase::GetStaticRayTracingInstances(FStaticRayTracingInstances& OutRayTracingInstances, FRHIUniformBuffer* UniformBuffer)
 {
-	for (const FRayTracingGeometry& PartRayTracingGeometry : PartRayTracingGeometries)
-	{
-		check(PartRayTracingGeometry.IsValid() && !PartRayTracingGeometry.IsEvicted());
-	}
-
 	checkf(bHasRayTracingRepresentation, TEXT("Shouldn't try to get ray tracing instances from proxy that doesn't have a ray tracing representation."));
 
 	const uint32 LODIndex = 0;
@@ -763,6 +758,11 @@ void FGeometryCollectionSceneProxyBase::GetStaticRayTracingInstances(FStaticRayT
 
 	for (int32 SectionIndex = 0; SectionIndex < PartRayTracingGeometries.Num(); ++SectionIndex)
 	{
+		if (!PartRayTracingGeometries[SectionIndex].IsValid() || PartRayTracingGeometries[SectionIndex].IsEvicted())
+		{
+			continue;
+		}
+
 		const FGeometryCollectionMeshElement& Section = SectionArray[SectionIndex];
 
 		const int32 InstanceIndex = TransformToInstanceMapping[Section.TransformIndex];
@@ -784,6 +784,12 @@ void FGeometryCollectionSceneProxyBase::GetStaticRayTracingInstances(FStaticRayT
 	}
 }
 
+RayTracing::FGeometryGroupHandle FGeometryCollectionSceneProxyBase::GetRayTracingGeometryGroupHandle() const
+{
+	check(IsInRenderingThread() || IsInParallelRenderingThread());
+	return RayTracingGeometryGroupHandle;
+}
+
 #endif
 
 uint32 FGeometryCollectionSceneProxyBase::GetAllocatedSize() const
diff --git a/Engine/Source/Runtime/Experimental/GeometryCollectionEngine/Private/GeometryCollection/GeometryCollectionSceneProxy.h b/Engine/Source/Runtime/Experimental/GeometryCollectionEngine/Private/GeometryCollection/GeometryCollectionSceneProxy.h
--- a/Engine/Source/Runtime/Experimental/GeometryCollectionEngine/Private/GeometryCollection/GeometryCollectionSceneProxy.h
+++ b/Engine/Source/Runtime/Experimental/GeometryCollectionEngine/Private/GeometryCollection/GeometryCollectionSceneProxy.h
@@ -200,6 +200,7 @@ protected:
 	bool HasRayTracingRepresentation() const { return bHasRayTracingRepresentation; }
 	void GetDynamicRayTracingInstances(FRayTracingInstanceCollector& Collector, const FMatrix& LocalToWorld, FRHIUniformBuffer* UniformBuffer, bool bAnyMaterialHasWorldPositionOffset);
 	void GetStaticRayTracingInstances(FStaticRayTracingInstances& OutRayTracingInstances, FRHIUniformBuffer* UniformBuffer);
+	RayTracing::FGeometryGroupHandle GetRayTracingGeometryGroupHandle() const;
 #endif
 
 	uint32 GetAllocatedSize() const;
@@ -250,6 +251,8 @@ private:
 	TConstArrayView<FRayTracingGeometry> PartRayTracingGeometries;
 	TArray<FMeshBatch> CachedPartRayTracingMaterials;
 	TArray<uint32> PrimitiveInstanceIndices;
+
+	RayTracing::FGeometryGroupHandle RayTracingGeometryGroupHandle;
 #endif
 
 	bool bRenderResourcesCreated = false;
@@ -393,6 +396,7 @@ public:
 	virtual void GetDynamicRayTracingInstances(FRayTracingInstanceCollector& Collector) override { return FGeometryCollectionSceneProxyBase::GetDynamicRayTracingInstances(Collector, GetLocalToWorld(), GetUniformBuffer(), bAnyMaterialHasWorldPositionOffset); }
 	virtual ERayTracingPrimitiveFlags GetRayTracingPrimitiveFlags() override;
 	virtual void GetStaticRayTracingInstances(FStaticRayTracingInstances& OutRayTracingInstances) override { return FGeometryCollectionSceneProxyBase::GetStaticRayTracingInstances(OutRayTracingInstances, GetUniformBuffer()); }
+	virtual RayTracing::FGeometryGroupHandle GetRayTracingGeometryGroupHandle() const override { return FGeometryCollectionSceneProxyBase::GetRayTracingGeometryGroupHandle(); }
 #endif
 
 protected:
diff --git a/Engine/Source/Runtime/RenderCore/Public/RayTracingGeometry.h b/Engine/Source/Runtime/RenderCore/Public/RayTracingGeometry.h
--- a/Engine/Source/Runtime/RenderCore/Public/RayTracingGeometry.h
+++ b/Engine/Source/Runtime/RenderCore/Public/RayTracingGeometry.h
@@ -57,6 +57,12 @@ public:
 		return RayTracingGeometryRHI;
 	}
 
+	/**
+	When LODIndex is set to AdditionalGeometryLODIndex, the geometry is added to the group as an "additional geometry".
+	"Additional geometries" skips LOD-based streaming/eviction logic but still contribute to reference tracking and invalidation of cached ray tracing state in proxies.
+	*/
+	static constexpr int8 AdditionalGeometryLODIndex = INT8_MIN;
+
 	RayTracing::FGeometryGroupHandle GroupHandle;
 
 	/** LOD of the mesh associated with this ray tracing geometry object (-1 if unknown) */
diff --git a/Engine/Source/Runtime/Renderer/Private/PrimitiveSceneInfo.cpp b/Engine/Source/Runtime/Renderer/Private/PrimitiveSceneInfo.cpp
--- a/Engine/Source/Runtime/Renderer/Private/PrimitiveSceneInfo.cpp
+++ b/Engine/Source/Runtime/Renderer/Private/PrimitiveSceneInfo.cpp
@@ -1291,6 +1291,15 @@ void FPrimitiveSceneInfo::CacheRayTracingPrimitives(FScene* Scene, const TArrayV
 									FPrimitiveSceneInfo::SetupCachedRayTracingInstance(SceneInfo, RayTracingInstance, DeferredInstanceData.Instance, DeferredInstanceData.SegmentCount);
 									DeferredInstanceData.bCacheInstance = true;
 								}
+
+								if (RayTracingInstanceData.RayTracingGeometry != nullptr)
+								{
+									// Checks to make sure that the primitive gets the necessary update requests when the geometry is evicted or streamed out
+									checkf(!RayTracingInstanceData.RayTracingGeometry->GroupHandle.IsNull(),
+										TEXT("Static ray tracing instances must use a geometry assigned to the same group as returned by GetRayTracingGeometryGroupHandle()."));
+									checkf(RayTracingInstanceData.RayTracingGeometry->GroupHandle == PrimitiveRayTracingData.RayTracingGeometryGroupHandle,
+										TEXT("Static ray tracing instances must use a geometry assigned to the same group as returned by GetRayTracingGeometryGroupHandle()."));
+								}
 							}
 						}
 
