Preshader2 now supports compact VS only UniformBuffer.

#rnx
#Synced-CL 51670918

diff --git a/Engine/Source/Runtime/Engine/Private/Materials/MaterialIRModuleBuilder.cpp b/Engine/Source/Runtime/Engine/Private/Materials/MaterialIRModuleBuilder.cpp
--- a/Engine/Source/Runtime/Engine/Private/Materials/MaterialIRModuleBuilder.cpp
+++ b/Engine/Source/Runtime/Engine/Private/Materials/MaterialIRModuleBuilder.cpp
@@ -28,6 +28,7 @@
 #include "MaterialDomain.h"
 #include "MaterialExpressionIO.h"
 #include "MaterialShared.h"
+#include "ShaderCompiler.h"
 #include "Materials/MaterialExpression.h"
 #include "Materials/MaterialInsights.h"
 #include "Math/Color.h"
@@ -1208,6 +1209,20 @@ struct FMaterialIRModuleBuilderImpl
 		// UniformExpressionSet
 		CompilationOutput.UniformExpressionSet.SetParameterCollections(Module->ParameterCollections);
 
+		// Create CompactUniformsVS so the MaterialVS uniform buffer struct gets declared in the shader.
+		// VS-frequency preshader outputs were allocated first in PreshaderFixup, so the compact VS buffer
+		// can reference the contiguous prefix of the PreshaderBuffer.
+		if (UseMaterialVSUniformBuffer())
+		{
+			FCompactUniformExpressionSet* CompactUniformsVS = CompilationOutput.UniformExpressionSet.CreateCompactUniformsVS();
+			check(CompactUniformsVS);
+			CompactUniformsVS->UniformPreshaderBufferSize = Preshader.VSPreshaderFloat4s;
+		}
+		else
+		{
+			CompilationOutput.UniformExpressionSet.EmptyCompactUniformsVS();
+		}
+
 		// NumUsedUVScalars
 		CompilationOutput.NumUsedUVScalars = Module->Statistics.NumInterpolatedTexCoords * 2;
 
diff --git a/Engine/Source/Runtime/Engine/Public/MaterialShared.h b/Engine/Source/Runtime/Engine/Public/MaterialShared.h
--- a/Engine/Source/Runtime/Engine/Public/MaterialShared.h
+++ b/Engine/Source/Runtime/Engine/Public/MaterialShared.h
@@ -815,6 +815,12 @@ public:
 		CompactUniformsVSOptional.Add(FCompactUniformExpressionSet());
 		return &CompactUniformsVSOptional[0];
 	}
+
+	void EmptyCompactUniformsVS()
+	{
+		CompactUniformsVSOptional.Empty();
+	}
+
 protected:
 	union FVTPackedStackAndLayerIndex
 	{
diff --git a/Engine/Source/Runtime/Engine/Public/Materials/MaterialIRPreshader.cpp b/Engine/Source/Runtime/Engine/Public/Materials/MaterialIRPreshader.cpp
--- a/Engine/Source/Runtime/Engine/Public/Materials/MaterialIRPreshader.cpp
+++ b/Engine/Source/Runtime/Engine/Public/Materials/MaterialIRPreshader.cpp
@@ -2,6 +2,7 @@
 
 #include "MaterialIRPreshader.h"
 #include "Materials/MaterialIRModule.h"
+#include "ShaderCompiler.h"
 
 #if WITH_EDITOR
 
@@ -924,9 +925,8 @@ void FMaterialIRPreshader::PreshaderFixup(UE::Shader::FPreshaderData& PreshaderO
 	FComponentAllocator OutAllocator;
 	FComponentAllocator TempAllocator;
 
-	// Allocate all PreshaderOut values first, so we know the starting offset for PreshaderTemp.
-	// We don't have to worry about lifetime tracking for these, since they are permanent.
-	for (int32 PreshaderIndex = 1; PreshaderIndex < ValueInfos.Num(); PreshaderIndex++)
+	// Helper to allocate a single PreshaderOut value in the OutAllocator.
+	auto AllocatePreshaderOut = [&](int32 PreshaderIndex)
 	{
 		FMaterialIRPreshader::FValueInfo& ValueInfo = ValueInfos[PreshaderIndex];
 
@@ -934,7 +934,7 @@ void FMaterialIRPreshader::PreshaderFixup(UE::Shader::FPreshaderData& PreshaderO
 		MIR::FValue* Value = ValueInfo.Value;
 		if (!Value)
 		{
-			continue;
+			return;
 		}
 
 		if (Value->HasSubgraphProperties(MIR::EGraphProperties::PreshaderOut) && !ValueInfo.Alias)
@@ -946,6 +946,41 @@ void FMaterialIRPreshader::PreshaderFixup(UE::Shader::FPreshaderData& PreshaderO
 
 			Remap[PreshaderIndex] = (uint16)GlobalComponentOffset + 1;		// Add one!
 		}
+	};
+
+	// Allocate all PreshaderOut values first, so we know the starting offset for PreshaderTemp.
+	// We don't have to worry about lifetime tracking for these, since they are permanent.
+	// When UseMaterialVSUniformBuffer() is enabled, VS-frequency outputs must be allocated first
+	// so they occupy a contiguous prefix of the PreshaderBuffer that the compact MaterialVS UB can reference.
+	if (UseMaterialVSUniformBuffer())
+	{
+		// First pass: allocate VS-frequency PreshaderOut values at the beginning of the buffer.
+		for (int32 PreshaderIndex = 1; PreshaderIndex < ValueInfos.Num(); PreshaderIndex++)
+		{
+			MIR::FValue* Value = ValueInfos[PreshaderIndex].Value;
+			if (Value && Value->IsAnalyzed(MIR::Stage_Vertex))
+			{
+				AllocatePreshaderOut(PreshaderIndex);
+			}
+		}
+		VSPreshaderFloat4s = OutAllocator.AllocatedFloat4s;
+
+		// Second pass: allocate non-VS PreshaderOut values after the VS prefix.
+		for (int32 PreshaderIndex = 1; PreshaderIndex < ValueInfos.Num(); PreshaderIndex++)
+		{
+			MIR::FValue* Value = ValueInfos[PreshaderIndex].Value;
+			if (Value && !Value->IsAnalyzed(MIR::Stage_Vertex))
+			{
+				AllocatePreshaderOut(PreshaderIndex);
+			}
+		}
+	}
+	else
+	{
+		for (int32 PreshaderIndex = 1; PreshaderIndex < ValueInfos.Num(); PreshaderIndex++)
+		{
+			AllocatePreshaderOut(PreshaderIndex);
+		}
 	}
 
 	Module->GetCompilationOutput().UniformExpressionSet.AllocateFromUniformBuffer(OutAllocator.AllocatedFloat4s);
diff --git a/Engine/Source/Runtime/Engine/Public/Materials/MaterialIRPreshader.h b/Engine/Source/Runtime/Engine/Public/Materials/MaterialIRPreshader.h
--- a/Engine/Source/Runtime/Engine/Public/Materials/MaterialIRPreshader.h
+++ b/Engine/Source/Runtime/Engine/Public/Materials/MaterialIRPreshader.h
@@ -213,6 +213,9 @@ struct FMaterialIRPreshader
 
 	int32 OutputCount = 0;
 
+	// Number of float4s allocated for VS-frequency PreshaderOut values. Set by PreshaderFixup when UseMaterialVSUniformBuffer() is true.
+	uint32 VSPreshaderFloat4s = 0;
+
 	static EPreshader2Type ToPreshader2Type(MIR::EScalarKind Kind)
 	{
 		switch (Kind)
