Fix conservative parameters leading to excessive amount of shaders being compiled.

NOTE: This can be disabled using r.Material.UseShaderCompilationParameters .

- This is done by hiding the main FMaterialShaderParameters creation behind 2 specific types: FConservativeShaderParameters and FCompilationShaderParameters.
FConservativeShaderParameters mimics the old behaviour and is now used only for DDC key creation/iterative cooking while FCompilationShaderParameters has the precise parameters based on Translation output.
Shaders resulted from using FConservativeShaderParameters are always a superset of the shaders resulted from using FCompilationShaderParameters.

- The specific types are used to guard against wrong usage.
FMaterialShaderParameters is still used in methods where both types have to be accepted, like methods which are used to create a layout for DDC key creation and for finding shaders to compile. However, FMaterialShaderParameters itself cannot be created, so intent has to be known when using either of the 2 specialized types.

- Regenerate GUID for Shadermaps, as these changes will not change DDC key directly.

- Remove ShouldCache override for PreviewMaterial.
This cleans up the code by getting rid of the shader name to hardcoded string compares and also allows for the shader count stats to piggyback off the Preview Shader compilation process.
The material being edited is translated and compiled for Preview anyway, with the translation taking the longest, so we might as well enable all shaders the main material will have, so we can use compilation shader parameters to get the actual shader output instead of the conservative one.

Additionally, 3 other problems were discovered during the implementation of this, which needed to be fixed:
- Even though inline default values for root node properties which have caused problems have been hidden, the values are still serialized. This leads to the translator to determine that WPO is needed if any other value than the default is saved.
The fix is to exclude those properties from using inline default value compilation.

- Materials with Layers did not have the mask for connected properties set if Material Attributes were not used, which would lead to the conservative DDC Layout not containing any shaders. This is fixed by running the code for building that mask on materials with Layers as well.

- Dynamic Static Switches did not take into consideration unconnected properties compiled as default. Even if the returned code id would point to the same default code chunk, it would try to lerp between them, which the translator would detect as different than a constant and think that the property is used. This is fixed by checking whether A and B are the same code chunk and returning it if so.

Note:
ShaderMap Content MaterialPath has been replaced from Material->GetBaseMaterialPathName() to Material->GetAssetPath().ToString() as it just much more helpful in finding the root problem if a Material Instance is the cause.

FCompilationShaderParameters requires compilation output and does not try to get it itself from the given FMaterial, because the ShaderMap is thread dependent and the caller code has better information about what the present thread might be.

[RN] medium | Rendering Architecture

#Synced-CL 51708060

#Synced-CL 51708060

#Synced-CL 51708060

#Synced-CL 51708060

diff --git a/Engine/Source/Editor/MaterialEditor/Private/MaterialStats.cpp b/Engine/Source/Editor/MaterialEditor/Private/MaterialStats.cpp
--- a/Engine/Source/Editor/MaterialEditor/Private/MaterialStats.cpp
+++ b/Engine/Source/Editor/MaterialEditor/Private/MaterialStats.cpp
@@ -154,31 +154,48 @@ FText FShaderPlatformSettings::GetShaderCodeText(const EMaterialQualityLevel::Ty
 
 FText FShaderPlatformSettings::GetShaderCountText(const EMaterialQualityLevel::Type QualityType, const int32 InstanceIndex)
 {
-	if (PlatformData[QualityType].Instances.IsEmpty())
+	// Preview material always gets translated, so use its parameters
+	const UMaterialInterface* CurrentMaterial = Material ? static_cast<UMaterialInterface*>(Material) : static_cast<UMaterialInterface*>(MaterialInstance);
+
+	if (!CurrentMaterial)
 	{
 		return FText::GetEmpty();
 	}
 
-	check(InstanceIndex >= 0 && InstanceIndex < PlatformData[QualityType].Instances.Num());
+	const FMaterialResource* CurrentMatResource = CurrentMaterial->GetMaterialResource(PlatformShaderID, QualityType);
 
-	int TotalShadersForMaterial = 0;
+	if (!CurrentMatResource)
+	{
+		return FText::GetEmpty();
+	}
+
+	const FMaterialShaderMap* ShaderMap = CurrentMatResource->GetGameThreadShaderMap();
 
-	if (FMaterialResource* Resource = PlatformData[QualityType].Instances[InstanceIndex].MaterialResourcesStats)
+	if (!ShaderMap)
 	{
-		FPlatformTypeLayoutParameters LayoutParams;
-		LayoutParams.InitializeForPlatform(nullptr);
+		return FText::GetEmpty();
+	}
 
-		TArray<FDebugShaderTypeInfo> OutShaderInfo;
-		Resource->GetShaderTypes(LayoutParams, OutShaderInfo);
+	FPlatformTypeLayoutParameters LayoutParams;
+	LayoutParams.InitializeForPlatform(nullptr);
 
-		for (const FDebugShaderTypeInfo& ShaderInfo : OutShaderInfo)
-		{
-			TotalShadersForMaterial += ShaderInfo.ShaderTypes.Num();
+	// Dummy FMaterialResource to use Preview Material Shader Parameters without
+	// FPreviewMaterial::ShouldCache, which eliminates shaders for editor performance
+	FMaterialResource MatResShaderCount;
+	MatResShaderCount.SetMaterial(nullptr, nullptr, CurrentMatResource->GetShaderPlatform(), CurrentMatResource->GetQualityLevel());
 
-			for (const FDebugShaderPipelineInfo& PipelineInfo : ShaderInfo.Pipelines)
-			{
-				TotalShadersForMaterial += PipelineInfo.ShaderTypes.Num();
-			}
+	TArray<FDebugShaderTypeInfo> ShaderInfos;
+	const FCompilationShaderParameters MaterialParameters(CurrentMatResource, ShaderMap->GetCompilationOutput());
+	MatResShaderCount.GetShaderTypes(LayoutParams, ShaderInfos, MaterialParameters);
+
+	int TotalShadersForMaterial = 0;
+	for (const FDebugShaderTypeInfo& ShaderInfo : ShaderInfos)
+	{
+		TotalShadersForMaterial += ShaderInfo.ShaderTypes.Num();
+
+		for (const FDebugShaderPipelineInfo& PipelineInfo : ShaderInfo.Pipelines)
+		{
+			TotalShadersForMaterial += PipelineInfo.ShaderTypes.Num();
 		}
 	}
 
@@ -356,56 +373,60 @@ bool FShaderPlatformSettings::Update()
 	// if a shader compilation has been requested check if completed and extract shader names needed by code viewer combo-box
 	for (int32 QualityLevelIndex = 0; QualityLevelIndex < EMaterialQualityLevel::Num; ++QualityLevelIndex)
 	{
-		auto &Data = PlatformData[QualityLevelIndex];
+		auto& Data = PlatformData[QualityLevelIndex];
 
 		for (int32 InstanceIndex = 0; InstanceIndex < Data.Instances.Num(); ++InstanceIndex)
 		{
-			auto& Instance = Data.Instances[InstanceIndex]; 
+			auto& Instance = Data.Instances[InstanceIndex];
 
-			if (Instance.bCompilingShaders)
+			if (!Instance.bCompilingShaders)
 			{
-				FMaterialResource* Resource = Instance.MaterialResourcesStats;
+				continue;
+			}
 
-				// if compilation is complete extract the list of compiled shader names
-				const bool bCompilationFinished = Resource->IsCompilationFinished();
-				if (bCompilationFinished)
-				{
-					Instance.bCompilingShaders = false;
-					Instance.bUpdateShaderCode = true;
+			FMaterialResource* Resource = Instance.MaterialResourcesStats;
 
-					const FMaterialShaderMap* MaterialShaderMap = Resource->GetGameThreadShaderMap();
-					if (MaterialShaderMap != nullptr)
-					{
-						TMap<FShaderId, TShaderRef<FShader>> ShaderMap;
-						MaterialShaderMap->GetShaderList(ShaderMap);
+			// if compilation is complete extract the list of compiled shader names
+			const bool bCompilationFinished = Resource->IsCompilationFinished();
+			if (!bCompilationFinished)
+			{
+				continue;
+			}
 
-						Instance.ArrShaderEntries.Empty();
-						for (const auto& Entry : ShaderMap)
-						{
-							FShaderType* EntryShader = Entry.Value.GetType();
-							FVertexFactoryType* VertexFactory = Entry.Value.GetVertexFactoryType();
-							FString ShaderName = FString::Printf(TEXT("%s/%s"), VertexFactory ? VertexFactory->GetName() : TEXT("NullVF"), EntryShader->GetName());
-							Instance.ArrShaderEntries.Add(MakeShareable(new FMaterialShaderEntry{ Entry.Key, ShaderName }));
-						}
+			Instance.bCompilingShaders = false;
+			Instance.bUpdateShaderCode = true;
 
-						if (Instance.ArrShaderEntries.Num() > 0)
+			if (const FMaterialShaderMap * MaterialShaderMap = Resource->GetGameThreadShaderMap())
+			{
+				TMap<FShaderId, TShaderRef<FShader>> ShaderMap;
+				MaterialShaderMap->GetShaderList(ShaderMap);
+
+				Instance.ArrShaderEntries.Empty();
+				for (const auto& Entry : ShaderMap)
+				{
+					FShaderType* EntryShader = Entry.Value.GetType();
+					FVertexFactoryType* VertexFactory = Entry.Value.GetVertexFactoryType();
+					FString ShaderName = FString::Printf(TEXT("%s/%s"), VertexFactory ? VertexFactory->GetName() : TEXT("NullVF"), EntryShader->GetName());
+					Instance.ArrShaderEntries.Add(MakeShareable(new FMaterialShaderEntry{ Entry.Key, ShaderName }));
+				}
+
+				if (Instance.ArrShaderEntries.Num() > 0)
+				{
+					Instance.ArrShaderEntries.Sort([](const TSharedPtr<FMaterialShaderEntry>& First, const TSharedPtr<FMaterialShaderEntry>& Second)
 						{
-							Instance.ArrShaderEntries.Sort([](const TSharedPtr<FMaterialShaderEntry>& First, const TSharedPtr<FMaterialShaderEntry>& Second)
-								{
-									return First->Text < Second->Text;
-								});
+							return First->Text < Second->Text;
+						});
 
-							Instance.ComboBoxSelectedEntry = *Instance.ArrShaderEntries[0];
-						}
-					}
+					Instance.ComboBoxSelectedEntry = *Instance.ArrShaderEntries[0];
+				}
 
-					FMaterialStatsUtils::ExtractMatertialStatsInfo(PlatformShaderID, Instance.ShaderStatsInfo, Resource);
+				FMaterialStatsUtils::ExtractMatertialStatsInfo(PlatformShaderID, Instance.ShaderStatsInfo, Resource);
 
-					Instance.bNeedToWarnAboutCompilationErrors = Resource->GetCompileErrors().Num() > 0;
+				Instance.bNeedToWarnAboutCompilationErrors = Resource->GetCompileErrors().Num() > 0;
 
-					bRetValue = true;
-				}
+				bRetValue = true;
 			}
+
 		}
 	}
 
diff --git a/Engine/Source/Editor/MaterialEditor/Private/MaterialStatsCommon.cpp b/Engine/Source/Editor/MaterialEditor/Private/MaterialStatsCommon.cpp
--- a/Engine/Source/Editor/MaterialEditor/Private/MaterialStatsCommon.cpp
+++ b/Engine/Source/Editor/MaterialEditor/Private/MaterialStatsCommon.cpp
@@ -311,7 +311,7 @@ void FMaterialStatsUtils::GetRepresentativeShaderTypesAndDescriptions(TMap<FName
 	{
 		static FName TBasePassPSFNoLightMapPolicyName = TEXT("TBasePassPSFNoLightMapPolicy");
 
-		const bool bTranslucent = IsTranslucentBlendMode(TargetMaterial);
+		const bool bTranslucent = IsTranslucentBlendMode(*TargetMaterial);
 		if (TargetMaterial->GetShadingModels().IsUnlit())
 		{
 			ShaderTypeNamesAndDescriptions.FindOrAdd(FLocalVertexFactoryName)
@@ -578,115 +578,120 @@ void FMaterialStatsUtils::ExtractMatertialStatsInfo(EShaderPlatform ShaderPlatfo
 
 	bool bNoErrors = OutInfo.StrShaderErrors.Len() == 0;
 
-	if (bNoErrors)
+	const FMaterialShaderMap* MaterialShaderMap = MaterialResource->GetGameThreadShaderMap();
+	if (!bNoErrors || !MaterialShaderMap)
 	{
-		// extract instructions info
-		TArray<FMaterialStatsUtils::FShaderInstructionsInfo> ShaderInstructionInfo;
-		GetRepresentativeInstructionCounts(ShaderInstructionInfo, MaterialResource);
-
-		for (int32 InstructionIndex = 0; InstructionIndex < ShaderInstructionInfo.Num(); InstructionIndex++)
-		{
-			FShaderStatsInfo::FContent Content;
+		return;
+	}
 
-			Content.StrDescription = ShaderInstructionInfo[InstructionIndex].InstructionCount > 0 ? FString::Printf(TEXT("%u"), ShaderInstructionInfo[InstructionIndex].InstructionCount) : TEXT("n/a");
-			Content.StrDescriptionLong = ShaderInstructionInfo[InstructionIndex].InstructionCount > 0 ?
-				FString::Printf(TEXT("%s: %u instructions\nStats: %s"), *ShaderInstructionInfo[InstructionIndex].ShaderDescription, ShaderInstructionInfo[InstructionIndex].InstructionCount, *ShaderInstructionInfo[InstructionIndex].ShaderStatisticsString) :
-				TEXT("Offline shader compiler not available or an error was encountered!");
+	// extract instructions info
+	TArray<FMaterialStatsUtils::FShaderInstructionsInfo> ShaderInstructionInfo;
+	GetRepresentativeInstructionCounts(ShaderInstructionInfo, MaterialResource);
 
-			OutInfo.ShaderInstructionCount.Add(ShaderInstructionInfo[InstructionIndex].ShaderType, Content);
+	for (int32 InstructionIndex = 0; InstructionIndex < ShaderInstructionInfo.Num(); InstructionIndex++)
+	{
+		FShaderStatsInfo::FContent Content;
 
-			FString Description = ShaderInstructionInfo[InstructionIndex].ShaderStatisticsString;
-			FShaderStatsInfo::FContent GenericContent;
-			GenericContent.StrDescription = Description;
-			GenericContent.StrDescriptionLong = Description;
-			OutInfo.GenericShaderStatistics.Add(ShaderInstructionInfo[InstructionIndex].ShaderType, GenericContent);
-		}
+		Content.StrDescription = ShaderInstructionInfo[InstructionIndex].InstructionCount > 0 ? FString::Printf(TEXT("%u"), ShaderInstructionInfo[InstructionIndex].InstructionCount) : TEXT("n/a");
+		Content.StrDescriptionLong = ShaderInstructionInfo[InstructionIndex].InstructionCount > 0 ?
+			FString::Printf(TEXT("%s: %u instructions\nStats: %s"), *ShaderInstructionInfo[InstructionIndex].ShaderDescription, ShaderInstructionInfo[InstructionIndex].InstructionCount, *ShaderInstructionInfo[InstructionIndex].ShaderStatisticsString) :
+			TEXT("Offline shader compiler not available or an error was encountered!");
 
-		// extract samplers info
-		const int32 SamplersUsed = FMath::Max(MaterialResource->GetSamplerUsage(), 0);
-		const int32 MaxSamplers = GetExpectedFeatureLevelMaxTextureSamplers(MaterialResource->GetFeatureLevel());
-		OutInfo.SamplersCount.StrDescription = FString::Printf(TEXT("%u/%u"), SamplersUsed, MaxSamplers);
-		OutInfo.SamplersCount.StrDescriptionLong = FString::Printf(TEXT("%s samplers: %u/%u"), TEXT("Texture"), SamplersUsed, MaxSamplers);
+		OutInfo.ShaderInstructionCount.Add(ShaderInstructionInfo[InstructionIndex].ShaderType, Content);
 
-		// extract estimated sample info
-		uint32 NumVSTextureSamples = 0, NumPSTextureSamples = 0;
-		MaterialResource->GetEstimatedNumTextureSamples(NumVSTextureSamples, NumPSTextureSamples);
+		FString Description = ShaderInstructionInfo[InstructionIndex].ShaderStatisticsString;
+		FShaderStatsInfo::FContent GenericContent;
+		GenericContent.StrDescription = Description;
+		GenericContent.StrDescriptionLong = Description;
+		OutInfo.GenericShaderStatistics.Add(ShaderInstructionInfo[InstructionIndex].ShaderType, GenericContent);
+	}
 
-		OutInfo.TextureSampleCount.StrDescription = FString::Printf(TEXT("VS(%u), PS(%u)"), NumVSTextureSamples, NumPSTextureSamples);
-		OutInfo.TextureSampleCount.StrDescriptionLong = FString::Printf(TEXT("Texture Lookups (Est.): Vertex(%u), Pixel(%u)"), NumVSTextureSamples, NumPSTextureSamples);
+	// extract samplers info
+	const int32 SamplersUsed = FMath::Max(MaterialResource->GetSamplerUsage(), 0);
+	const int32 MaxSamplers = GetExpectedFeatureLevelMaxTextureSamplers(MaterialResource->GetFeatureLevel());
+	OutInfo.SamplersCount.StrDescription = FString::Printf(TEXT("%u/%u"), SamplersUsed, MaxSamplers);
+	OutInfo.SamplersCount.StrDescriptionLong = FString::Printf(TEXT("%s samplers: %u/%u"), TEXT("Texture"), SamplersUsed, MaxSamplers);
 
-		// extract estimated VT info
-		const uint32 NumVirtualTextureLookups = MaterialResource->GetEstimatedNumVirtualTextureLookups();
-		OutInfo.VirtualTextureLookupCount.StrDescription = FString::Printf(TEXT("%u"), NumVirtualTextureLookups);
-		OutInfo.VirtualTextureLookupCount.StrDescriptionLong = FString::Printf(TEXT("Virtual Texture Lookups (Est.): %u"), NumVirtualTextureLookups);
+	// extract estimated sample info
+	uint32 NumVSTextureSamples = 0, NumPSTextureSamples = 0;
+	MaterialResource->GetEstimatedNumTextureSamples(NumVSTextureSamples, NumPSTextureSamples);
 
-		const uint32 NumVirtualTextureStacks = MaterialResource->GetNumVirtualTextureStacks();
-		OutInfo.VirtualTextureStackCount.StrDescription = FString::Printf(TEXT("%u"), NumVirtualTextureStacks);
-		OutInfo.VirtualTextureStackCount.StrDescriptionLong = FString::Printf(TEXT("Virtual Texture Stacks (Est.): %u"), NumVirtualTextureStacks);
+	OutInfo.TextureSampleCount.StrDescription = FString::Printf(TEXT("VS(%u), PS(%u)"), NumVSTextureSamples, NumPSTextureSamples);
+	OutInfo.TextureSampleCount.StrDescriptionLong = FString::Printf(TEXT("Texture Lookups (Est.): Vertex(%u), Pixel(%u)"), NumVSTextureSamples, NumPSTextureSamples);
 
-		// extract interpolators info
-		uint32 UVScalarsUsed, CustomInterpolatorScalarsUsed;
-		MaterialResource->GetUserInterpolatorUsage(UVScalarsUsed, CustomInterpolatorScalarsUsed);
+	// extract estimated VT info
+	const uint32 NumVirtualTextureLookups = MaterialResource->GetEstimatedNumVirtualTextureLookups();
+	OutInfo.VirtualTextureLookupCount.StrDescription = FString::Printf(TEXT("%u"), NumVirtualTextureLookups);
+	OutInfo.VirtualTextureLookupCount.StrDescriptionLong = FString::Printf(TEXT("Virtual Texture Lookups (Est.): %u"), NumVirtualTextureLookups);
 
-		const uint32 TotalScalars = UVScalarsUsed + CustomInterpolatorScalarsUsed;
-		const uint32 MaxScalars = FMath::DivideAndRoundUp(TotalScalars, 4u) * 4;
+	const uint32 NumVirtualTextureStacks = MaterialResource->GetNumVirtualTextureStacks();
+	OutInfo.VirtualTextureStackCount.StrDescription = FString::Printf(TEXT("%u"), NumVirtualTextureStacks);
+	OutInfo.VirtualTextureStackCount.StrDescriptionLong = FString::Printf(TEXT("Virtual Texture Stacks (Est.): %u"), NumVirtualTextureStacks);
 
-		OutInfo.InterpolatorsCount.StrDescription = FString::Printf(TEXT("%u/%u"), TotalScalars, MaxScalars);
-		OutInfo.InterpolatorsCount.StrDescriptionLong = FString::Printf(TEXT("User interpolators: %u/%u Scalars (%u/4 Vectors) (TexCoords: %i, Custom: %i)"),
-			TotalScalars, MaxScalars, MaxScalars / 4, UVScalarsUsed, CustomInterpolatorScalarsUsed);
+	// extract interpolators info
+	uint32 UVScalarsUsed, CustomInterpolatorScalarsUsed;
+	MaterialResource->GetUserInterpolatorUsage(UVScalarsUsed, CustomInterpolatorScalarsUsed);
 
-		// Extract total shader count w/o having to compile shaders.
-		FPlatformTypeLayoutParameters LayoutParams;
-		LayoutParams.InitializeForPlatform(nullptr);
+	const uint32 TotalScalars = UVScalarsUsed + CustomInterpolatorScalarsUsed;
+	const uint32 MaxScalars = FMath::DivideAndRoundUp(TotalScalars, 4u) * 4;
 
-		TArray<FDebugShaderTypeInfo> OutShaderInfo;
-		MaterialResource->GetShaderTypes(LayoutParams, OutShaderInfo);
+	OutInfo.InterpolatorsCount.StrDescription = FString::Printf(TEXT("%u/%u"), TotalScalars, MaxScalars);
+	OutInfo.InterpolatorsCount.StrDescriptionLong = FString::Printf(TEXT("User interpolators: %u/%u Scalars (%u/4 Vectors) (TexCoords: %i, Custom: %i)"),
+		TotalScalars, MaxScalars, MaxScalars / 4, UVScalarsUsed, CustomInterpolatorScalarsUsed);
 
-		int TotalShadersForMaterial = 0;
-		for (const FDebugShaderTypeInfo& ShaderInfo : OutShaderInfo)
-		{
-			TotalShadersForMaterial += ShaderInfo.ShaderTypes.Num();
+	// Extract total shader count w/o having to compile shaders.
+	FPlatformTypeLayoutParameters LayoutParams;
+	LayoutParams.InitializeForPlatform(nullptr);
+		
 
-			for (const FDebugShaderPipelineInfo& PipelineInfo : ShaderInfo.Pipelines)
-			{
-				TotalShadersForMaterial += PipelineInfo.ShaderTypes.Num();
-			}
-		}
+	const FMaterialCompilationOutput& CompilationOutput = MaterialShaderMap->GetCompilationOutput();
+	TArray<FDebugShaderTypeInfo> OutShaderInfo;
+	MaterialResource->GetShaderTypes(LayoutParams, OutShaderInfo, FCompilationShaderParameters(MaterialResource, CompilationOutput));
 
-		OutInfo.ShaderCount.StrDescription = FString::Printf(TEXT("%u"), TotalShadersForMaterial);
-		OutInfo.ShaderCount.StrDescriptionLong = FString::Printf(TEXT("Total Shaders: %u"), TotalShadersForMaterial);
+	int TotalShadersForMaterial = 0;
+	for (const FDebugShaderTypeInfo& ShaderInfo : OutShaderInfo)
+	{
+		TotalShadersForMaterial += ShaderInfo.ShaderTypes.Num();
 
-		FString LWCMessage;
-		FMaterialResource::FLWCUsagesArray LWCFuncUsagesVS;
-		FMaterialResource::FLWCUsagesArray LWCFuncUsagesPS;
-		FMaterialResource::FLWCUsagesArray LWCFuncUsagesCS;
-		MaterialResource->GetEstimatedLWCFuncUsages(LWCFuncUsagesVS, LWCFuncUsagesPS, LWCFuncUsagesCS);
-		for (int KindIndex = 0; KindIndex < (int)ELWCFunctionKind::Max; ++KindIndex)
+		for (const FDebugShaderPipelineInfo& PipelineInfo : ShaderInfo.Pipelines)
 		{
-			int Usages = LWCFuncUsagesVS[KindIndex] + LWCFuncUsagesPS[KindIndex] + LWCFuncUsagesCS[KindIndex];
-			if (Usages > 0)
-			{
-				LWCMessage += FString::Printf(
-					TEXT("%s: %u (VS), %u (PS), %u (CS)\n"), 
-					*UEnum::GetDisplayValueAsText((ELWCFunctionKind)KindIndex).ToString(), 
-					LWCFuncUsagesVS[KindIndex],
-					LWCFuncUsagesPS[KindIndex],
-					LWCFuncUsagesCS[KindIndex]);
-			}
+			TotalShadersForMaterial += PipelineInfo.ShaderTypes.Num();
 		}
+	}
 
-		OutInfo.LWCUsage.StrDescription = LWCMessage;
-		OutInfo.LWCUsage.StrDescriptionLong = LWCMessage;
+	OutInfo.ShaderCount.StrDescription = FString::Printf(TEXT("%u"), TotalShadersForMaterial);
+	OutInfo.ShaderCount.StrDescriptionLong = FString::Printf(TEXT("Total Shaders: %u"), TotalShadersForMaterial);
 
-		if (FMaterialShaderMap* ShaderMap = MaterialResource->GetGameThreadShaderMap())
+	FString LWCMessage;
+	FMaterialResource::FLWCUsagesArray LWCFuncUsagesVS;
+	FMaterialResource::FLWCUsagesArray LWCFuncUsagesPS;
+	FMaterialResource::FLWCUsagesArray LWCFuncUsagesCS;
+	MaterialResource->GetEstimatedLWCFuncUsages(LWCFuncUsagesVS, LWCFuncUsagesPS, LWCFuncUsagesCS);
+	for (int KindIndex = 0; KindIndex < (int)ELWCFunctionKind::Max; ++KindIndex)
+	{
+		int Usages = LWCFuncUsagesVS[KindIndex] + LWCFuncUsagesPS[KindIndex] + LWCFuncUsagesCS[KindIndex];
+		if (Usages > 0)
 		{
-			// Add number of preshaders and stats
-			uint32 TotalParams, TotalOps;
-			MaterialResource->GetPreshaderStats(TotalParams, TotalOps);
-			OutInfo.PreShaderCount.StrDescription = FString::Printf(TEXT("%u outputs\n%u params\n%u ops"), ShaderMap->GetNumPreshaders(), TotalParams, TotalOps);
-			OutInfo.PreShaderCount.StrDescriptionLong = FString::Printf(TEXT("%u outputs, %u parameter fetches, %u total operations"), ShaderMap->GetNumPreshaders(), TotalParams, TotalOps);
+			LWCMessage += FString::Printf(
+				TEXT("%s: %u (VS), %u (PS), %u (CS)\n"), 
+				*UEnum::GetDisplayValueAsText((ELWCFunctionKind)KindIndex).ToString(), 
+				LWCFuncUsagesVS[KindIndex],
+				LWCFuncUsagesPS[KindIndex],
+				LWCFuncUsagesCS[KindIndex]);
 		}
 	}
+
+	OutInfo.LWCUsage.StrDescription = LWCMessage;
+	OutInfo.LWCUsage.StrDescriptionLong = LWCMessage;
+
+	if (FMaterialShaderMap* ShaderMap = MaterialResource->GetGameThreadShaderMap())
+	{
+		// Add number of preshaders and stats
+		uint32 TotalParams, TotalOps;
+		MaterialResource->GetPreshaderStats(TotalParams, TotalOps);
+		OutInfo.PreShaderCount.StrDescription = FString::Printf(TEXT("%u outputs\n%u params\n%u ops"), ShaderMap->GetNumPreshaders(), TotalParams, TotalOps);
+		OutInfo.PreShaderCount.StrDescriptionLong = FString::Printf(TEXT("%u outputs, %u parameter fetches, %u total operations"), ShaderMap->GetNumPreshaders(), TotalParams, TotalOps);
+	}
 }
 
 /*end FMaterialStatsUtils */
diff --git a/Engine/Source/Editor/UnrealEd/Private/Commandlets/DumpMaterialShaderTypes.cpp b/Engine/Source/Editor/UnrealEd/Private/Commandlets/DumpMaterialShaderTypes.cpp
--- a/Engine/Source/Editor/UnrealEd/Private/Commandlets/DumpMaterialShaderTypes.cpp
+++ b/Engine/Source/Editor/UnrealEd/Private/Commandlets/DumpMaterialShaderTypes.cpp
@@ -600,7 +600,7 @@ static void ProcessSwitchOptimizer(const ITargetPlatform* TargetPlatform, const
 				MaterialAndSizeSet& InnerValue = *InnerHashMap.FindOrAdd(ShaderMapId.GetStaticSwitchParameters(), MaterialAndSizeSet());
 
 				OutShaderInfo.Reset();
-				Resource->GetShaderTypes(LayoutParams, OutShaderInfo);
+				Resource->GetShaderTypes(LayoutParams, OutShaderInfo, FConservativeShaderParameters(Resource));
 				uint32 TotalNumShaders = GetTotalShaders(OutShaderInfo);
 				if(TotalNumShaders)
 				{
@@ -632,7 +632,7 @@ static void ProcessSwitchOptimizer(const ITargetPlatform* TargetPlatform, const
 				MaterialAndSizeSet& InnerValue = *InnerHashMap.FindOrAdd(ShaderMapId.GetStaticSwitchParameters(), MaterialAndSizeSet());
 
 				OutShaderInfo.Reset();
-				Resource->GetShaderTypes(LayoutParams, OutShaderInfo);
+				Resource->GetShaderTypes(LayoutParams, OutShaderInfo, FConservativeShaderParameters(Resource));
 				uint32 TotalNumShaders = GetTotalShaders(OutShaderInfo);
 				if (TotalNumShaders)
 				{
diff --git a/Engine/Source/Runtime/Engine/Private/Materials/HLSLMaterialTranslator.cpp b/Engine/Source/Runtime/Engine/Private/Materials/HLSLMaterialTranslator.cpp
--- a/Engine/Source/Runtime/Engine/Private/Materials/HLSLMaterialTranslator.cpp
+++ b/Engine/Source/Runtime/Engine/Private/Materials/HLSLMaterialTranslator.cpp
@@ -1282,7 +1282,7 @@ EHLSLMaterialTranslatorResult FHLSLMaterialTranslator::Translate(bool bForceDisa
 		{
 			if (bVerbose)
 			{
-				UE_LOG(LogMaterial, Display, TEXT("Pushing material '%s' translation results to DDCn."), *Material->GetMaterialInterface()->GetFullName());
+				UE_LOG(LogMaterial, Display, TEXT("Pushing material '%s' translation results to DDC."), *Material->GetMaterialInterface()->GetFullName());
 			}
 
 			PushResultsToDDC();
@@ -6443,6 +6443,11 @@ int32 FHLSLMaterialTranslator::DynamicBranch(int32 Condition, int32 A, int32 B)
 		return B;
 	}
 
+	if (A == B)
+	{
+		return A;
+	}
+
 	EMaterialValueType TypeA = GetParameterType(A);
 	EMaterialValueType TypeB = GetParameterType(B);
 	if (IsLWCType(TypeA) || IsLWCType(TypeB))
@@ -13904,7 +13909,7 @@ int32 FHLSLMaterialTranslator::SubstrateCompilePreview(int32 SubstrateDataCodeCh
 
 bool FHLSLMaterialTranslator::SubstrateSkipsOpacityEvaluation()
 {
-	return !IsTranslucentBlendMode(Material)
+	return !IsTranslucentBlendMode(*Material)
 		&& GetSubstrateMaterialExportType() == ESubstrateMaterialExport::SME_None	// We never skip when exporting information
 		&& Material->GetShadingModels().CountShadingModels() == 1
 		&& !Material->GetShadingModels().HasShadingModel(MSM_SingleLayerWater)
diff --git a/Engine/Source/Runtime/Engine/Private/Materials/Material.cpp b/Engine/Source/Runtime/Engine/Private/Materials/Material.cpp
--- a/Engine/Source/Runtime/Engine/Private/Materials/Material.cpp
+++ b/Engine/Source/Runtime/Engine/Private/Materials/Material.cpp
@@ -4864,7 +4864,7 @@ void UMaterial::GetShaderTypes(EShaderPlatform ShaderPlatform, const ITargetPlat
 
 	for (FMaterialResource* Resource : NewResourcesToCache)
 	{
-		Resource->GetShaderTypes(LayoutParams, OutShaderInfo);
+		Resource->GetShaderTypes(LayoutParams, OutShaderInfo, FConservativeShaderParameters(Resource));
 		delete Resource;
 	}
 
diff --git a/Engine/Source/Runtime/Engine/Private/Materials/MaterialCachedData.cpp b/Engine/Source/Runtime/Engine/Private/Materials/MaterialCachedData.cpp
--- a/Engine/Source/Runtime/Engine/Private/Materials/MaterialCachedData.cpp
+++ b/Engine/Source/Runtime/Engine/Private/Materials/MaterialCachedData.cpp
@@ -812,7 +812,7 @@ void FMaterialCachedExpressionData::UpdateForExpressions(const FMaterialCachedEx
 	}
 }
 
-void FMaterialCachedExpressionData::AnalyzeMaterial(UMaterial& Material)
+void FMaterialCachedExpressionData::UpdatePropertyConnections(UMaterial& Material)
 {
 	if (!Material.bUseMaterialAttributes)
 	{
@@ -826,6 +826,12 @@ void FMaterialCachedExpressionData::AnalyzeMaterial(UMaterial& Material)
 			}
 		}
 	}
+}
+
+void FMaterialCachedExpressionData::AnalyzeMaterial(UMaterial& Material)
+{
+
+	UpdatePropertyConnections(Material);
 		
 	FMaterialCachedExpressionContext Context;
 	UpdateForExpressions(Context, Material.GetExpressions(), EMaterialParameterAssociation::GlobalParameter, -1);
diff --git a/Engine/Source/Runtime/Engine/Private/Materials/MaterialInstance.cpp b/Engine/Source/Runtime/Engine/Private/Materials/MaterialInstance.cpp
--- a/Engine/Source/Runtime/Engine/Private/Materials/MaterialInstance.cpp
+++ b/Engine/Source/Runtime/Engine/Private/Materials/MaterialInstance.cpp
@@ -5793,7 +5793,7 @@ void UMaterialInstance::GetShaderTypes(EShaderPlatform Platform, const ITargetPl
 
 		for (FMaterialResource* CurrentResource : StaticPermutationMaterialResources)
 		{
-			CurrentResource->GetShaderTypes(LayoutParams, OutShaderInfo);
+			CurrentResource->GetShaderTypes(LayoutParams, OutShaderInfo, FConservativeShaderParameters(CurrentResource));
 		}
 	}
 }
diff --git a/Engine/Source/Runtime/Engine/Private/Materials/MaterialInstanceConstant.cpp b/Engine/Source/Runtime/Engine/Private/Materials/MaterialInstanceConstant.cpp
--- a/Engine/Source/Runtime/Engine/Private/Materials/MaterialInstanceConstant.cpp
+++ b/Engine/Source/Runtime/Engine/Private/Materials/MaterialInstanceConstant.cpp
@@ -232,6 +232,7 @@ void UMaterialInstanceConstant::UpdateCachedData()
 			FMaterialCachedExpressionContext Context;
 			Context.LayerOverrides = &MaterialLayers;
 			LocalCachedExpressionData.Reset(new FMaterialCachedExpressionData());
+			LocalCachedExpressionData->UpdatePropertyConnections(*BaseMaterial);
 			LocalCachedExpressionData->UpdateForExpressions(Context, BaseMaterial->GetExpressions(), GlobalParameter, INDEX_NONE);
 
 			ReferencedDefaultTextures = LocalCachedExpressionData->EditorOnlyData->EditorReferencedDefaultTextures;
diff --git a/Engine/Source/Runtime/Engine/Private/Materials/MaterialShader.cpp b/Engine/Source/Runtime/Engine/Private/Materials/MaterialShader.cpp
--- a/Engine/Source/Runtime/Engine/Private/Materials/MaterialShader.cpp
+++ b/Engine/Source/Runtime/Engine/Private/Materials/MaterialShader.cpp
@@ -113,11 +113,6 @@ bool PoolSpecialMaterialsCompileJobs();
  */
 //TMap<TRefCountPtr<FMaterialShaderMap>, TArray<FMaterial*> > FMaterialShaderMap::ShaderMapsBeingCompiled;
 
-static inline bool ShouldCacheMaterialShader(const FMaterialShaderType* ShaderType, EShaderPlatform Platform, EShaderPermutationFlags PermutationFlags, const FMaterial* Material, int32 PermutationId)
-{
-	return ShaderType->ShouldCompilePermutation(Platform, Material, PermutationId, PermutationFlags) && Material->ShouldCache(ShaderType, nullptr);
-}
-
 
 /** Converts an EMaterialShadingModel to a string description. */
 FString GetShadingModelString(EMaterialShadingModel ShadingModel)
@@ -210,7 +205,7 @@ FString GetBlendModeString(EBlendMode BlendMode)
 #if WITH_EDITOR
 /** Creates a string key for the derived data cache given a shader map id. */
 FString GetMaterialShaderMapKeyString(const FMaterialShaderMapId& ShaderMapId,
-	const FMaterialShaderParameters& ShaderParameters, EShaderPlatform Platform,
+	const FConservativeShaderParameters& ShaderParameters, EShaderPlatform Platform,
 	bool bIncludeKeyStringShaderDependencies)
 {
 	FString ShaderMapKeyString;
@@ -224,11 +219,11 @@ FString GetMaterialShaderMapKeyString(const FMaterialShaderMapId& ShaderMapId,
 }
 
 void RecordOrEmitMaterialShaderMapKey(FMaterialKeyGeneratorContext& Context,
-	const FMaterialShaderMapId& ShaderMapId, const FMaterialShaderParameters& ShaderParameters)
+	const FMaterialShaderMapId& ShaderMapId, const FConservativeShaderParameters& ShaderParameters)
 {
 	check(!Context.IsLoading());
 	RecordOrEmitMaterialShaderMapKey(Context, const_cast<FMaterialShaderMapId&>(ShaderMapId),
-		const_cast<FMaterialShaderParameters&>(ShaderParameters));
+		const_cast<FConservativeShaderParameters&>(ShaderParameters));
 }
 
 /**
@@ -236,7 +231,7 @@ void RecordOrEmitMaterialShaderMapKey(FMaterialKeyGeneratorContext& Context,
  * for creating the ShaderKey later.
  */
 void RecordOrEmitMaterialShaderMapKey(FMaterialKeyGeneratorContext& Context,
-	FMaterialShaderMapId& ShaderMapId, FMaterialShaderParameters& ShaderParameters)
+	FMaterialShaderMapId& ShaderMapId, FConservativeShaderParameters& ShaderParameters)
 {
 	TRACE_CPUPROFILER_EVENT_SCOPE(RecordOrEmitMaterialShaderMapKey);
 
@@ -557,7 +552,7 @@ void RecordMaterialDependenciesForCook(UE::Cook::FCookEventContext CookContext,
 
 		FRecordedMaterialResourceForCooking& RecordedResource = RecordedResources.Emplace_GetRef();
 		RecordedResource.ShaderPlatform = ResourceForCooking.Platform;
-		RecordedResource.ShaderParameters = FMaterialShaderParameters(&Resource);
+		RecordedResource.ShaderParameters = FConservativeShaderParameters(&Resource);
 
 		EMaterialKeyInclude IncludeFlags = EMaterialKeyInclude::All & ~EMaterialKeyInclude::Globals;
 		Resource.BuildShaderMapId(RecordedResource.ShaderMapId, TargetPlatform,
@@ -1875,7 +1870,7 @@ void FMaterialShaderMapId::SetShaderDependencies(const TArray<FShaderType*>& Sha
 static void PrepareMaterialShaderCompileJob(EShaderPlatform Platform,
 	EShaderPermutationFlags PermutationFlags,
 	const FMaterial* Material,
-	const FMaterialShaderMapId& ShaderMapId,
+	const FMaterialCompilationOutput& CompilationOutput,
 	FSharedShaderCompilerEnvironment* MaterialEnvironment,
 	const FShaderPipelineType* ShaderPipeline,
 	const FString& DebugGroupName,
@@ -1898,7 +1893,7 @@ static void PrepareMaterialShaderCompileJob(EShaderPlatform Platform,
 	Material->SetupExtraCompilationSettings(NewJob->Input.ExtraSettings);
 
 	// Allow the shader type to modify the compile environment.
-	ShaderType->SetupCompileEnvironment(Platform, Material, Key.PermutationId, PermutationFlags, ShaderEnvironment);
+	ShaderType->SetupCompileEnvironment(Platform, FCompilationShaderParameters(Material, CompilationOutput), Key.PermutationId, PermutationFlags, ShaderEnvironment);
 
 	// Compile the shader environment passed in with the shader type's source code.
 	::GlobalBeginCompileShader(
@@ -1926,7 +1921,7 @@ void FMaterialShaderType::BeginCompileShader(
 	uint32 ShaderMapJobId,
 	int32 PermutationId,
 	const FMaterial* Material,
-	const FMaterialShaderMapId& ShaderMapId,
+	const FMaterialCompilationOutput& CompilationOutput,
 	FSharedShaderCompilerEnvironment* MaterialEnvironment,
 	EShaderPlatform Platform,
 	EShaderPermutationFlags PermutationFlags,
@@ -1939,7 +1934,7 @@ void FMaterialShaderType::BeginCompileShader(
 	FShaderCompileJob* NewJob = GShaderCompilingManager->PrepareShaderCompileJob(ShaderMapJobId, FShaderCompileJobKey(this, Platform, nullptr, PermutationId), Priority);
 	if (NewJob)
 	{
-		PrepareMaterialShaderCompileJob(Platform, PermutationFlags, Material, ShaderMapId, MaterialEnvironment, nullptr, DebugGroupName, DebugDescription, DebugExtension, NewJob);
+		PrepareMaterialShaderCompileJob(Platform, PermutationFlags, Material, CompilationOutput, MaterialEnvironment, nullptr, DebugGroupName, DebugDescription, DebugExtension, NewJob);
 		NewJobs.Add(FShaderCommonCompileJobPtr(NewJob));
 	}
 }
@@ -1950,7 +1945,7 @@ void FMaterialShaderType::BeginCompileShaderPipeline(
 	EShaderPlatform Platform,
 	EShaderPermutationFlags PermutationFlags,
 	const FMaterial* Material,
-	const FMaterialShaderMapId& ShaderMapId,
+	const FMaterialCompilationOutput& CompilationOutput,
 	FSharedShaderCompilerEnvironment* MaterialEnvironment,
 	const FShaderPipelineType* ShaderPipeline,
 	TArray<FShaderCommonCompileJobPtr>& NewJobs,
@@ -1967,7 +1962,7 @@ void FMaterialShaderType::BeginCompileShaderPipeline(
 	{
 		for (FShaderCompileJob* StageJob : NewPipelineJob->StageJobs)
 		{
-			PrepareMaterialShaderCompileJob(Platform, PermutationFlags, Material, ShaderMapId, MaterialEnvironment, ShaderPipeline, DebugGroupName, DebugDescription, DebugExtension, StageJob);
+			PrepareMaterialShaderCompileJob(Platform, PermutationFlags, Material, CompilationOutput, MaterialEnvironment, ShaderPipeline, DebugGroupName, DebugDescription, DebugExtension, StageJob);
 		}
 		NewJobs.Add(FShaderCommonCompileJobPtr(NewPipelineJob));
 	}
@@ -2065,82 +2060,81 @@ void FMaterialShaderMap::LoadFromDerivedDataCache(const FMaterial* Material, con
 	InOutShaderMap = BeginLoadFromDerivedDataCache(Material, ShaderMapId, InPlatform, TargetPlatform, InOutShaderMap, OutDDCKeyDesc)->Get();
 }
 
-TSharedRef<FMaterialShaderMap::FAsyncLoadContext> FMaterialShaderMap::BeginLoadFromDerivedDataCache(const FMaterial* Material, const FMaterialShaderMapId& ShaderMapId, EShaderPlatform InPlatform, const ITargetPlatform* TargetPlatform, TRefCountPtr<FMaterialShaderMap>& InShaderMap, FString& OutDDCKeyDesc)
+struct FMaterialShaderMapAsyncLoadContext : public FMaterialShaderMap::FAsyncLoadContext
 {
-	TRACE_CPUPROFILER_EVENT_SCOPE(FMaterialShaderMap::BeginLoadFromDerivedDataCache);
-	using namespace UE::DerivedData;
+	FString								DataKey;
+	FShaderCacheLoadContext				LoadContext;
+	EShaderPlatform						Platform;
+	const FMaterial*					Material = nullptr;
+	UE::DerivedData::FRequestOwner		RequestOwner{ UE::DerivedData::EPriority::Normal };
+	TRefCountPtr<FMaterialShaderMap>	ShaderMap;
 
-	struct FMaterialShaderMapAsyncLoadContext : public FMaterialShaderMap::FAsyncLoadContext
+	bool IsReady() const override
 	{
-		FString								DataKey;
-		FShaderCacheLoadContext				LoadContext;
-		EShaderPlatform						Platform;
-		const FMaterial*					Material = nullptr;
-		FRequestOwner						RequestOwner{UE::DerivedData::EPriority::Normal};
-		TRefCountPtr<FMaterialShaderMap>	ShaderMap;
+		return RequestOwner.Poll();
+	}
 
-		bool IsReady() const override
+	TRefCountPtr<FMaterialShaderMap> Get() override
+	{
+		// Make sure the async work is complete
+		if (!RequestOwner.Poll())
 		{
-			return RequestOwner.Poll();
+			COOK_STAT(auto Timer = MaterialShaderCookStats::UsageStats.TimeAsyncWait());
+			COOK_STAT(Timer.TrackCyclesOnly());
+			RequestOwner.Wait();
 		}
 
-		TRefCountPtr<FMaterialShaderMap> Get() override
-		{
-			// Make sure the async work is complete
-			if (!RequestOwner.Poll())
-			{
-				COOK_STAT(auto Timer = MaterialShaderCookStats::UsageStats.TimeAsyncWait());
-				COOK_STAT(Timer.TrackCyclesOnly());
-				RequestOwner.Wait();
-			}
+		TRACE_CPUPROFILER_EVENT_SCOPE(FMaterialShaderMap::FinishLoadFromDerivedDataCache);
+		COOK_STAT(auto Timer = MaterialShaderCookStats::UsageStats.TimeSyncWork());
 
-			TRACE_CPUPROFILER_EVENT_SCOPE(FMaterialShaderMap::FinishLoadFromDerivedDataCache);
-			COOK_STAT(auto Timer = MaterialShaderCookStats::UsageStats.TimeSyncWork());
+		if (LoadContext.HasData())
+		{
+			TRACE_COUNTER_INCREMENT(Shaders_FMaterialShaderMapDDCHits);
 
-			if (LoadContext.HasData())
-			{
-				TRACE_COUNTER_INCREMENT(Shaders_FMaterialShaderMapDDCHits);
+			int64 BytesReceived = LoadContext.GetSerializedSize();
+			TRACE_COUNTER_ADD(Shaders_FMaterialShaderMapDDCBytesReceived, BytesReceived);
+			COOK_STAT(Timer.AddHit(BytesReceived));
 
-				int64 BytesReceived = LoadContext.GetSerializedSize();
-				TRACE_COUNTER_ADD(Shaders_FMaterialShaderMapDDCBytesReceived, BytesReceived);
-				COOK_STAT(Timer.AddHit(BytesReceived));
+			ShaderMap = new FMaterialShaderMap();
+			// Deserialize from the cached data
+			ShaderMap->Serialize(LoadContext);
 
-				ShaderMap = new FMaterialShaderMap();
-				// Deserialize from the cached data
-				ShaderMap->Serialize(LoadContext);
+			check(Material != nullptr);
 
-				check(Material != nullptr);
+			const FString InDataKey = GetMaterialShaderMapKeyString(ShaderMap->GetShaderMapId(), FConservativeShaderParameters(Material), Platform, true);
 
-				const FString InDataKey = GetMaterialShaderMapKeyString(ShaderMap->GetShaderMapId(), FMaterialShaderParameters(Material), Platform, true);
+			if (InDataKey != DataKey)
+			{
+				UE_LOG(LogMaterial, Warning, TEXT("Shader map key recomputed from DDC data: %s"), *InDataKey);
+				UE_LOG(LogMaterial, Warning, TEXT("Shader map key from request: %s"), *DataKey);
+				UE_LOG(LogMaterial, Warning, TEXT("Cached data size %llu"), BytesReceived);
+				checkf(false, TEXT("DDC key constructed from deserialized shadermap does not match request key!"));
+			}
 
-				if (InDataKey != DataKey)
-				{
-					UE_LOG(LogMaterial, Warning, TEXT("Shader map key recomputed from DDC data: %s"), *InDataKey);
-					UE_LOG(LogMaterial, Warning, TEXT("Shader map key from request: %s"), *DataKey);
-					UE_LOG(LogMaterial, Warning, TEXT("Cached data size %llu"), BytesReceived);
-					checkf(false, TEXT("DDC key constructed from deserialized shadermap does not match request key!"));
-				}
+			// Register in the global map
+			ShaderMap->Register(Platform);
 
-				// Register in the global map
-				ShaderMap->Register(Platform);
+			GShaderCompilerStats->AddDDCHit(1);
 
-				GShaderCompilerStats->AddDDCHit(1);
+			UE::DerivedData::FCacheKey DdcKey = GetMaterialShaderMapKey(DataKey);
+			UE_LOG(LogMaterial, Verbose, TEXT("Loaded shaders for %s from DDC (key hash: %s)"), *Material->GetAssetName(), *LexToString(DdcKey.Hash));
+		}
+		else
+		{
+			TRACE_COUNTER_INCREMENT(Shaders_FMaterialShaderMapDDCRequests);
+			// We should be build the data later, and we can track that the resource was built there when we push it to the DDC.
+			COOK_STAT(Timer.TrackCyclesOnly());
 
-				FCacheKey DdcKey = GetMaterialShaderMapKey(DataKey);
-				UE_LOG(LogMaterial, Verbose, TEXT("Loaded shaders for %s from DDC (key hash: %s)"), *Material->GetAssetName(), *LexToString(DdcKey.Hash));
-			}
-			else
-			{
-				TRACE_COUNTER_INCREMENT(Shaders_FMaterialShaderMapDDCRequests);
-				// We should be build the data later, and we can track that the resource was built there when we push it to the DDC.
-				COOK_STAT(Timer.TrackCyclesOnly());
+			GShaderCompilerStats->AddDDCMiss(1);
+		}
 
-				GShaderCompilerStats->AddDDCMiss(1);
-			}
+		return ShaderMap;
+	}
+};
 
-			return ShaderMap;
-		}
-	};
+TSharedRef<FMaterialShaderMap::FAsyncLoadContext> FMaterialShaderMap::BeginLoadFromDerivedDataCache(const FMaterial* Material, const FMaterialShaderMapId& ShaderMapId, EShaderPlatform InPlatform, const ITargetPlatform* TargetPlatform, TRefCountPtr<FMaterialShaderMap>& InShaderMap, FString& OutDDCKeyDesc)
+{
+	using namespace UE::DerivedData;
 
 	TSharedRef<FMaterialShaderMapAsyncLoadContext> Result = MakeShared<FMaterialShaderMapAsyncLoadContext>();
 
@@ -2160,7 +2154,7 @@ TSharedRef<FMaterialShaderMap::FAsyncLoadContext> FMaterialShaderMap::BeginLoadF
 			SCOPE_SECONDS_COUNTER(MaterialDDCTime);
 			COOK_STAT(auto Timer = MaterialShaderCookStats::UsageStats.TimeSyncWork());
 			COOK_STAT(Timer.TrackCyclesOnly());
-			Result->DataKey = GetMaterialShaderMapKeyString(ShaderMapId, FMaterialShaderParameters(Material), InPlatform);
+			Result->DataKey = GetMaterialShaderMapKeyString(ShaderMapId, FConservativeShaderParameters(Material), InPlatform);
 			FCacheKey CacheKey = GetMaterialShaderMapKey(Result->DataKey);
 			OutDDCKeyDesc = LexToString(CacheKey.Hash);
 
@@ -2226,7 +2220,7 @@ TSharedRef<FMaterialShaderMap::FAsyncLoadContext> FMaterialShaderMap::BeginLoadF
 	return Result;
 }
 
-void FMaterialShaderMap::SaveToDerivedDataCache(const FMaterialShaderParameters& ShaderParameters)
+void FMaterialShaderMap::SaveToDerivedDataCache(const FConservativeShaderParameters& ShaderParameters)
 {
 	TRACE_CPUPROFILER_EVENT_SCOPE(FMaterialShaderMap::SaveToDerivedDataCache);
 	COOK_STAT(auto Timer = MaterialShaderCookStats::UsageStats.TimeSyncWork());
@@ -2527,7 +2521,10 @@ int32 FMaterialShaderMap::SubmitCompileJobs(uint32 CompilingShaderMapId,
 
 	const EShaderPlatform ShaderPlatform = GetShaderPlatform();
 	const EShaderPermutationFlags LocalPermutationFlags = ShaderMapId.GetPermutationFlags();
-	const FMaterialShaderParameters MaterialParameters(Material);
+
+	const FMaterialCompilationOutput& CompilationOutput = GetCompilationOutput();
+	const FCompilationShaderParameters MaterialParameters(Material, CompilationOutput);
+
 	const FMaterialShaderMapLayout& Layout = AcquireMaterialShaderMapLayout(ShaderPlatform, LocalPermutationFlags, MaterialParameters);
 
 #if ALLOW_SHADERMAP_DEBUG_DATA
@@ -2568,8 +2565,8 @@ int32 FMaterialShaderMap::SubmitCompileJobs(uint32 CompilingShaderMapId,
 			}
 
 			// Verify that the shader map Id contains inputs for any shaders that will be put into this shader map
-			checkf(ShaderMapId.ContainsVertexFactoryType(MeshLayout.VertexFactoryType), TEXT("Material shader map %s is missing expected vertex factory type %s"), GetFriendlyName(), MeshLayout.VertexFactoryType->GetName());
-			checkf(ShaderMapId.ContainsShaderType(ShaderType, kUniqueShaderPermutationId), TEXT("Material shader map %s is missing expected shader type %s"), GetFriendlyName(), ShaderType->GetName());
+			checkf(ShaderMapId.ContainsVertexFactoryType(MeshLayout.VertexFactoryType), TEXT("Material shader map %s from material %s is missing expected vertex factory type %s"), GetFriendlyName(), GetMaterialPath(), MeshLayout.VertexFactoryType->GetName());
+			checkf(ShaderMapId.ContainsShaderType(ShaderType, kUniqueShaderPermutationId), TEXT("Material shader map %s from material %s is missing expected shader type %s"), GetFriendlyName(), GetMaterialPath(), ShaderType->GetName());
 
 			NumShadersPerVF++;
 			// only compile the shader if we don't already have it and it is not a pipelined one
@@ -2582,7 +2579,7 @@ int32 FMaterialShaderMap::SubmitCompileJobs(uint32 CompilingShaderMapId,
 					ShaderPlatform,
 					LocalPermutationFlags,
 					Material,
-					ShaderMapId,
+					CompilationOutput,
 					MaterialEnvironment,
 					MeshLayout.VertexFactoryType,
 					CompileJobs,
@@ -2604,8 +2601,8 @@ int32 FMaterialShaderMap::SubmitCompileJobs(uint32 CompilingShaderMapId,
 			auto& StageTypes = Pipeline->GetStages();
 
 			// Verify that the shader map Id contains inputs for any shaders that will be put into this shader map
-			check(ShaderMapId.ContainsShaderPipelineType(Pipeline));
-			check(ShaderMapId.ContainsVertexFactoryType(MeshLayout.VertexFactoryType));
+			checkf(ShaderMapId.ContainsShaderPipelineType(Pipeline), TEXT("Material shader map %s from material %s missing expected pipeline type %s"), GetFriendlyName(), GetMaterialPath(), Pipeline->GetName());
+			checkf(ShaderMapId.ContainsVertexFactoryType(MeshLayout.VertexFactoryType), TEXT("Material shader map %s from material %s missing expected Vertex Factory type %s"), GetFriendlyName(), GetMaterialPath(), MeshLayout.VertexFactoryType->GetName());
 
 			if (Pipeline->ShouldOptimizeUnusedOutputs(ShaderPlatform))
 			{
@@ -2614,7 +2611,7 @@ int32 FMaterialShaderMap::SubmitCompileJobs(uint32 CompilingShaderMapId,
 				for (auto* ShaderType : StageTypes)
 				{
 					// Verify that the shader map Id contains inputs for any shaders that will be put into this shader map
-					checkf(ShaderMapId.ContainsShaderType(ShaderType, kUniqueShaderPermutationId), TEXT("Material shader map %s missing expected shader type %s"), GetFriendlyName(), ShaderType->GetName());
+					checkf(ShaderMapId.ContainsShaderType(ShaderType, kUniqueShaderPermutationId), TEXT("Material shader map %s from material %s missing expected shader type %s"), GetFriendlyName(), GetMaterialPath(), ShaderType->GetName());
 				}
 
 				// Make a pipeline job with all the stages
@@ -2624,7 +2621,7 @@ int32 FMaterialShaderMap::SubmitCompileJobs(uint32 CompilingShaderMapId,
 					ShaderPlatform,
 					LocalPermutationFlags,
 					Material,
-					ShaderMapId,
+					CompilationOutput,
 					MaterialEnvironment,
 					MeshLayout.VertexFactoryType,
 					Pipeline,
@@ -2672,7 +2669,7 @@ int32 FMaterialShaderMap::SubmitCompileJobs(uint32 CompilingShaderMapId,
 		}
 
 		// Verify that the shader map Id contains inputs for any shaders that will be put into this shader map
-		checkf(ShaderMapId.ContainsShaderType(ShaderType, kUniqueShaderPermutationId), TEXT("Material shader map %s missing expected shader type %s"), GetFriendlyName(), ShaderType->GetName());
+		checkf(ShaderMapId.ContainsShaderType(ShaderType, kUniqueShaderPermutationId), TEXT("Material shader map %s from Material %s missing expected shader type %s"), GetFriendlyName(), GetMaterialPath(), ShaderType->GetName());
 
 		// Compile this material shader for this material.
 		TArray<FString> ShaderErrors;
@@ -2684,7 +2681,7 @@ int32 FMaterialShaderMap::SubmitCompileJobs(uint32 CompilingShaderMapId,
 				CompilingShaderMapId,
 				Shader.PermutationId,
 				Material,
-				ShaderMapId,
+				CompilationOutput,
 				MaterialEnvironment,
 				ShaderPlatform,
 				LocalPermutationFlags,
@@ -2719,7 +2716,7 @@ int32 FMaterialShaderMap::SubmitCompileJobs(uint32 CompilingShaderMapId,
 					ShaderPlatform,
 					LocalPermutationFlags,
 					Material,
-					ShaderMapId,
+					CompilationOutput,
 					MaterialEnvironment,
 					Pipeline,
 					CompileJobs,
@@ -2810,7 +2807,7 @@ void FMaterialShaderMap::Compile(
 	// Store the material name for debugging purposes.
 	// Note: Material instances with static parameters will have the same FriendlyName for their shader maps!
 	NewContent->FriendlyName = Material->GetFriendlyName();
-	NewContent->MaterialPath = Material->GetBaseMaterialPathName();
+	NewContent->MaterialPath = Material->GetAssetPath().ToString();
 
 	// Log debug information about the material being compiled.
 	const FString MaterialUsage = Material->GetMaterialUsageDescription();
@@ -2929,7 +2926,7 @@ void FMaterialShaderMap::Compile(
 #if WITH_EDITOR
 				if (bIsPersistent)
 				{
-					SaveToDerivedDataCache(FMaterialShaderParameters(Material));
+					SaveToDerivedDataCache(FConservativeShaderParameters(Material));
 				}
 #endif
 			}
@@ -3125,7 +3122,6 @@ void FMaterialShaderMap::ProcessCompilationResults(const TArray<FShaderCommonCom
 
 #endif // WITH_EDITOR
 
-// TODO - Move to header before submit
 uint64 GetMaterialShaderParametersHash(const FMaterialShaderParameters& MaterialParameters, uint64 Seed)
 {
 	if (MaterialParameters.MaterialShaderTags.Num())
@@ -3137,7 +3133,6 @@ uint64 GetMaterialShaderParametersHash(const FMaterialShaderParameters& Material
 	return CityHash64WithSeed((char*)&static_cast<const FMaterialShaderParametersBase&>(MaterialParameters), sizeof(FMaterialShaderParametersBase), Seed);
 }
 
-
 class FMaterialShaderMapLayoutCache
 {
 public:
@@ -3149,9 +3144,16 @@ public:
 
 	const FMaterialShaderMapLayout& AcquireLayout(EShaderPlatform Platform, EShaderPermutationFlags Flags, const FMaterialShaderParameters& MaterialParameters)
 	{
-		
 		const uint64 ParameterHash = GetMaterialShaderParametersHash(MaterialParameters, (uint64)Platform | ((uint64)Flags << 32));
 
+		FMaterialShaderMapLayout* Layout = FindOrCreateLayout(Platform, Flags, ParameterHash, MaterialParameters);
+
+		return *Layout;
+	}
+
+private:
+	FMaterialShaderMapLayout* FindOrCreateLayout(const EShaderPlatform Platform, const EShaderPermutationFlags Flags, const uint64 ParameterHash, const FMaterialShaderParameters& MaterialParameters)
+	{
 		FMaterialShaderMapLayout* Layout = nullptr;
 		{
 			FReadScopeLock Locker(LayoutLock);
@@ -3172,19 +3174,16 @@ public:
 
 				const int32 Index = ShaderMapLayouts.Add(TUniquePtr<FMaterialShaderMapLayout>(Layout));
 				MaterialParameterHashes.Add(ParameterHash);
-				MaterialShaderParameters.Add(MaterialParameters);
 				ShaderMapHashTable.Add(ParameterHash, Index);
 
 				check(MaterialParameterHashes.Num() == ShaderMapLayouts.Num());
-				check(MaterialShaderParameters.Num() == ShaderMapLayouts.Num());
 			}
 		}
 
 		check(Layout);
-		return *Layout;
+		return Layout;
 	}
 
-private:
 	FMaterialShaderMapLayout* FindLayout(uint64 ParameterHash)
 	{
 		for (int32 Index = ShaderMapHashTable.First(ParameterHash); ShaderMapHashTable.IsValid(Index); Index = ShaderMapHashTable.Next(Index))
@@ -3334,7 +3333,6 @@ private:
 	}
 
 	TArray<TUniquePtr<FMaterialShaderMapLayout>> ShaderMapLayouts;
-	TArray<FMaterialShaderParameters> MaterialShaderParameters;
 	TArray<uint64> MaterialParameterHashes;
 	FHashTable ShaderMapHashTable;
 	FRWLock LayoutLock;
@@ -3351,7 +3349,9 @@ bool FMaterialShaderMap::IsComplete(const FMaterial* Material, bool bSilent)
 
 	const FMaterialShaderMapContent* LocalContent = GetContent();
 	const EShaderPlatform Platform = LocalContent->GetShaderPlatform();
-	const FMaterialShaderParameters MaterialParameters(Material);
+	const FMaterialCompilationOutput& CompilationOutput = GetCompilationOutput();
+	const FCompilationShaderParameters MaterialParameters(Material, CompilationOutput);
+
 
 	// IsODSCOnly shaders aren't in shader map so don't include when gathering for IsComplete().
 	EShaderPermutationFlags Flags = ShaderMapId.GetPermutationFlags() & ~EShaderPermutationFlags::IsODSCOnly;
diff --git a/Engine/Source/Runtime/Engine/Private/Materials/MaterialShared.cpp b/Engine/Source/Runtime/Engine/Private/Materials/MaterialShared.cpp
--- a/Engine/Source/Runtime/Engine/Private/Materials/MaterialShared.cpp
+++ b/Engine/Source/Runtime/Engine/Private/Materials/MaterialShared.cpp
@@ -111,6 +111,12 @@ static TAutoConsoleVariable<bool> CVarUsingUseNewMaterialTranslatorPrototype(
 	TEXT("Controls whether to enable the new material translator prototype (WIP) ."),
 	ECVF_RenderThreadSafe);
 
+static TAutoConsoleVariable<int32> CVarMaterialUseShaderCompilationParameters(
+	TEXT("r.Material.UseShaderCompilationParameters"),
+	1,
+	TEXT("Controls whether the parameters used for determining which shaders to compile are based on translation output or conservative material data."),
+	ECVF_ReadOnly | ECVF_RenderThreadSafe);
+
 bool IsUsingNewMaterialTranslatorPrototype()
 {
 	return CVarUsingUseNewMaterialTranslatorPrototype.GetValueOnAnyThread();
@@ -618,9 +624,19 @@ FString FVectorMaterialInput::GetDefaultValue() const
 }
 
 #if WITH_EDITOR
+
+// As present in UMaterialGraphNode_Root::UpdateInputUseConstant, which fixes up Materials on modification.
+// This ensures Materials which have not been fixed up don't cause issues.
+static inline bool IsValidPropertyForInlineConst(const EMaterialProperty InProp)
+{
+	return InProp != MP_WorldPositionOffset && InProp != MP_Anisotropy && InProp != MP_PixelDepthOffset && InProp != MP_Displacement;
+}
+
 int32 FColorMaterialInput::CompileWithDefault(class FMaterialCompiler* Compiler, EMaterialProperty Property)
 {
-	if (UseConstant)
+	const bool bValidPropertyForInlineConstant = IsValidPropertyForInlineConst(Property);
+
+	if (UseConstant && bValidPropertyForInlineConstant)
 	{
 		return Compiler->Constant3(Constant.R, Constant.G, Constant.B);
 	}
@@ -636,9 +652,13 @@ int32 FColorMaterialInput::CompileWithDefault(class FMaterialCompiler* Compiler,
 	return Compiler->ForceCast(FMaterialAttributeDefinitionMap::CompileDefaultExpression(Compiler, Property), MCT_Float3);
 }
 
+
+
 int32 FScalarMaterialInput::CompileWithDefault(class FMaterialCompiler* Compiler, EMaterialProperty Property)
 {
-	if (UseConstant)
+	const bool bValidPropertyForInlineConstant = IsValidPropertyForInlineConst(Property);
+
+	if (UseConstant && bValidPropertyForInlineConstant)
 	{
 		return Compiler->Constant(Constant);
 	}
@@ -682,9 +702,12 @@ int32 FSubstrateMaterialInput::CompileWithDefault(class FMaterialCompiler* Compi
 	return Compiler->ForceCast(FMaterialAttributeDefinitionMap::CompileDefaultExpression(Compiler, Property), MCT_Substrate);
 }
 
+
 int32 FVectorMaterialInput::CompileWithDefault(class FMaterialCompiler* Compiler, EMaterialProperty Property)
 {
-	if (UseConstant)
+	const bool bValidPropertyForInlineConstant = IsValidPropertyForInlineConst(Property);
+
+	if (UseConstant && bValidPropertyForInlineConstant)
 	{
 		return Compiler->Constant3(Constant.X, Constant.Y, Constant.Z);
 	}
@@ -701,7 +724,9 @@ int32 FVectorMaterialInput::CompileWithDefault(class FMaterialCompiler* Compiler
 
 int32 FVector2MaterialInput::CompileWithDefault(class FMaterialCompiler* Compiler, EMaterialProperty Property)
 {
-	if (UseConstant)
+	const bool bValidPropertyForInlineConstant = IsValidPropertyForInlineConst(Property);
+
+	if (UseConstant && bValidPropertyForInlineConstant)
 	{
 		return Compiler->Constant2(Constant.X, Constant.Y);
 	}
@@ -1101,8 +1126,8 @@ bool FMaterial::MaterialModifiesMeshPosition_GameThread() const
 
 bool FMaterial::MaterialMayModifyMeshPosition() const
 {
-	// Conservative estimate when called before material translation has occurred. 
-	// This function is only intended for use in deciding whether or not shader permutations are required.
+	// Conservative estimate based on expressions present in graph.
+	// This function is only intended to be used for DDC Key/Iterative Cooking
 	return HasVertexPositionOffsetConnected() || HasPixelDepthOffsetConnected() || HasDisplacementConnected() || HasFirstPersonOutput();
 }
 
@@ -3414,7 +3439,7 @@ void FMaterial::CacheGivenTypes(const TArray<const FVertexFactoryType*>& VFTypes
 					GetShaderPlatform(),
 					GameThreadShaderMap->GetPermutationFlags(),
 					this,
-					GameThreadShaderMap->GetShaderMapId(),
+					GameThreadShaderMap->GetCompilationOutput(),
 					GameThreadPendingCompilerEnvironment,
 					VFType,
 					PipelineType,
@@ -3426,14 +3451,14 @@ void FMaterial::CacheGivenTypes(const TArray<const FVertexFactoryType*>& VFTypes
 			else if (ShaderType->GetTypeForDynamicCast() == FShaderType::EShaderTypeForDynamicCast::Material)
 			{
 				const EShaderPermutationFlags ShaderPermutation = GameThreadShaderMap->GetPermutationFlags();
+				const FMaterialCompilationOutput& CompilationOutput = GameThreadShaderMap->GetCompilationOutput();
 
 				const uint32 CompilingShaderMapId = GetGameThreadCompilingShaderMapId();
-				const FMaterialShaderMapId& ShaderMapId = GameThreadShaderMap->GetShaderMapId();
 				const FMaterialShaderType* MaterialShaderType = ShaderType->GetMaterialShaderType();
 
 				for (int32 PermutationId = 0; PermutationId < ShaderType->GetPermutationCount(); ++PermutationId)
 				{
-					const bool bShaderShouldCompile = MaterialShaderType->ShouldCompilePermutation(GetShaderPlatform(), this, PermutationId, ShaderPermutation);
+					const bool bShaderShouldCompile = MaterialShaderType->ShouldCompilePermutation(GetShaderPlatform(), FCompilationShaderParameters(this, CompilationOutput), PermutationId, ShaderPermutation);
 					if (!bShaderShouldCompile)
 					{
 						continue;
@@ -3444,7 +3469,7 @@ void FMaterial::CacheGivenTypes(const TArray<const FVertexFactoryType*>& VFTypes
 						CompilingShaderMapId,
 						PermutationId,
 						this,
-						ShaderMapId,
+						GameThreadShaderMap->GetCompilationOutput(),
 						GameThreadPendingCompilerEnvironment,
 						GetShaderPlatform(),
 						ShaderPermutation,
@@ -3457,12 +3482,13 @@ void FMaterial::CacheGivenTypes(const TArray<const FVertexFactoryType*>& VFTypes
 			else if (ShaderType->GetTypeForDynamicCast() == FShaderType::EShaderTypeForDynamicCast::MeshMaterial)
 			{
 				const EShaderPermutationFlags ShaderPermutation = GameThreadShaderMap->GetPermutationFlags();
+				const FMaterialCompilationOutput& CompilationOutput = GameThreadShaderMap->GetCompilationOutput();
 
 				const FMeshMaterialShaderType* MeshMaterialShaderType = ShaderType->GetMeshMaterialShaderType();
 
 				for (int32 PermutationId = 0; PermutationId < ShaderType->GetPermutationCount(); ++PermutationId)
 				{
-					const bool bShaderShouldCompile = MeshMaterialShaderType->ShouldCompilePermutation(GetShaderPlatform(), this, VFType, PermutationId, ShaderPermutation);
+					const bool bShaderShouldCompile = MeshMaterialShaderType->ShouldCompilePermutation(GetShaderPlatform(), FCompilationShaderParameters(this, CompilationOutput), VFType, PermutationId, ShaderPermutation);
 					if (!bShaderShouldCompile)
 					{
 						continue;
@@ -3475,7 +3501,7 @@ void FMaterial::CacheGivenTypes(const TArray<const FVertexFactoryType*>& VFTypes
 						GetShaderPlatform(),
 						ShaderPermutation,
 						this,
-						GameThreadShaderMap->GetShaderMapId(),
+						GameThreadShaderMap->GetCompilationOutput(),
 						GameThreadPendingCompilerEnvironment,
 						VFType,
 						CompileJobs,
@@ -3496,8 +3522,7 @@ bool FMaterial::Translate_Legacy(const FMaterialShaderMapId& ShaderMapId,
 	FMaterialCompilationOutput& OutCompilationOutput,
 	TRefCountPtr<FSharedShaderCompilerEnvironment>& OutMaterialEnvironment)
 {
-	FString MaterialTranslationDDCKeyString = GetMaterialShaderMapKeyString(ShaderMapId, FMaterialShaderParameters(this), GetShaderPlatform(), false);
-	
+	FString MaterialTranslationDDCKeyString = GetMaterialShaderMapKeyString(ShaderMapId, FConservativeShaderParameters(this), GetShaderPlatform(), false);
 	FHLSLMaterialTranslator MaterialTranslator(this, OutCompilationOutput, InStaticParameters, GetShaderPlatform(), GetQualityLevel(), ShaderMapId.FeatureLevel, InTargetPlatform, &ShaderMapId.SubstrateCompilationConfig, MaterialTranslationDDCKeyString);
 	EHLSLMaterialTranslatorResult Result = MaterialTranslator.Translate(false);
 
@@ -3841,12 +3866,13 @@ TShaderRef<FShader> FMaterial::GetShader(FMeshMaterialShaderType* ShaderType, FV
 				const auto CachedGameThreadShaderMap = GameThreadShaderMap;
 				const auto CachedGameMeshShaderMap = CachedGameThreadShaderMap ? CachedGameThreadShaderMap->GetMeshShaderMap(VertexFactoryType) : nullptr;
 				bool bShaderWasFoundInGameShaderMap = CachedGameMeshShaderMap && CachedGameMeshShaderMap->GetShader(ShaderType, PermutationId) != nullptr;
+				const FMaterialCompilationOutput& CompilationOutput = RenderingThreadShaderMap->GetCompilationOutput();
 
 				// Get the ShouldCache results that determine whether the shader should be compiled
 				auto ShaderPermutation = RenderingThreadShaderMap->GetPermutationFlags();
 				bool bMaterialShouldCache = ShouldCache(ShaderType, VertexFactoryType);
-				bool bVFShouldCache = FMeshMaterialShaderType::ShouldCompileVertexFactoryPermutation(GetShaderPlatform(), this, VertexFactoryType, ShaderType, ShaderPermutation);
-				bool bShaderShouldCache = ShaderType->ShouldCompilePermutation(GetShaderPlatform(), this, VertexFactoryType, PermutationId, ShaderPermutation);
+				bool bVFShouldCache = FMeshMaterialShaderType::ShouldCompileVertexFactoryPermutation(GetShaderPlatform(), FCompilationShaderParameters(this, CompilationOutput), VertexFactoryType, ShaderType, ShaderPermutation);
+				bool bShaderShouldCache = ShaderType->ShouldCompilePermutation(GetShaderPlatform(), FCompilationShaderParameters(this, CompilationOutput), VertexFactoryType, PermutationId, ShaderPermutation);
 				FString MaterialUsage = GetMaterialUsageDescription();
 
 				int BreakPoint = 0;
@@ -3940,27 +3966,28 @@ bool FMaterial::TryGetShaders(const FMaterialShaderTypes& InTypes, const FVertex
 		? static_cast<const FShaderMapContent*>(ShaderMap->GetMeshShaderMap(InVertexFactoryType))
 		: static_cast<const FShaderMapContent*>(ShaderMap->GetContent());
 
+	const FMaterialCompilationOutput& CompilationOutput = ShaderMap->GetCompilationOutput();
+
 	TArray<FShaderCommonCompileJobPtr> CompileJobs;
 	bool bMissingShader = false;
 
-	auto ShouldCacheShaderType = [this, &PermutationFlags](const FShaderType* ShaderType, const FVertexFactoryType* VertexFactoryType, const int32 PermutationId) -> bool {
+	auto ShouldCacheShaderType = [this, &PermutationFlags, &CompilationOutput](const FShaderType* ShaderType, const FVertexFactoryType* VertexFactoryType, const int32 PermutationId) -> bool {
 		// Check to see if the FMaterial should cache these types.
 		if (!ShouldCache(ShaderType, VertexFactoryType))
 		{
 			return false;
 		}
-
 		// if we are just a MaterialShaderType (not associated with a mesh)
 		if (const FMaterialShaderType* MaterialShader = ShaderType->GetMaterialShaderType())
 		{
-			return MaterialShader->ShouldCompilePermutation(GetShaderPlatform(), this, PermutationId, PermutationFlags);
+			return MaterialShader->ShouldCompilePermutation(GetShaderPlatform(), FCompilationShaderParameters(this, CompilationOutput), PermutationId, PermutationFlags);
 		}
 
 		// if we are a MeshMaterialShader
 		if (const FMeshMaterialShaderType* MeshMaterialShader = ShaderType->GetMeshMaterialShaderType())
 		{
-			const bool bVFShouldCache = FMeshMaterialShaderType::ShouldCompileVertexFactoryPermutation(GetShaderPlatform(), this, VertexFactoryType, ShaderType, PermutationFlags);
-			const bool bShaderShouldCache = MeshMaterialShader->ShouldCompilePermutation(GetShaderPlatform(), this, VertexFactoryType, PermutationId, PermutationFlags);
+			const bool bVFShouldCache = FMeshMaterialShaderType::ShouldCompileVertexFactoryPermutation(GetShaderPlatform(), FCompilationShaderParameters(this, CompilationOutput), VertexFactoryType, ShaderType, PermutationFlags);
+			const bool bShaderShouldCache = MeshMaterialShader->ShouldCompilePermutation(GetShaderPlatform(), FCompilationShaderParameters(this, CompilationOutput), VertexFactoryType, PermutationId, PermutationFlags);
 			return bVFShouldCache && bShaderShouldCache;
 		}
 		
@@ -4062,7 +4089,7 @@ bool FMaterial::TryGetShaders(const FMaterialShaderTypes& InTypes, const FVertex
 									ShaderPlatform, 
 									PermutationFlags, 
 									this,
-									ShaderMap->GetShaderMapId(),
+									ShaderMap->GetCompilationOutput(),
 									RenderingThreadPendingCompilerEnvironment, 
 									InVertexFactoryType, 
 									InTypes.PipelineType, 
@@ -4079,7 +4106,7 @@ bool FMaterial::TryGetShaders(const FMaterialShaderTypes& InTypes, const FVertex
 									ShaderPlatform, 
 									PermutationFlags, 
 									this,
-									ShaderMap->GetShaderMapId(),
+									ShaderMap->GetCompilationOutput(),
 									RenderingThreadPendingCompilerEnvironment, 
 									InTypes.PipelineType, 
 									CompileJobs,
@@ -4171,7 +4198,7 @@ bool FMaterial::TryGetShaders(const FMaterialShaderTypes& InTypes, const FVertex
 										ShaderPlatform,
 										PermutationFlags,
 										this,
-										ShaderMap->GetShaderMapId(),
+										ShaderMap->GetCompilationOutput(),
 										RenderingThreadPendingCompilerEnvironment,
 										InVertexFactoryType,
 										CompileJobs,
@@ -4186,7 +4213,7 @@ bool FMaterial::TryGetShaders(const FMaterialShaderTypes& InTypes, const FVertex
 										CompilingShaderMapId,
 										PermutationId,
 										this,
-										ShaderMap->GetShaderMapId(),
+										ShaderMap->GetCompilationOutput(),
 										RenderingThreadPendingCompilerEnvironment,
 										ShaderPlatform,
 										PermutationFlags,
@@ -4252,12 +4279,16 @@ FShaderPipelineRef FMaterial::GetShaderPipeline(class FShaderPipelineType* Shade
 				auto ShaderPermutation = RenderingThreadShaderMap->GetPermutationFlags();
 				FString MaterialUsage = GetMaterialUsageDescription();
 
+				const FMaterialCompilationOutput& CompilationOutput = RenderingThreadShaderMap->GetCompilationOutput();
+
 				UE_LOG(LogMaterial, Error,
 					TEXT("Couldn't find ShaderPipeline %s for Material Resource %s!"), ShaderPipelineType->GetName(), *GetFriendlyName());
 
 				for (auto* ShaderType : ShaderPipelineType->GetStages())
 				{
 					FShader* Shader = MeshShaderMap ? MeshShaderMap->GetShader((FShaderType*)ShaderType) : RenderingThreadShaderMap->GetShader((FShaderType*)ShaderType).GetShader();
+					const FCompilationShaderParameters ShaderParams(this, CompilationOutput);
+
 					if (!Shader)
 					{
 						UE_LOG(LogMaterial, Error, TEXT("Missing %s shader %s!"), GetShaderFrequencyString(ShaderType->GetFrequency(), false), ShaderType->GetName());
@@ -4265,8 +4296,8 @@ FShaderPipelineRef FMaterial::GetShaderPipeline(class FShaderPipelineType* Shade
 					else if (ShaderType->GetMeshMaterialShaderType())
 					{
 						bool bMaterialShouldCache = ShouldCache(ShaderType->GetMeshMaterialShaderType(), VertexFactoryType);
-						bool bVFShouldCache = FMeshMaterialShaderType::ShouldCompileVertexFactoryPermutation(GetShaderPlatform(), this, VertexFactoryType, ShaderType, ShaderPermutation);
-						bool bShaderShouldCache = ShaderType->GetMeshMaterialShaderType()->ShouldCompilePermutation(GetShaderPlatform(), this, VertexFactoryType, kUniqueShaderPermutationId, ShaderPermutation);
+						bool bVFShouldCache = FMeshMaterialShaderType::ShouldCompileVertexFactoryPermutation(GetShaderPlatform(), ShaderParams, VertexFactoryType, ShaderType, ShaderPermutation);
+						bool bShaderShouldCache = ShaderType->GetMeshMaterialShaderType()->ShouldCompilePermutation(GetShaderPlatform(), ShaderParams, VertexFactoryType, kUniqueShaderPermutationId, ShaderPermutation);
 
 						UE_LOG(LogMaterial, Error, TEXT("%s %s ShouldCache: Mat=%u, VF=%u, Shader=%u"),
 							GetShaderFrequencyString(ShaderType->GetFrequency(), false), ShaderType->GetName(), bMaterialShouldCache, bVFShouldCache, bShaderShouldCache);
@@ -4274,7 +4305,7 @@ FShaderPipelineRef FMaterial::GetShaderPipeline(class FShaderPipelineType* Shade
 					else if (ShaderType->GetMaterialShaderType())
 					{
 						bool bMaterialShouldCache = ShouldCache(ShaderType->GetMaterialShaderType(), VertexFactoryType);
-						bool bShaderShouldCache = ShaderType->GetMaterialShaderType()->ShouldCompilePermutation(GetShaderPlatform(), this, kUniqueShaderPermutationId, ShaderPermutation);
+						bool bShaderShouldCache = ShaderType->GetMaterialShaderType()->ShouldCompilePermutation(GetShaderPlatform(), ShaderParams, kUniqueShaderPermutationId, ShaderPermutation);
 
 						UE_LOG(LogMaterial, Error, TEXT("%s %s ShouldCache: Mat=%u, NO VF, Shader=%u"),
 							GetShaderFrequencyString(ShaderType->GetFrequency(), false), ShaderType->GetName(), bMaterialShouldCache, bShaderShouldCache);
@@ -4591,7 +4622,7 @@ void FMaterial::GetDependentShaderAndVFTypes(EShaderPlatform Platform, const FPl
 
 void FMaterial::GetDependentShaderAndVFTypes(const FPlatformTypeLayoutParameters& LayoutParams, TArray<FShaderType*>& OutShaderTypes, TArray<const FShaderPipelineType*>& OutShaderPipelineTypes, TArray<FVertexFactoryType*>& OutVFTypes) const
 {
-	const FMaterialShaderParameters MaterialParameters(this);
+	const FConservativeShaderParameters MaterialParameters(this);
 	const FMaterialShaderMapLayout& Layout = AcquireMaterialShaderMapLayout(GetShaderPlatform(), GetShaderPermutationFlags(LayoutParams), MaterialParameters);
 
 	for (const FShaderLayoutEntry& Shader : Layout.Shaders)
@@ -4930,20 +4961,26 @@ void FMaterial::GetShaderTypesForLayout(const FShaderMapLayout& Layout, FVertexF
 // Deprecated in 5.7
 void FMaterial::GetShaderTypes(EShaderPlatform Platform, const FPlatformTypeLayoutParameters& LayoutParams, TArray<FDebugShaderTypeInfo>& OutShaderInfo) const
 {
-	GetShaderTypes(LayoutParams, OutShaderInfo);
+	GetShaderTypes(LayoutParams, OutShaderInfo, FConservativeShaderParameters(this));
 }
 
 void FMaterial::GetShaderTypes(const FPlatformTypeLayoutParameters& LayoutParams, TArray<FDebugShaderTypeInfo>& OutShaderInfo) const
 {
-	const FMaterialShaderParameters MaterialParameters(this);
-	const FMaterialShaderMapLayout& Layout = AcquireMaterialShaderMapLayout(GetShaderPlatform(), GetShaderPermutationFlags(LayoutParams), MaterialParameters);
+	const FConservativeShaderParameters MaterialParameters(this);
+	GetShaderTypes(LayoutParams, OutShaderInfo, MaterialParameters);
+}
+
+void FMaterial::GetShaderTypes(const FPlatformTypeLayoutParameters& LayoutParams, TArray<FDebugShaderTypeInfo>& OutShaderInfo, const FMaterialShaderParameters& InParameters) const
+{
+	const FMaterialShaderMapLayout& Layout = AcquireMaterialShaderMapLayout(GetShaderPlatform(), GetShaderPermutationFlags(LayoutParams), InParameters);
 	GetShaderTypesForLayout(Layout, nullptr, OutShaderInfo);
 
 	for (const FMeshMaterialShaderMapLayout& MeshLayout : Layout.MeshShaderMaps)
 	{
-		GetShaderTypesForLayout( MeshLayout, MeshLayout.VertexFactoryType, OutShaderInfo);
+		GetShaderTypesForLayout(MeshLayout, MeshLayout.VertexFactoryType, OutShaderInfo);
 	}
 }
+
 #endif
 
 FMaterialUpdateContext::FMaterialUpdateContext(uint32 Options, EShaderPlatform InShaderPlatform)
@@ -5921,6 +5958,7 @@ FMaterialShaderParametersBase::FMaterialShaderParametersBase(const FMaterial* In
 	bShouldCastDynamicShadows = InMaterial->ShouldCastDynamicShadows();
 	bWritesEveryPixel = InMaterial->WritesEveryPixel(false);
 	bWritesEveryPixelShadowPass = InMaterial->WritesEveryPixel(true);
+
 	if (Substrate::IsSubstrateEnabled())
 	{
 		bIsSubstrateMaterial = InMaterial->IsSubstrateMaterial();
@@ -5933,7 +5971,6 @@ FMaterialShaderParametersBase::FMaterialShaderParametersBase(const FMaterial* In
 		bHasMetallicConnected = InMaterial->HasMaterialPropertyConnected(MP_Metallic);
 		bHasEmissiveColorConnected = InMaterial->HasMaterialPropertyConnected(MP_EmissiveColor);
 		bHasAmbientOcclusionConnected = InMaterial->HasMaterialPropertyConnected(MP_AmbientOcclusion);
-		bHasAnisotropyConnected = InMaterial->HasMaterialPropertyConnected(MP_Anisotropy);
 	}
 	else
 	{
@@ -5944,13 +5981,9 @@ FMaterialShaderParametersBase::FMaterialShaderParametersBase(const FMaterial* In
 		bHasMetallicConnected = InMaterial->HasMetallicConnected();
 		bHasEmissiveColorConnected = InMaterial->HasEmissiveColorConnected();
 		bHasAmbientOcclusionConnected = InMaterial->HasAmbientOcclusionConnected();
-		bHasAnisotropyConnected = InMaterial->HasAnisotropyConnected();
 	}
-	bHasVertexPositionOffsetConnected = InMaterial->HasVertexPositionOffsetConnected();
-	bHasPixelDepthOffsetConnected = InMaterial->HasPixelDepthOffsetConnected();
+
 	bIsTessellationEnabled = InMaterial->IsTessellationEnabled();
-	bHasDisplacementConnected = InMaterial->HasDisplacementConnected();
-	bMaterialMayModifyMeshPosition = InMaterial->MaterialMayModifyMeshPosition();
 	bIsUsedWithStaticLighting = InMaterial->IsUsedWithStaticLighting();
 	bIsUsedWithParticleSprites = InMaterial->IsUsedWithParticleSprites();
 	bIsUsedWithMeshParticles = InMaterial->IsUsedWithMeshParticles();
@@ -5971,9 +6004,7 @@ FMaterialShaderParametersBase::FMaterialShaderParametersBase(const FMaterial* In
 	bIsTranslucencyWritingCustomDepth = InMaterial->IsTranslucencyWritingCustomDepth();
 	bIsDitheredLODTransition = InMaterial->IsDitheredLODTransition();
 	bIsUsedWithInstancedStaticMeshes = InMaterial->IsUsedWithInstancedStaticMeshes();
-	bHasPerInstanceCustomData = InMaterial->HasPerInstanceCustomData();
 	bHasPerInstanceRandom = InMaterial->HasPerInstanceRandom();
-	bHasVertexInterpolator = InMaterial->HasVertexInterpolator();
 	bHasRuntimeVirtualTextureOutput = InMaterial->HasRuntimeVirtualTextureOutput();
 	bIsUsedWithLidarPointCloud = InMaterial->IsUsedWithLidarPointCloud();
 	bIsUsedWithVirtualHeightfieldMesh = InMaterial->IsUsedWithVirtualHeightfieldMesh();
@@ -5993,8 +6024,6 @@ FMaterialShaderParametersBase::FMaterialShaderParametersBase(const FMaterial* In
 	bIsUsedWithStaticMesh = InMaterial->IsUsedWithStaticMesh();
 	bIsTranslucentSurfaceLighting = InMaterial->GetTranslucencyLightingMode() == ETranslucencyLightingMode::TLM_SurfacePerPixelLighting; // get lighting mode for mobile translucent receive CSM 
 	bHasFirstPersonInterpolation = InMaterial->HasFirstPersonOutput();
-	bUsesTemporalResponsiveness = InMaterial->UsesTemporalResponsiveness();
-	bUsesMotionVectorWorldOffset = InMaterial->UsesMotionVectorWorldOffset();
 	bSamplesMaterialCache = InMaterial->SamplesMaterialCache();
 	bHasMaterialCacheOutput = InMaterial->HasMaterialCacheOutput();
 	bIsUsedWithMeshDeformer = InMaterial->IsUsedWithMeshDeformer();
@@ -6009,8 +6038,73 @@ FMaterialShaderParameters::FMaterialShaderParameters(const FMaterial* InMaterial
 		InMaterial->GetShaderTags(MaterialShaderTags);
 	}
 #endif
+
 }
 
+void FMaterialShaderParameters::InitializeShaderLayoutParamsFromMaterial(const FMaterial* InMaterial)
+{
+	if (Substrate::IsSubstrateEnabled())
+	{
+		bHasAnisotropyConnected = InMaterial->HasMaterialPropertyConnected(MP_Anisotropy);
+	}
+	else
+	{
+		bHasAnisotropyConnected = InMaterial->HasAnisotropyConnected();
+	}
+
+	bHasVertexPositionOffsetConnected = InMaterial->HasVertexPositionOffsetConnected();
+	bHasPixelDepthOffsetConnected = InMaterial->HasPixelDepthOffsetConnected();
+	bHasDisplacementConnected = InMaterial->HasDisplacementConnected();
+	bMaterialMayModifyMeshPosition = InMaterial->MaterialMayModifyMeshPosition();
+	bHasPerInstanceCustomData = InMaterial->HasPerInstanceCustomData();
+	bHasVertexInterpolator = InMaterial->HasVertexInterpolator();
+	bUsesTemporalResponsiveness = InMaterial->UsesTemporalResponsiveness();
+	bUsesMotionVectorWorldOffset = InMaterial->UsesMotionVectorWorldOffset();
+}
+
+void FMaterialShaderParameters::InitializeShaderLayoutParamsFromCompOutput(const FMaterialCompilationOutput& InCompilationOutput)
+{
+	bHasAnisotropyConnected = InCompilationOutput.bUsesAnisotropy;
+	bHasVertexPositionOffsetConnected = InCompilationOutput.bUsesWorldPositionOffset;
+	bHasPixelDepthOffsetConnected = InCompilationOutput.bUsesPixelDepthOffset;
+	bHasDisplacementConnected = InCompilationOutput.bUsesDisplacement;
+	bMaterialMayModifyMeshPosition = InCompilationOutput.bModifiesMeshPosition;
+	bHasPerInstanceCustomData = InCompilationOutput.bUsesPerInstanceCustomData;
+	bHasVertexInterpolator = InCompilationOutput.bUsesVertexInterpolator;
+	bUsesTemporalResponsiveness = InCompilationOutput.bUsesTemporalResponsiveness;
+	bUsesMotionVectorWorldOffset = InCompilationOutput.bUsesMotionVectorWorldOffset;
+}
+
+FCompilationShaderParameters::FCompilationShaderParameters(const FMaterial* InMaterial, const FMaterialCompilationOutput& InCompilationOutput)
+PRAGMA_DISABLE_DEPRECATION_WARNINGS
+	: FMaterialShaderParameters(InMaterial)
+	PRAGMA_ENABLE_DEPRECATION_WARNINGS
+{
+	static const int32 bUseShaderCompilationParameters = CVarMaterialUseShaderCompilationParameters.GetValueOnAnyThread();
+	if (bUseShaderCompilationParameters)
+	{
+		InitializeShaderLayoutParamsFromCompOutput(InCompilationOutput);
+	}
+	else
+	{
+		InitializeShaderLayoutParamsFromMaterial(InMaterial);
+	}
+}
+
+FConservativeShaderParameters::FConservativeShaderParameters()
+PRAGMA_DISABLE_DEPRECATION_WARNINGS
+	: FMaterialShaderParameters() {}
+PRAGMA_ENABLE_DEPRECATION_WARNINGS
+
+FConservativeShaderParameters::FConservativeShaderParameters(const FMaterial* InMaterial)
+PRAGMA_DISABLE_DEPRECATION_WARNINGS
+	: FMaterialShaderParameters(InMaterial)
+PRAGMA_ENABLE_DEPRECATION_WARNINGS
+{
+	InitializeShaderLayoutParamsFromMaterial(InMaterial);
+}
+
+
 #if WITH_EDITOR
 void FMaterialShaderParametersBase::Save(FCbWriter& Writer) const
 {
@@ -6035,7 +6129,7 @@ bool LoadFromCompactBinary(FCbFieldView Field, FMaterialShaderParametersBase& Ou
 	return OutValue.TryLoad(Field);
 }
 
-void FMaterialShaderParameters::Record(FMaterialKeyGeneratorContext& Context)
+void FConservativeShaderParameters::Record(FMaterialKeyGeneratorContext& Context)
 {
 	Context.Record("MaterialShaderParametersBase", static_cast<FMaterialShaderParametersBase&>(*this));
 	Context.Record("MaterialShaderTags", MaterialShaderTags);
diff --git a/Engine/Source/Runtime/Engine/Private/Materials/MaterialSharedPrivate.h b/Engine/Source/Runtime/Engine/Private/Materials/MaterialSharedPrivate.h
--- a/Engine/Source/Runtime/Engine/Private/Materials/MaterialSharedPrivate.h
+++ b/Engine/Source/Runtime/Engine/Private/Materials/MaterialSharedPrivate.h
@@ -42,7 +42,7 @@ void HashMaterialStaticClassDependenciesForCook(FAppendToClassSchemaContext& Con
 struct FRecordedMaterialResourceForCooking
 {
 	FMaterialShaderMapId ShaderMapId;
-	FMaterialShaderParameters ShaderParameters;
+	FConservativeShaderParameters ShaderParameters;
 	EShaderPlatform ShaderPlatform;
 	/**
 	 * We have to handle ExpressionIncludes separately rather than using ShaderMapId.ExpressionIncludesHash,
diff --git a/Engine/Source/Runtime/Engine/Private/Materials/MeshMaterialShader.cpp b/Engine/Source/Runtime/Engine/Private/Materials/MeshMaterialShader.cpp
--- a/Engine/Source/Runtime/Engine/Private/Materials/MeshMaterialShader.cpp
+++ b/Engine/Source/Runtime/Engine/Private/Materials/MeshMaterialShader.cpp
@@ -13,7 +13,7 @@
 static void PrepareMeshMaterialShaderCompileJob(EShaderPlatform Platform,
 	EShaderPermutationFlags PermutationFlags,
 	const FMaterial* Material,
-	const FMaterialShaderMapId& MaterialShaderMapId,
+	const FMaterialCompilationOutput& CompilationOutput,
 	FSharedShaderCompilerEnvironment* MaterialEnvironment,
 	const FShaderPipelineType* ShaderPipeline,
 	const FString& DebugGroupName,
@@ -37,7 +37,7 @@ static void PrepareMeshMaterialShaderCompileJob(EShaderPlatform Platform,
 		NewJob->bErrorsAreLikelyToBeCode = true;
 	}
 
-	const FMaterialShaderParameters MaterialParameters(Material);
+	const FCompilationShaderParameters MaterialParameters(Material, CompilationOutput);
 
 	// apply the vertex factory changes to the compile environment
 	check(VertexFactoryType);
@@ -86,7 +86,7 @@ void FMeshMaterialShaderType::BeginCompileShader(
 	EShaderPlatform Platform,
 	EShaderPermutationFlags PermutationFlags,
 	const FMaterial* Material,
-	const FMaterialShaderMapId& ShaderMapId,
+	const FMaterialCompilationOutput& CompilationOutput,
 	FSharedShaderCompilerEnvironment* MaterialEnvironment,
 	const FVertexFactoryType* VertexFactoryType,
 	TArray<FShaderCommonCompileJobPtr>& NewJobs,
@@ -97,7 +97,7 @@ void FMeshMaterialShaderType::BeginCompileShader(
 	FShaderCompileJob* NewJob = GShaderCompilingManager->PrepareShaderCompileJob(ShaderMapJobId, FShaderCompileJobKey(this, Platform, VertexFactoryType, PermutationId), Priority);
 	if (NewJob)
 	{
-		PrepareMeshMaterialShaderCompileJob(Platform, PermutationFlags, Material, ShaderMapId, MaterialEnvironment, nullptr, DebugGroupName, DebugDescription, DebugExtension, NewJob);
+		PrepareMeshMaterialShaderCompileJob(Platform, PermutationFlags, Material, CompilationOutput, MaterialEnvironment, nullptr, DebugGroupName, DebugDescription, DebugExtension, NewJob);
 		NewJobs.Add(FShaderCommonCompileJobPtr(NewJob));
 	}
 }
@@ -109,7 +109,7 @@ void FMeshMaterialShaderType::BeginCompileShaderPipeline(
 	EShaderPlatform Platform,
 	EShaderPermutationFlags PermutationFlags,
 	const FMaterial* Material,
-	const FMaterialShaderMapId& ShaderMapId,
+	const FMaterialCompilationOutput& CompilationOutput,
 	FSharedShaderCompilerEnvironment* MaterialEnvironment,
 	const FVertexFactoryType* VertexFactoryType,
 	const FShaderPipelineType* ShaderPipeline,
@@ -127,7 +127,7 @@ void FMeshMaterialShaderType::BeginCompileShaderPipeline(
 	{
 		for (FShaderCompileJob* StageJob : NewPipelineJob->StageJobs)
 		{
-			PrepareMeshMaterialShaderCompileJob(Platform, PermutationFlags, Material, ShaderMapId, MaterialEnvironment, ShaderPipeline, DebugGroupName, DebugDescription, DebugExtension, StageJob);
+			PrepareMeshMaterialShaderCompileJob(Platform, PermutationFlags, Material, CompilationOutput, MaterialEnvironment, ShaderPipeline, DebugGroupName, DebugDescription, DebugExtension, StageJob);
 		}
 		NewJobs.Add(FShaderCommonCompileJobPtr(NewPipelineJob));
 	}
diff --git a/Engine/Source/Runtime/Engine/Private/ShaderCompiler/ShaderCompiler.cpp b/Engine/Source/Runtime/Engine/Private/ShaderCompiler/ShaderCompiler.cpp
--- a/Engine/Source/Runtime/Engine/Private/ShaderCompiler/ShaderCompiler.cpp
+++ b/Engine/Source/Runtime/Engine/Private/ShaderCompiler/ShaderCompiler.cpp
@@ -2345,7 +2345,7 @@ void FShaderCompilingManager::ProcessCompiledShaderMaps(
 					ShaderMapToUseForRendering->bCompilationFinalized = true;
 					if (ShaderMapToUseForRendering->bIsPersistent)
 					{
-						ShaderMapToUseForRendering->SaveToDerivedDataCache(FMaterialShaderParameters(SingleMaterial));
+						ShaderMapToUseForRendering->SaveToDerivedDataCache(FConservativeShaderParameters(SingleMaterial));
 					}
 				}
 
diff --git a/Engine/Source/Runtime/Engine/Public/MaterialCachedData.h b/Engine/Source/Runtime/Engine/Public/MaterialCachedData.h
--- a/Engine/Source/Runtime/Engine/Public/MaterialCachedData.h
+++ b/Engine/Source/Runtime/Engine/Public/MaterialCachedData.h
@@ -216,6 +216,7 @@ struct FMaterialCachedExpressionData
 
 #if WITH_EDITOR
 	void UpdateForExpressions(const FMaterialCachedExpressionContext& Context, TConstArrayView<TObjectPtr<UMaterialExpression>> Expressions, EMaterialParameterAssociation Association, int32 ParameterIndex);
+	void UpdatePropertyConnections(UMaterial& Material);
 	void UpdateForFunction(const FMaterialCachedExpressionContext& Context, UMaterialFunctionInterface* Function, EMaterialParameterAssociation Association, int32 ParameterIndex);
 	void UpdateForLayerFunctions(const FMaterialCachedExpressionContext& Context, const FMaterialLayersFunctions& LayerFunctions);
 	void AnalyzeMaterial(UMaterial& Material);
diff --git a/Engine/Source/Runtime/Engine/Public/MaterialShaderType.h b/Engine/Source/Runtime/Engine/Public/MaterialShaderType.h
--- a/Engine/Source/Runtime/Engine/Public/MaterialShaderType.h
+++ b/Engine/Source/Runtime/Engine/Public/MaterialShaderType.h
@@ -28,10 +28,12 @@ class FShaderKeyGenerator;
 class FUniformExpressionSet;
 class FVertexFactoryType;
 struct FMaterialShaderParameters;
+struct FConservativeShaderParameters;
 struct FMaterialShadingModelField;
 enum EBlendMode : int;
 enum EMaterialShadingModel : int;
 enum class EShaderCompileJobPriority : uint8;
+class FMaterialCompilationOutput;
 #if WITH_EDITOR
 class FMaterialKeyGeneratorContext;
 #endif
@@ -54,28 +56,21 @@ extern ENGINE_API FString GetBlendModeString(EBlendMode BlendMode);
 #if WITH_EDITOR
 extern ENGINE_API FString GetMaterialShaderMapKeyString(
 	const FMaterialShaderMapId& ShaderMapId,
-	const FMaterialShaderParameters& ShaderParameters,
+	const FConservativeShaderParameters& ShaderParameters,
 	EShaderPlatform Platform,
 	bool bIncludeKeyStringShaderDependencies = true);
 
 extern ENGINE_API void RecordOrEmitMaterialShaderMapKey(
 	FMaterialKeyGeneratorContext& Context,
 	FMaterialShaderMapId& ShaderMapId,
-	FMaterialShaderParameters& ShaderParameters);
+	FConservativeShaderParameters& ShaderParameters);
 
 // Alternate arguments for RecordOrEmitMaterialShaderMapKey: Support being called with const&
 // when emitting or saving.
 extern ENGINE_API void RecordOrEmitMaterialShaderMapKey(
 	FMaterialKeyGeneratorContext& Context,
 	const FMaterialShaderMapId& ShaderMapId,
-	const FMaterialShaderParameters& ShaderParameters);
-inline void RecordOrEmitMaterialShaderMapKey(
-	FMaterialKeyGeneratorContext& Context,
-	FMaterialShaderMapId& ShaderMapId,
-	FMaterialShaderParameters&& ShaderParameters)
-{
-	RecordOrEmitMaterialShaderMapKey(Context, ShaderMapId, ShaderParameters);
-}
+	const FConservativeShaderParameters& ShaderParameters);
 #endif
 
 /** Called for every material shader to update the appropriate stats. */
@@ -135,7 +130,7 @@ public:
 		uint32 ShaderMapJobId,
 		int32 PermutationId,
 		const FMaterial* Material,
-		const FMaterialShaderMapId& ShaderMapId,
+		const FMaterialCompilationOutput& CompilationOutput,
 		FSharedShaderCompilerEnvironment* MaterialEnvironment,
 		EShaderPlatform Platform,
 		EShaderPermutationFlags PermutationFlags,
@@ -151,7 +146,7 @@ public:
 		EShaderPlatform Platform,
 		EShaderPermutationFlags PermutationFlags,
 		const FMaterial* Material,
-		const FMaterialShaderMapId& ShaderMapId,
+		const FMaterialCompilationOutput& CompilationOutput,
 		FSharedShaderCompilerEnvironment* MaterialEnvironment,
 		const FShaderPipelineType* ShaderPipeline,
 		TArray<TRefCountPtr<FShaderCommonCompileJob>>& NewJobs,
diff --git a/Engine/Source/Runtime/Engine/Public/MaterialShared.h b/Engine/Source/Runtime/Engine/Public/MaterialShared.h
--- a/Engine/Source/Runtime/Engine/Public/MaterialShared.h
+++ b/Engine/Source/Runtime/Engine/Public/MaterialShared.h
@@ -1235,7 +1235,8 @@ struct FMaterialShaderParametersBase
 		};
 	};
 
-	FMaterialShaderParametersBase(const FMaterial* InMaterial=nullptr);
+protected:
+	FMaterialShaderParametersBase(const FMaterial* InMaterial = nullptr);
 
 private:
 #if WITH_EDITOR
@@ -1254,14 +1255,31 @@ private:
 // FMaterialShaderParametersBase split out to ease serializing of raw data
 struct FMaterialShaderParameters : public FMaterialShaderParametersBase
 {
-	ENGINE_API FMaterialShaderParameters(const FMaterial* InMaterial = nullptr);
+	TArray<FName> MaterialShaderTags;
+
+	UE_DEPRECATED(5.8, "This will be moved under protected. Use FCompilationShaderParameters or FConservativeShaderParameters to create Shader Parameters.")
+	ENGINE_API explicit FMaterialShaderParameters(const FMaterial* InMaterial = nullptr);
+
+protected:
+	void InitializeShaderLayoutParamsFromMaterial(const FMaterial* InMaterial);
+	void InitializeShaderLayoutParamsFromCompOutput(const FMaterialCompilationOutput& InCompilationOutput);
+};
+
+// For Shader Compilation. Built based on Translation output.
+struct FCompilationShaderParameters : public FMaterialShaderParameters
+{
+	ENGINE_API explicit FCompilationShaderParameters(const FMaterial* InMaterial, const FMaterialCompilationOutput& InCompilationOutput);
+};
+
+// For DDC Key/Iterative cooking. Built based on asset data.
+struct FConservativeShaderParameters : public FMaterialShaderParameters
+{
+	ENGINE_API FConservativeShaderParameters();
+	ENGINE_API explicit FConservativeShaderParameters(const FMaterial* InMaterial);
 
-public:
 #if WITH_EDITOR
 	void Record(FMaterialKeyGeneratorContext& Context);
 #endif
-
-	TArray<FName> MaterialShaderTags;
 };
 
 /** Contains all the information needed to uniquely identify a FMaterialShaderMap. */
@@ -1747,7 +1765,7 @@ public:
 
 #if WITH_EDITOR
 	/** Saves this shader map to the derived data cache. */
-	void SaveToDerivedDataCache(const FMaterialShaderParameters& ShaderParameters);
+	void SaveToDerivedDataCache(const FConservativeShaderParameters& ShaderParameters);
 #endif
 
 	/** Serializes a shader map to an archive (used with recompiling shaders for a remote console) */
@@ -1799,6 +1817,7 @@ public:
 	const TCHAR* GetDebugDescription() const { return TEXT(""); }
 	const TCHAR* GetMaterialPath() const { return TEXT(""); }
 #endif
+	const FMaterialCompilationOutput& GetCompilationOutput() const { return GetContent()->MaterialCompilationOutput; }
 	bool RequiresSceneColorCopy() const { return GetContent()->MaterialCompilationOutput.RequiresSceneColorCopy(); }
 	bool NeedsSceneTextures() const { return GetContent()->MaterialCompilationOutput.bNeedsSceneTextures; }
 	bool UsesGlobalDistanceField() const { return GetContent()->MaterialCompilationOutput.bUsesGlobalDistanceField; }
@@ -2855,10 +2874,16 @@ public:
 	*	@param	OutShaderInfo	Array of results sorted by vertex factory type, and shader type.
 	*
 	*/
-	UE_DEPRECATED(5.7, "Call GetShaderTypes without EShaderPlatform as its now cached inside the FMaterial")
+	UE_DEPRECATED(5.7, "Call GetShaderTypes without EShaderPlatform and with specialized ShaderParameters")
 	ENGINE_API void GetShaderTypes(EShaderPlatform Platform, const FPlatformTypeLayoutParameters& LayoutParams, TArray<FDebugShaderTypeInfo>& OutShaderInfo) const;
+
+	UE_DEPRECATED(5.8, "Call GetShaderTypes with specialized ShaderParameters")
+	/* Returns shader types based on conservative Shader Parameters. */
 	ENGINE_API void GetShaderTypes(const FPlatformTypeLayoutParameters& LayoutParams, TArray<FDebugShaderTypeInfo>& OutShaderInfo) const;
 
+	/* Returns shader types based on given Shader Parameters. */
+	ENGINE_API void GetShaderTypes(const FPlatformTypeLayoutParameters& LayoutParams, TArray<FDebugShaderTypeInfo>& OutShaderInfo, const FMaterialShaderParameters& InParameters) const;
+
 	UE_DEPRECATED(5.7, "Call GetShaderTypesForLayout without EShaderPlatform as its now cached inside the FMaterial")
 	void GetShaderTypesForLayout(EShaderPlatform Platform, const FShaderMapLayout& Layout, FVertexFactoryType* VertexFactory, TArray<FDebugShaderTypeInfo>& OutShaderInfo) const;
 	void GetShaderTypesForLayout(const FShaderMapLayout& Layout, FVertexFactoryType* VertexFactory, TArray<FDebugShaderTypeInfo>& OutShaderInfo) const;
@@ -2976,7 +3001,7 @@ protected:
 	/** Gets the Guid that represents this material. */
 	virtual FGuid GetMaterialId() const = 0;
 	
-	/** Produces arrays of any shader and vertex factory type that this material is dependent on. */
+	/** Produces arrays of any shader and vertex factory type that this material is dependent on, using conservative parameters. */
 	UE_DEPRECATED(5.7, "Call GetDependentShaderAndVFTypes without EShaderPlatform as its now cached inside the FMaterial")
 	ENGINE_API void GetDependentShaderAndVFTypes(EShaderPlatform Platform, const FPlatformTypeLayoutParameters& LayoutParams, TArray<FShaderType*>& OutShaderTypes, TArray<const FShaderPipelineType*>& OutShaderPipelineTypes, TArray<FVertexFactoryType*>& OutVFTypes) const;
 	ENGINE_API void GetDependentShaderAndVFTypes(const FPlatformTypeLayoutParameters& LayoutParams, TArray<FShaderType*>& OutShaderTypes, TArray<const FShaderPipelineType*>& OutShaderPipelineTypes, TArray<FVertexFactoryType*>& OutVFTypes) const;
diff --git a/Engine/Source/Runtime/Engine/Public/MeshMaterialShaderType.h b/Engine/Source/Runtime/Engine/Public/MeshMaterialShaderType.h
--- a/Engine/Source/Runtime/Engine/Public/MeshMaterialShaderType.h
+++ b/Engine/Source/Runtime/Engine/Public/MeshMaterialShaderType.h
@@ -59,7 +59,7 @@ public:
 		EShaderPlatform Platform,
 		EShaderPermutationFlags PermutationFlags,
 		const FMaterial* Material,
-		const FMaterialShaderMapId& ShaderMapId,
+		const FMaterialCompilationOutput& CompilationOutput,
 		FSharedShaderCompilerEnvironment* MaterialEnvironment,
 		const FVertexFactoryType* VertexFactoryType,
 		TArray<TRefCountPtr<FShaderCommonCompileJob>>& NewJobs,
@@ -75,7 +75,7 @@ public:
 		EShaderPlatform Platform,
 		EShaderPermutationFlags PermutationFlags,
 		const FMaterial* Material,
-		const FMaterialShaderMapId& ShaderMapId,
+		const FMaterialCompilationOutput& CompilationOutput,
 		FSharedShaderCompilerEnvironment* MaterialEnvironment,
 		const FVertexFactoryType* VertexFactoryType,
 		const FShaderPipelineType* ShaderPipeline,
diff --git a/Engine/Source/Runtime/RenderCore/Private/Shader.cpp b/Engine/Source/Runtime/RenderCore/Private/Shader.cpp
--- a/Engine/Source/Runtime/RenderCore/Private/Shader.cpp
+++ b/Engine/Source/Runtime/RenderCore/Private/Shader.cpp
@@ -1783,12 +1783,19 @@ void ShaderMapAppendKey(EShaderPlatform Platform, FShaderKeyGenerator& KeyGen)
 		KeyGen.Append((CVar && CVar->GetValueOnAnyThread() != 0) ? TEXT("DEV") : TEXT("NoDEV"));
 	}
 
+	{
+		static const auto CVar = IConsoleManager::Get().FindTConsoleVariableDataInt(TEXT("r.Material.UseShaderCompilationParameters"));
+		KeyGen.AppendSeparator();
+		KeyGen.Append((CVar && CVar->GetValueOnAnyThread() != 0) ? TEXT("CSP") : TEXT("NoCSP"));
+	}
+
 	{
 		const bool bValue = IsStaticLightingAllowed();
 		KeyGen.AppendSeparator();
 		KeyGen.Append(bValue ? TEXT("SL") : TEXT("NoSL"));
 	}
 
+
 	{
 		static const auto CVar = IConsoleManager::Get().FindTConsoleVariableDataInt(TEXT("r.MaterialEditor.LWCTruncateMode"));
 		const int32 LWCTruncateValue = CVar ? CVar->GetValueOnAnyThread() : 0;
