MeshCard: update CARDREPRESENTATION_DERIVEDDATA_VER due to mesh description serialization change

#rnx
#Synced-CL 51707648

diff --git a/Engine/Source/Runtime/Engine/Private/MeshCardRepresentation.cpp b/Engine/Source/Runtime/Engine/Private/MeshCardRepresentation.cpp
--- a/Engine/Source/Runtime/Engine/Private/MeshCardRepresentation.cpp
+++ b/Engine/Source/Runtime/Engine/Private/MeshCardRepresentation.cpp
@@ -191,7 +191,7 @@ FCardRepresentationAsyncQueue* GCardRepresentationAsyncQueue = NULL;
 #if WITH_EDITORONLY_DATA
 
 // DDC key for card representation data, must be changed when modifying the generation code or data format
-#define CARDREPRESENTATION_DERIVEDDATA_VER TEXT("C832711A-D404-4898-80FE-21B263DEF19E")
+#define CARDREPRESENTATION_DERIVEDDATA_VER TEXT("BE3395EF-F023-4F1D-9827-1DD3C88E16DC")
 
 void AppendMeshCardConstantsToClassSchema(FAppendToClassSchemaContext& Context)
 {
