[FastGeo] Fixed FNavigationSystem::IsLevelVisibilityChanging to also support UAssetUserData.

#Synced-CL 51624647

diff --git a/Engine/Source/Runtime/Engine/Private/AI/Navigation/NavigationTypes.cpp b/Engine/Source/Runtime/Engine/Private/AI/Navigation/NavigationTypes.cpp
--- a/Engine/Source/Runtime/Engine/Private/AI/Navigation/NavigationTypes.cpp
+++ b/Engine/Source/Runtime/Engine/Private/AI/Navigation/NavigationTypes.cpp
@@ -28,23 +28,22 @@ namespace FNavigationSystem
 
 	bool IsLevelVisibilityChanging(const UObject* Object)
 	{
-		const UActorComponent* ObjectAsComponent = Cast<UActorComponent>(Object);
-		if (ObjectAsComponent)
+		const ULevel* Level = nullptr;
+
+		if (const UActorComponent* ObjectAsComponent = Cast<UActorComponent>(Object))
 		{
-			if (const ULevel* Level = ObjectAsComponent->GetComponentLevel())
-			{
-				return Level->HasVisibilityChangeRequestPending();
-			}
+			Level = ObjectAsComponent->GetComponentLevel();
 		}
 		else if (const AActor* Actor = Cast<AActor>(Object))
 		{
-			if (const ULevel* Level = Actor->GetLevel())
-			{
-				return Level->HasVisibilityChangeRequestPending();
-			}
+			Level = Actor->GetLevel();
+		}
+		else if (const UAssetUserData* AssetUserData = Cast<UAssetUserData>(Object))
+		{
+			Level = AssetUserData->GetTypedOuter<ULevel>();
 		}
 
-		return false;
+		return Level ? Level->HasVisibilityChangeRequestPending() : false;
 	}
 	
 	bool IsInBaseNavmesh(const UObject* Object)
