Disable costly WP logs in test and shipping builds.

#Synced-CL 51708345

diff --git a/Engine/Source/Runtime/Engine/Private/WorldPartition/DataLayer/WorldDataLayers.cpp b/Engine/Source/Runtime/Engine/Private/WorldPartition/DataLayer/WorldDataLayers.cpp
--- a/Engine/Source/Runtime/Engine/Private/WorldPartition/DataLayer/WorldDataLayers.cpp
+++ b/Engine/Source/Runtime/Engine/Private/WorldPartition/DataLayer/WorldDataLayers.cpp
@@ -276,10 +276,12 @@ bool AWorldDataLayers::SetDataLayerRuntimeState(const UDataLayerInstance* InData
 
 		++DataLayersStateEpoch;
 
+#if !UE_BUILD_TEST && !UE_BUILD_SHIPPING
 		UE_LOG(LogWorldPartition, Log, TEXT("Data Layer Instance '%s' state changed: %s -> %s"),
 			*InDataLayerInstance->GetDataLayerShortName(),
 			*StaticEnum<EDataLayerRuntimeState>()->GetDisplayNameTextByValue((int64)CurrentState).ToString(),
 			*StaticEnum<EDataLayerRuntimeState>()->GetDisplayNameTextByValue((int64)InState).ToString());
+#endif // !UE_BUILD_TEST && !UE_BUILD_SHIPPING
 
 		CSV_EVENT_GLOBAL(TEXT("DataLayer-%s-%s"), *InDataLayerInstance->GetDataLayerShortName(), *StaticEnum<EDataLayerRuntimeState>()->GetDisplayNameTextByValue((int64)InState).ToString());
 
diff --git a/Engine/Source/Runtime/Engine/Private/WorldPartition/WorldPartition.cpp b/Engine/Source/Runtime/Engine/Private/WorldPartition/WorldPartition.cpp
--- a/Engine/Source/Runtime/Engine/Private/WorldPartition/WorldPartition.cpp
+++ b/Engine/Source/Runtime/Engine/Private/WorldPartition/WorldPartition.cpp
@@ -734,6 +734,7 @@ void UWorldPartition::Initialize(UWorld* InWorld, const FTransform& InTransform)
 	const bool bIsPIEWorldTravel = (World->WorldType == EWorldType::PIE) && !StreamingPolicy;
 	const bool bIsDedicatedServer = IsRunningDedicatedServer();
 
+#if !UE_BUILD_TEST && !UE_BUILD_SHIPPING
 	UE_LOG(LogWorldPartition, Log, TEXT("UWorldPartition::Initialize : World = %s, World Type = %s, IsMainWorldPartition = %d, Location = %s, Rotation = %s, IsEditor = %d, IsGame = %d, IsPIEWorldTravel = %d, IsCooking = %d"),
 		*OuterWorld->GetPathName(),
 		LexToString(World->WorldType),
@@ -756,6 +757,7 @@ void UWorldPartition::Initialize(UWorld* InWorld, const FTransform& InTransform)
 			UseMakingVisibleTransactionRequests() ? 1 : 0, 
 			UseMakingInvisibleTransactionRequests() ? 1 : 0);
 	}
+#endif // !UE_BUILD_TEST && !UE_BUILD_SHIPPING
 
 	auto CreateAndInitializeDataLayerManager = [this]()
 	{
