Disable SteamableManagerProfiling in AutoRTFM modes since it is implemented using atomics.  This prevents some AsyncLoads from crashing when used inside a Transaction.

#Synced-CL 51676086

diff --git a/Engine/Source/Runtime/Engine/Public/ProfilingDebugging/StreamableManagerProfiling.h b/Engine/Source/Runtime/Engine/Public/ProfilingDebugging/StreamableManagerProfiling.h
--- a/Engine/Source/Runtime/Engine/Public/ProfilingDebugging/StreamableManagerProfiling.h
+++ b/Engine/Source/Runtime/Engine/Public/ProfilingDebugging/StreamableManagerProfiling.h
@@ -6,7 +6,7 @@
 #include "ProfilingDebugging/CountersTrace.h"
 
 #ifndef UE_STREAMABLE_MANAGER_PROFILING_ENABLED
-#define UE_STREAMABLE_MANAGER_PROFILING_ENABLED (COUNTERSTRACE_ENABLED || CSV_PROFILER_STATS)
+#define UE_STREAMABLE_MANAGER_PROFILING_ENABLED (COUNTERSTRACE_ENABLED || CSV_PROFILER_STATS) && !UE_AUTORTFM
 #endif //  UE_STREAMABLE_MANAGER_PROFILING_ENABLED
 
 #if UE_STREAMABLE_MANAGER_PROFILING_ENABLED
