[BugFix] unkown properties of intanced subobjects weren't being reinstanced properly on compile so users would lose the data when recompiling. This effects both subobjects in placeholders, and just regular loose subobject references
#Synced-CL 51689558

diff --git a/Engine/Source/Editor/UnrealEd/Private/Kismet2/KismetReinstanceUtilities.cpp b/Engine/Source/Editor/UnrealEd/Private/Kismet2/KismetReinstanceUtilities.cpp
--- a/Engine/Source/Editor/UnrealEd/Private/Kismet2/KismetReinstanceUtilities.cpp
+++ b/Engine/Source/Editor/UnrealEd/Private/Kismet2/KismetReinstanceUtilities.cpp
@@ -35,6 +35,7 @@
 #include "UObject/UObjectHash.h"
 #include "UObject/UObjectIterator.h"
 #include "UObject/VerseValueProperty.h"
+#include "UObject/InstanceDataObjectUtils.h"
 #include "VerseVM/VBPVMDynamicProperty.h"
 #include "Serialization/FindObjectReferencers.h"
 #include "Serialization/ArchiveReplaceObjectAndStructPropertyRef.h"
@@ -413,6 +414,12 @@ struct FReplaceReferenceHelper
 			Root = Container;
 		}
 
+		// loose subobjects need to be reinstanced too so iterate the IDO instead if it exists so we can include loose subobjects references
+		if (UObject* FoundIdo = UE::FPropertyBagRepository::Get().FindInstanceDataObject(Container))
+		{
+			Container = FoundIdo;
+		}
+
 		const UClass* ContainerClass = Container->GetClass();
 		for (FProperty* Prop = ContainerClass->RefLink; Prop; Prop = Prop->NextRef)
 		{
@@ -2840,6 +2847,21 @@ void FBlueprintCompileReinstancer::ReplaceInstancesOfClass_Inner(const TMap<UCla
 						continue;
 					}
 
+					// with IDOs, it's possible for a loose subobject to become non-loose again which the above condition would've missed.
+					// that object will be part of the new archetype however so we want to skip reinstancing it here because the outer will handle that.
+					if (UE::IsInstanceDataObjectSupportEnabled(OldObject))
+					{
+						if (UClass*const* NewOuterClass = InOldToNewClassMap.Find(OldObject->GetOuter()->GetClass()))
+						{
+							// if this subobject is in the outer's CDO, we can assume that it will be handled when the outer is reinstanced
+							UObject* NewOuterCDO = (*NewOuterClass)->GetDefaultObject();
+							if (StaticFindObjectFast(UObject::StaticClass(), NewOuterCDO, OldObject->GetFName()))
+							{
+								continue;
+							}
+						}
+					}
+
 					AActor* OldActor = Cast<AActor>(OldObject);
 					bool bIsValid = IsValid(OldObject);
 
diff --git a/Engine/Source/Runtime/CoreUObject/Private/UObject/Class.cpp b/Engine/Source/Runtime/CoreUObject/Private/UObject/Class.cpp
--- a/Engine/Source/Runtime/CoreUObject/Private/UObject/Class.cpp
+++ b/Engine/Source/Runtime/CoreUObject/Private/UObject/Class.cpp
@@ -1807,10 +1807,8 @@ void UStruct::SerializeVersionedTaggedProperties(FStructuredArchive::FSlot Slot,
 #if WITH_EDITORONLY_DATA
 					// Handle the sub object shadowing serialization,
 					// do not skip this property if loading into a unknown property or into a placeholder
-					// When loading into the IDO, we should only consider the unknown properties.
-					// (IDO are serializing properties using bImpersonateProperties(true) but they are not necessary unknown)
 					else if (Tag.OverrideOperation == EOverriddenPropertyOperation::SubObjectsShadowing && 
-						(!SerializeContext->bImpersonateProperties || !UE::IsSerializingUnknownData(UnderlyingArchive, Property)))
+						(!SerializeContext->bImpersonateProperties ))
 					{
 						UnderlyingArchive.Seek(StartOfProperty + Tag.Size);
 						bAdvanceProperty = true;
diff --git a/Engine/Source/Runtime/Engine/Private/UnrealEngine.cpp b/Engine/Source/Runtime/Engine/Private/UnrealEngine.cpp
--- a/Engine/Source/Runtime/Engine/Private/UnrealEngine.cpp
+++ b/Engine/Source/Runtime/Engine/Private/UnrealEngine.cpp
@@ -161,6 +161,8 @@ UnrealEngine.cpp: Implements the UEngine class and helpers.
 #include "GenericPlatform/GenericPlatformCrashContext.h"
 #include "Framework/Notifications/NotificationManager.h"
 #include "WorldPartition/DataLayer/ExternalDataLayerAsset.h"
+#include "Reinstancing/ReinstancingInternal.h"
+#include "UObject/InstanceDataObjectUtils.h"
 #endif
 // @todo this is here only due to circular dependency to AIModule. To be removed
 
@@ -18420,6 +18422,110 @@ private:
 	const TMap<UObject*, UObject*>& AggressiveReplacements;
 };
 
+#if WITH_EDITOR
+// When IDOs are enabled, if the old class has subobjects that the new class doesn't, we need new versions of those subobjects to exist so that we can round trip that data.
+// this helper function is responsible for constructing those objects so the loose ido object references can point to them and they maintain their data
+static void CreateLooseSubobjectsForCPFUO(const UObject* OldObject, UObject* NewObject, TArray<UObject*>& InOutSubobjectArray,
+	const TMap<FString, int32>& OldInstanceMap, const TIndirectArray<FInstancedObjectRecord>& SavedInstances, const UEngine::FCopyPropertiesForUnrelatedObjectsParams& Params)
+{
+	TArray<UObject*> UnsortedNewInstances = MoveTemp(InOutSubobjectArray);
+
+	TArray<int32> OldToNewIndex;
+	OldToNewIndex.Init(INDEX_NONE, SavedInstances.Num());
+	for (int32 NewInstanceIndex = 0; NewInstanceIndex < UnsortedNewInstances.Num(); ++NewInstanceIndex)
+	{
+		UObject* NewInstance = UnsortedNewInstances[NewInstanceIndex];
+		if (const int32* pOldInstanceIndex = OldInstanceMap.Find(NewInstance->GetPathName(NewObject)))
+		{
+			OldToNewIndex[*pOldInstanceIndex] = NewInstanceIndex;
+		}
+		else
+		{
+			// this instance is in the new instance but wasn't in the old instance map. This isn't terribly unexpected so just keep it in the array
+			InOutSubobjectArray.Add(NewInstance);
+		}
+	}
+
+	for (int32 OldInstanceIndex = 0; OldInstanceIndex < SavedInstances.Num(); ++OldInstanceIndex)
+	{
+		if (OldToNewIndex[OldInstanceIndex] != INDEX_NONE)
+		{
+			InOutSubobjectArray.Add(UnsortedNewInstances[OldToNewIndex[OldInstanceIndex]]);
+			continue;
+		}
+		// found a subobject that existed in the old graph and isn't in the new one. We'll create it so we don't loose that data
+		const FInstancedObjectRecord& Record = SavedInstances[OldInstanceIndex];
+
+		// determine the outer of this new object
+		UObject* DestinationOuter = nullptr;
+		// attempt to find outer via local data
+		if (const int32* pOldOuterIndex = OldInstanceMap.Find(Record.OldInstance->GetOuter()->GetPathName(OldObject)))
+		{
+			const int32 NewOuterIndex = OldToNewIndex[*pOldOuterIndex];
+			if (UnsortedNewInstances.IsValidIndex(NewOuterIndex))
+			{
+				DestinationOuter = UnsortedNewInstances[NewOuterIndex];
+			}
+		}
+		// fallback to provided CPFUO replacement map
+		if (Params.OptionalReplacementMappings && DestinationOuter == nullptr)
+		{
+			if (UObject* const* FoundOuter = Params.OptionalReplacementMappings->Find(Record.OldInstance->GetOuter()))
+			{
+				DestinationOuter = *FoundOuter;
+			}
+		}
+		// fallback to topmost new object
+		if (DestinationOuter == nullptr)
+		{
+			DestinationOuter = NewObject;
+		}
+
+		// determine the class of this new object
+		UClass* OldClass = Record.OldInstance->GetClass();
+		UClass* NewClass = OldClass;
+		if (Params.OptionalOldToNewClassMappings)
+		{
+			if (UClass* const* FoundClass = Params.OptionalOldToNewClassMappings->Find(NewClass))
+			{
+				NewClass = *FoundClass;
+			}
+		}
+		if (NewClass->HasAnyClassFlags(CLASS_NewerVersionExists))
+		{
+			// The old subobject class is a compiler artifact, with no reinstanced counterpart. Its source may have been renamed (or deleted) and thus it could not be recompiled.
+			// In this case, if property bag features are enabled, we attempt to redirect to a placeholder type so we can serialize to a property bag in order to preserve its data.
+			if (UE::FPropertyBagRepository::IsPropertyBagPlaceholderObjectFeatureEnabled(UE::FPropertyBagRepository::EPlaceholderObjectFeature::ReplaceMissingReinstancedTypes))
+			{
+				FTopLevelAssetPath OldSubObjectClassPath = NewClass->GetReinstancedClassPathName();
+				if (OldSubObjectClassPath.IsValid())
+				{
+					if (UPackage* OldSubObjectClassPackage = FindPackage(nullptr, *OldSubObjectClassPath.GetPackageName().ToString()))
+					{
+						NewClass = UE::FPropertyBagRepository::CreatePropertyBagPlaceholderClass(OldSubObjectClassPackage, NewClass->GetClass(), OldSubObjectClassPath.GetAssetName(), NewClass->GetFlags() & ~RF_Transient);
+					}
+				}
+			}
+		}
+
+		FName NewName = Record.OldInstance->GetFName();
+		EObjectFlags ObjectFlags = Record.OldInstance->GetFlags() & RF_Reinst;
+
+		// Making sure the template we are about to use had all its references patched up
+		UObject* Template = NewClass->GetDefaultObject();
+		if (Params.OptionalReplacementMappings)
+		{
+			FArchiveReplaceObjectRef<UObject> ReplaceInCDOAr(Template, *Params.OptionalReplacementMappings);
+		}
+
+		UObject* NewInstance = ::NewObject<UObject>(DestinationOuter, NewClass, NewName, ObjectFlags, Template);
+		InOutSubobjectArray.Add(NewInstance);
+		OldToNewIndex[OldInstanceIndex] = UnsortedNewInstances.Num();
+		UnsortedNewInstances.Add(NewInstance);
+	}
+}
+#endif
+
 void UEngine::CopyPropertiesForUnrelatedObjects(UObject* OldObject, UObject* NewObject, FCopyPropertiesForUnrelatedObjectsParams Params)
 {
 	// UObject::CollectDefaultSubobjects will not return all DSOs for the CDO (only returns inherited subobjects). It also
@@ -18600,6 +18706,13 @@ void UEngine::CopyPropertiesForUnrelatedObjects(UObject* OldObject, UObject* New
 	{
 		CollectAllSubobjects( NewObject, ComponentsOnNewObject );
 
+#if WITH_EDITOR
+		if (UE::IsInstanceDataObjectSupportEnabled(NewObject))
+		{
+			CreateLooseSubobjectsForCPFUO(OldObject, NewObject, ComponentsOnNewObject, OldInstanceMap, SavedInstances, Params);
+		}
+#endif
+
 		// populate the ReferenceReplacementMap 
 		for (int32 Index = 0; Index < ComponentsOnNewObject.Num(); Index++)
 		{
