UWorld: CookDeterminism: Fix indeterminism in the ModifiedDate tag.
#rnx

#Synced-CL 51719300

diff --git a/Engine/Source/Runtime/Engine/Private/World.cpp b/Engine/Source/Runtime/Engine/Private/World.cpp
--- a/Engine/Source/Runtime/Engine/Private/World.cpp
+++ b/Engine/Source/Runtime/Engine/Private/World.cpp
@@ -10277,13 +10277,22 @@ void UWorld::GetAssetRegistryTags(FAssetRegistryTagsContext Context) const
 		}
 	}
 
-	// Get the full file path with extension
-	FString FullFilePath;
-	if (FPackageName::TryConvertLongPackageNameToFilename(GetOutermost()->GetName(), FullFilePath, FPackageName::GetMapPackageExtension()))
+	FName DateModifiedTagName("DateModified");
+	if (Context.IsSaving() && !Context.IsProceduralSave())
 	{
-		// Save/Display the modify date. File size is handled generically for all packages
-		FDateTime AssetDateModified = IFileManager::Get().GetTimeStamp(*FullFilePath);
-		Context.AddTag(FAssetRegistryTag("DateModified", AssetDateModified.ToString(), FAssetRegistryTag::TT_Chronological, FAssetRegistryTag::TD_Date));
+		FDateTime AssetDateModified = FDateTime::Now();
+		Context.AddTag(FAssetRegistryTag(DateModifiedTagName, AssetDateModified.ToString(), FAssetRegistryTag::TT_Chronological, FAssetRegistryTag::TD_Date));
+	}
+	else
+	{
+		// To prevent cook indeterminism, loads and procedural saves such as cook saves do not alter the DateModified
+		// tag if it already exists, and they set it to a deterministic value if it does not exist. Note that
+		// cook-generated worlds such as WorldPartition streaming chunks will have this deterministic value.
+		if (!Context.FindTag(DateModifiedTagName))
+		{
+			FDateTime ZeroTime = FDateTime::MinValue();
+			Context.AddTag(FAssetRegistryTag(DateModifiedTagName, ZeroTime.ToString(), FAssetRegistryTag::TT_Chronological, FAssetRegistryTag::TD_Date));
+		}
 	}
 
 	PRAGMA_DISABLE_DEPRECATION_WARNINGS;
