Adapt Deferred Decal invalid blend mode warnings for Substrate.

[RN] minor | Rendering Architecture

#Synced-CL 51712365

diff --git a/Engine/Source/Runtime/Engine/Private/Materials/MaterialShared.cpp b/Engine/Source/Runtime/Engine/Private/Materials/MaterialShared.cpp
--- a/Engine/Source/Runtime/Engine/Private/Materials/MaterialShared.cpp
+++ b/Engine/Source/Runtime/Engine/Private/Materials/MaterialShared.cpp
@@ -6445,7 +6445,11 @@ void FinalCompileValidation(
 		// We could make the change for the user but it would be confusing when going to DeferredDecal and back
 		// or we would have to pay a performance cost to make the change more transparently.
 		// The change saves performance as with translucency we don't need to test for MeshDecals in all opaque rendering passes
-		OutErrors.Add(TEXT("Material using the DeferredDecal domain can only use the BlendModes Translucent, AlphaComposite or Modulate"));
+		const TCHAR* AllowedBlendModes = bSubstrateEnabled ? 
+			TEXT("TranslucentGreyTransmittance, TranslucentColorTransmittance, AlphaComposite (Premultiplied Alpha), or ColoredTransmittanceOnly.") :
+			TEXT("Translucent, AlphaComposite(Premultiplied Alpha), or Modulate.");
+
+		OutErrors.Add(FString::Printf(TEXT("Material using the DeferredDecal domain can only use the Blend Modes %s"), AllowedBlendModes));
 	}
 
 	int32 NumPostProcessInputs = MaterialCompilationOutput.GetNumPostProcessInputsUsed();
