Substrate - clarify switch nodes that are not working with substrate with a more specific error message.

#tests PC editor material with switch node flwoing substrate data.
#rnx
#Synced-CL 51714093

diff --git a/Engine/Source/Runtime/Engine/Private/Materials/MaterialExpressions.cpp b/Engine/Source/Runtime/Engine/Private/Materials/MaterialExpressions.cpp
--- a/Engine/Source/Runtime/Engine/Private/Materials/MaterialExpressions.cpp
+++ b/Engine/Source/Runtime/Engine/Private/Materials/MaterialExpressions.cpp
@@ -364,6 +364,11 @@ static FAutoConsoleVariableRef CVarAllowSkipAlphaComputationForMaterialBlend(
 	ECVF_ReadOnly
 );
 
+static void ThrowSubstrateUnsupportedSwitchNodeError(struct FSubstrateTranslatorDataInterface& SubstrateTranslatorData, const FString& NodeClassName)
+{
+	SubstrateTranslatorData.Errorf(FString::Printf(TEXT("%s nodes do not support flowing Substrate data.\nThis is because the Substrate tree representing material topology must be known at translation time.\nBut this node uses branching with usage known at compilation time using #define or at runtime using dynamic branching."), *NodeClassName));
+}
+
 #endif // #if WITH_EDITOR
 
 // Helper functions to handle determinism during cooking
@@ -2453,8 +2458,7 @@ FSubstrateOperator* UMaterialExpression::SubstrateGenerateMaterialTopologyTree(c
 
 FSubstrateOperator* UMaterialExpression::SubstrateGenerateMaterialTopologyTree(struct FSubstrateTranslatorDataInterface& SubstrateTranslatorData, class UMaterialExpression* Parent, int32 OutputIndex)
 {
-	// Deprecated to call the general version.
-	SubstrateTranslatorData.Errorf(FString::Printf(TEXT("%s nodes type does not support generating/processing/flowing Substrate data.\nPlease reach out to the development team for feedback and if you want support to be added."), *GetClass()->GetName()));
+	SubstrateTranslatorData.Errorf(FString::Printf(TEXT("%s nodes type does not support generating/processing/flowing Substrate data.\nPlease reach out to your development team for implementation, if it is supported and can be implemented."), *GetClass()->GetName()));
 	return nullptr;
 }
 
@@ -9474,6 +9478,12 @@ void UMaterialExpressionDataDrivenShaderPlatformInfoSwitch::PostEditChangeProper
 
 	Super::PostEditChangeProperty(PropertyChangedEvent);
 }
+
+FSubstrateOperator* UMaterialExpressionDataDrivenShaderPlatformInfoSwitch::SubstrateGenerateMaterialTopologyTree(struct FSubstrateTranslatorDataInterface& SubstrateTranslatorData, class UMaterialExpression* Parent, int32 OutputIndex)
+{
+	ThrowSubstrateUnsupportedSwitchNodeError(SubstrateTranslatorData, *GetClass()->GetName());
+	return nullptr;
+}
 #endif // WITH_EDITOR
 
 //
@@ -17739,6 +17749,12 @@ bool UMaterialExpressionRayTracingQualitySwitch::IsResultMaterialAttributes(int3
 	return false;
 }
 
+FSubstrateOperator* UMaterialExpressionRayTracingQualitySwitch::SubstrateGenerateMaterialTopologyTree(struct FSubstrateTranslatorDataInterface& SubstrateTranslatorData, class UMaterialExpression* Parent, int32 OutputIndex)
+{
+	ThrowSubstrateUnsupportedSwitchNodeError(SubstrateTranslatorData, *GetClass()->GetName());
+	return nullptr;
+}
+
 void UMaterialExpressionRayTracingQualitySwitch::GetCaption(TArray<FString>& OutCaptions) const
 {
 	OutCaptions.Add(TEXT("RayTracingQualitySwitchReplace"));
@@ -17805,6 +17821,12 @@ EMaterialValueType UMaterialExpressionPathTracingQualitySwitch::GetInputValueTyp
 {
 	return MCT_Unknown;
 }
+
+FSubstrateOperator* UMaterialExpressionPathTracingQualitySwitch::SubstrateGenerateMaterialTopologyTree(struct FSubstrateTranslatorDataInterface& SubstrateTranslatorData, class UMaterialExpression* Parent, int32 OutputIndex)
+{
+	ThrowSubstrateUnsupportedSwitchNodeError(SubstrateTranslatorData, *GetClass()->GetName());
+	return nullptr;
+}
 #endif // WITH_EDITOR
 
 
@@ -17847,6 +17869,11 @@ void UMaterialExpressionPathTracingRayTypeSwitch::GetCaption(TArray<FString>& Ou
 	OutCaptions.Add(TEXT("PathTracingRayTypeSwitch"));
 }
 
+FSubstrateOperator* UMaterialExpressionPathTracingRayTypeSwitch::SubstrateGenerateMaterialTopologyTree(struct FSubstrateTranslatorDataInterface& SubstrateTranslatorData, class UMaterialExpression* Parent, int32 OutputIndex)
+{
+	ThrowSubstrateUnsupportedSwitchNodeError(SubstrateTranslatorData, *GetClass()->GetName());
+	return nullptr;
+}
 #endif // WITH_EDITOR
 
 ///////////////////////////////////////////////////////////////////////////////
diff --git a/Engine/Source/Runtime/Engine/Public/Materials/MaterialExpressionDataDrivenShaderPlatformInfoSwitch.h b/Engine/Source/Runtime/Engine/Public/Materials/MaterialExpressionDataDrivenShaderPlatformInfoSwitch.h
--- a/Engine/Source/Runtime/Engine/Public/Materials/MaterialExpressionDataDrivenShaderPlatformInfoSwitch.h
+++ b/Engine/Source/Runtime/Engine/Public/Materials/MaterialExpressionDataDrivenShaderPlatformInfoSwitch.h
@@ -68,6 +68,7 @@ class UMaterialExpressionDataDrivenShaderPlatformInfoSwitch : public UMaterialEx
 	virtual bool IsInputConnectionRequired(int32 InputIndex) const override;
 	virtual FName GetInputName(int32 InputIndex) const override;
 	virtual bool IsResultMaterialAttributes(int32 OutputIndex) override;
+	virtual FSubstrateOperator* SubstrateGenerateMaterialTopologyTree(struct FSubstrateTranslatorDataInterface& SubstrateTranslatorData, class UMaterialExpression* Parent, int32 OutputIndex) override;
 	virtual EMaterialValueType GetInputValueType(int32 InputIndex) override { return MCT_Unknown; }
 	virtual EMaterialValueType GetOutputValueType(int32 OutputIndex) override { return MCT_Unknown; }
 #endif // WITH_EDITOR
diff --git a/Engine/Source/Runtime/Engine/Public/Materials/MaterialExpressionPathTracingQualitySwitch.h b/Engine/Source/Runtime/Engine/Public/Materials/MaterialExpressionPathTracingQualitySwitch.h
--- a/Engine/Source/Runtime/Engine/Public/Materials/MaterialExpressionPathTracingQualitySwitch.h
+++ b/Engine/Source/Runtime/Engine/Public/Materials/MaterialExpressionPathTracingQualitySwitch.h
@@ -32,6 +32,7 @@ public:
 	virtual bool IsResultMaterialAttributes(int32 OutputIndex) override;
 	virtual EMaterialValueType GetInputValueType(int32 InputIndex) override;
 	virtual EMaterialValueType GetOutputValueType(int32 OutputIndex) override { return MCT_Unknown; }
+	virtual FSubstrateOperator* SubstrateGenerateMaterialTopologyTree(struct FSubstrateTranslatorDataInterface& SubstrateTranslatorData, class UMaterialExpression* Parent, int32 OutputIndex) override;
 
 #endif
 	//~ End UMaterialExpression Interface
diff --git a/Engine/Source/Runtime/Engine/Public/Materials/MaterialExpressionPathTracingRayTypeSwitch.h b/Engine/Source/Runtime/Engine/Public/Materials/MaterialExpressionPathTracingRayTypeSwitch.h
--- a/Engine/Source/Runtime/Engine/Public/Materials/MaterialExpressionPathTracingRayTypeSwitch.h
+++ b/Engine/Source/Runtime/Engine/Public/Materials/MaterialExpressionPathTracingRayTypeSwitch.h
@@ -44,6 +44,7 @@ public:
 	virtual bool IsResultMaterialAttributes(int32 OutputIndex) override;
 	virtual EMaterialValueType GetInputValueType(int32 InputIndex) override { return MCT_Unknown; }
 	virtual EMaterialValueType GetOutputValueType(int32 OutputIndex) override { return MCT_Unknown; }
+	virtual FSubstrateOperator* SubstrateGenerateMaterialTopologyTree(struct FSubstrateTranslatorDataInterface& SubstrateTranslatorData, class UMaterialExpression* Parent, int32 OutputIndex) override;
 #endif
 	//~ End UMaterialExpression Interface
 };
diff --git a/Engine/Source/Runtime/Engine/Public/Materials/MaterialExpressionRayTracingQualitySwitch.h b/Engine/Source/Runtime/Engine/Public/Materials/MaterialExpressionRayTracingQualitySwitch.h
--- a/Engine/Source/Runtime/Engine/Public/Materials/MaterialExpressionRayTracingQualitySwitch.h
+++ b/Engine/Source/Runtime/Engine/Public/Materials/MaterialExpressionRayTracingQualitySwitch.h
@@ -30,6 +30,7 @@ public:
 	virtual int32 Compile(class FMaterialCompiler* Compiler, int32 OutputIndex) override;
 	virtual void GetCaption(TArray<FString>& OutCaptions) const override;
 	virtual bool IsResultMaterialAttributes(int32 OutputIndex) override;
+	virtual FSubstrateOperator* SubstrateGenerateMaterialTopologyTree(struct FSubstrateTranslatorDataInterface& SubstrateTranslatorData, class UMaterialExpression* Parent, int32 OutputIndex) override;
 	virtual EMaterialValueType GetInputValueType(int32 InputIndex) override;
 	virtual EMaterialValueType GetOutputValueType(int32 OutputIndex) override { return MCT_Unknown; }
 
