[Streamables] [JITAsyncLoader] Make FStreamableAsyncLoaderRequestQueuedDelegate a thread safe delegate to avoid race conditions between producer(game thread)/consumer(asyncloadingthread). Ensure initial batch is dequeued being kicked to avoid contention with async loading thread.

[RN] [Streamables] Fixed a data race with the JIT async loader that could occur when kicking the initial update while also processing next requests on the async loading thread.
#tests DefaultSquad (PC) and ReplayRuns on multiple platforms to ensure we do not hit the issue
#Synced-CL 51619317

diff --git a/Engine/Source/Runtime/Engine/Private/StreamableAsyncLoader.cpp b/Engine/Source/Runtime/Engine/Private/StreamableAsyncLoader.cpp
--- a/Engine/Source/Runtime/Engine/Private/StreamableAsyncLoader.cpp
+++ b/Engine/Source/Runtime/Engine/Private/StreamableAsyncLoader.cpp
@@ -98,14 +98,25 @@ void FStreamableJustInTimeAsyncLoader::StartAsyncLoadRequests()
 
 	const float ToBatch = GStreamableJITAsyncLoadingInitialBatchingFactor * AsyncLoadRequests.Num();
 	int32 InitialBatchSize = FMath::Max(FMath::CeilToInt(ToBatch), 1);
+
+	// Dequeue requests from queue before kicking off async loading so as to prevent any contention with the ALT.
+	TArray<FStreamableJITAsyncLoadRequest> InitialBatchRequests;
+	InitialBatchRequests.Reserve(InitialBatchSize);
 	while(InitialBatchSize > 0)
 	{
-		if (!KickNextLoadingRequest())
+		FStreamableJITAsyncLoadRequest AsyncLoadRequest;
+		if (!AsyncLoadRequests.Dequeue(AsyncLoadRequest))
 		{
 			break;
 		}
+		InitialBatchRequests.Emplace(MoveTemp(AsyncLoadRequest));
 		--InitialBatchSize;
 	}
+
+	for (const FStreamableJITAsyncLoadRequest& Request : InitialBatchRequests)
+	{
+		KickLoadingRequest(Request);
+	}
 }
 
 void FStreamableJustInTimeAsyncLoader::MakeProgress(const FLoadPackageAsyncProgressParams& InParams)
@@ -157,23 +168,28 @@ bool FStreamableJustInTimeAsyncLoader::KickNextLoadingRequest()
 	FStreamableJITAsyncLoadRequest AsyncLoadRequest;
 	if(AsyncLoadRequests.Dequeue(AsyncLoadRequest))
 	{
-		const int32 LoadPriority = Priority.load(std::memory_order_relaxed);
-		FLoadPackageAsyncOptionalParams LoadPackageAsyncOptionalParams =
-		{
-			.CompletionDelegate = MakeUnique<FLoadPackageAsyncDelegate>(AsyncLoadRequest.CompletionDelegate),
-			.ProgressDelegate = ProgressDelegate,
-			.PackagePriority = LoadPriority // Queue using updated priorities.
-		};
-		int32 RequestID = LoadPackageAsync(AsyncLoadRequest.PackagePath, MoveTemp(LoadPackageAsyncOptionalParams));
-		AsyncLoadRequest.RequestQueuedDelegate.ExecuteIfBound(RequestID);
-		UE_LOG(LogStreamableAsyncLoading, Verbose, TEXT("FStreamableJustInTimeAsyncLoader queued async load request [Package:%s] [Priority: %d] [RequestID]:%d"), *AsyncLoadRequest.PackagePath.GetDebugName(), LoadPriority, RequestID);
-		TotalRemainingRequests.fetch_sub(1, std::memory_order_relaxed);
+		KickLoadingRequest(AsyncLoadRequest);
 		return true;
 	}
 
 	return false;
 }
 
+void FStreamableJustInTimeAsyncLoader::KickLoadingRequest(const FStreamableJITAsyncLoadRequest& AsyncLoadRequest)
+{
+	const int32 LoadPriority = Priority.load(std::memory_order_relaxed);
+	FLoadPackageAsyncOptionalParams LoadPackageAsyncOptionalParams =
+	{
+		.CompletionDelegate = MakeUnique<FLoadPackageAsyncDelegate>(AsyncLoadRequest.CompletionDelegate),
+		.ProgressDelegate = ProgressDelegate,
+		.PackagePriority = LoadPriority // Queue using updated priorities.
+	};
+	int32 RequestID = LoadPackageAsync(AsyncLoadRequest.PackagePath, MoveTemp(LoadPackageAsyncOptionalParams));
+	AsyncLoadRequest.RequestQueuedDelegate.ExecuteIfBound(RequestID);
+	UE_LOG(LogStreamableAsyncLoading, Verbose, TEXT("FStreamableJustInTimeAsyncLoader queued async load request [Package:%s] [Priority: %d] [RequestID]:%d"), *AsyncLoadRequest.PackagePath.GetDebugName(), LoadPriority, RequestID);
+	TotalRemainingRequests.fetch_sub(1, std::memory_order_relaxed);
+}
+
 void FStreamableJustInTimeAsyncLoader::UpdateHandleState()
 {
 	// Expected to be called on the game thread (see case EAsyncLoadingProgress::FullyLoaded).
diff --git a/Engine/Source/Runtime/Engine/Private/StreamableAsyncLoader.h b/Engine/Source/Runtime/Engine/Private/StreamableAsyncLoader.h
--- a/Engine/Source/Runtime/Engine/Private/StreamableAsyncLoader.h
+++ b/Engine/Source/Runtime/Engine/Private/StreamableAsyncLoader.h
@@ -7,7 +7,7 @@
 #include "Containers/SpscQueue.h"
 #include "UObject/UObjectGlobals.h"
 
-DECLARE_DELEGATE_OneParam(FStreamableAsyncLoaderRequestQueuedDelegate, int32 /*RequestID*/)
+using FStreamableAsyncLoaderRequestQueuedDelegate = TTSDelegate<void(int32)>;
 
 struct FStreamableHandle;
 
@@ -63,24 +63,27 @@ public:
 	// End IStreamableAsyncLoader interface.
 
 private:
+	struct FStreamableJITAsyncLoadRequest
+	{
+		FPackagePath PackagePath;
+		FStreamableAsyncLoaderRequestQueuedDelegate RequestQueuedDelegate;
+		FLoadPackageAsyncDelegate CompletionDelegate;
+	};
+	using FStreamableJITAsyncLoadRequests = TSpscQueue<FStreamableJITAsyncLoadRequest>;
+
 	// Can be called on game thread or async loading thread.
 	void MakeProgress(const FLoadPackageAsyncProgressParams& InParams);
 
 	// Can be called on game thread (for initial batch) followed by subsequent calls being made on the async loading thread.
 	bool KickNextLoadingRequest();
+	void KickLoadingRequest(const FStreamableJITAsyncLoadRequest& AsyncLoadRequest);
 
 	void UpdateHandleState();
 	
 	void Cancel();
 	void Cleanup();
 		
-	struct FStreamableJITAsyncLoadRequest
-	{
-		FPackagePath PackagePath;
-		FStreamableAsyncLoaderRequestQueuedDelegate RequestQueuedDelegate;
-		FLoadPackageAsyncDelegate CompletionDelegate;
-	};
-	TSpscQueue<FStreamableJITAsyncLoadRequest> AsyncLoadRequests;
+	FStreamableJITAsyncLoadRequests AsyncLoadRequests;
 	TSharedPtr<FLoadPackageAsyncProgressDelegate> ProgressDelegate;
 	std::atomic<TAsyncLoadPriority> Priority = 0;
 	std::atomic<int32> TotalRemainingRequests = 0;
