Newton Impulse Component Thread Safety Issues AND Entity Compat

#changelist validated
#virtualized
#Synced-CL 51720594

diff --git a/Engine/Source/Runtime/Engine/Private/Collision/CollisionConversions.cpp b/Engine/Source/Runtime/Engine/Private/Collision/CollisionConversions.cpp
--- a/Engine/Source/Runtime/Engine/Private/Collision/CollisionConversions.cpp
+++ b/Engine/Source/Runtime/Engine/Private/Collision/CollisionConversions.cpp
@@ -162,14 +162,20 @@ static void SetHitResultFromShapeAndFaceIndex(const FPhysicsShape& Shape,  const
 
 		//Normal case where we hit a body
 		OutResult.Item = BodyInst->InstanceBodyIndex;
-		const UBodySetupCore* BodySetup = BodyInst->BodySetup.Get();	//this data should be immutable at runtime so ok to check from worker thread.
+		const UBodySetupCore* BodySetup = BodyInst->BodySetup.Get(); //this data should be immutable at runtime so ok to check from worker thread.
 		if (BodySetup)
 		{
 			OutResult.BoneName = BodySetup->BoneName;
 		}
 
 		OwningComponent = BodyInst->OwnerComponent;
-		OutResult.PhysicsObject = BodyInst->GetPhysicsActor() ? BodyInst->GetPhysicsActor()->GetPhysicsObject() : nullptr;
+
+		// Use GetActorReferenceWithWelding() instead of GetPhysicsActor() so that
+		// welded children resolve to their parent's PhysicsObject rather than returning null.
+		// When a body is welded, its shapes are merged into the parent particle,
+		// so the parent's PhysicsObject is the correct one for force/impulse application.
+		const FPhysicsActorHandle& PhysicsActor = BodyInst->GetActorReferenceWithWelding();
+		OutResult.PhysicsObject = PhysicsActor ? PhysicsActor->GetPhysicsObject() : nullptr;
 	}
 	else if (Chaos::FPhysicsObject* PhysObject = GetPhysicsObjectFromUserData(Actor))
 	{
@@ -700,7 +706,13 @@ void ConvertQueryOverlap(const FPhysicsShape& Shape, const FPhysicsActor& Actor,
 			OutOverlap.OverlapObjectHandle = FActorInstanceHandle::MakeActorHandleToResolve(BodyInst->OwnerComponent, BodyInst->InstanceBodyIndex);
 			OutOverlap.Component = BodyInst->OwnerComponent; // Copying weak pointer is faster than assigning raw pointer.
 		}
-		OutOverlap.PhysicsObject = BodyInst->GetPhysicsActor() ? BodyInst->GetPhysicsActor()->GetPhysicsObject() : nullptr;
+
+		// Use GetActorReferenceWithWelding() instead of GetPhysicsActor() so that
+		// welded children resolve to their parent's PhysicsObject rather than returning null.
+		// When a body is welded, its shapes are merged into the parent particle,
+		// so the parent's PhysicsObject is the correct one for force/impulse application.
+		const FPhysicsActorHandle& PhysicsActor = BodyInst->GetActorReferenceWithWelding();
+		OutOverlap.PhysicsObject = PhysicsActor ? PhysicsActor->GetPhysicsObject() : nullptr;
 
 		if (OutOverlap.PhysicsObject != nullptr)
 		{
