Remove  misspelled UPrimtiiveComponent.

#Synced-CL 51536473
#Synced-CL 51536473

diff --git a/Engine/Source/Runtime/Engine/Public/PrimitiveSceneDesc.h b/Engine/Source/Runtime/Engine/Public/PrimitiveSceneDesc.h
--- a/Engine/Source/Runtime/Engine/Public/PrimitiveSceneDesc.h
+++ b/Engine/Source/Runtime/Engine/Public/PrimitiveSceneDesc.h
@@ -12,7 +12,6 @@ class FPrimitiveSceneProxy;
 struct FPrimitiveSceneProxyDesc;
 class HHitProxy;
 class IPrimitiveComponent;
-class UPrimtiiveComponent;
 
 /**
 * FPrimitiveSceneDesc is a structure that can be used to Add/Remove/Update primitives in an FScene. 
