BodySetup: Cooking: During TryGetCookedData, enter a lock around the read of CookedFormatData in addition to its write, to prevent threads crashing when deciding whether they need to enter the write lock due to another thread already in it and mutating the CookedFormatData.
Also, enter the writelock only after blocking on Compile:

PhysicsDerivedCooker->GatherKeyData() -> UBodySetup::GetGeometryDDCKey() -> UStaticMesh::GetMeshId() -> UStaticMesh::GetRenderData() -> WaitUntilAsyncPropertyReleased(EStaticMeshAsyncProperties::RenderData)

to prevent deadlock due to ExecutePostLoadInternal entering the
WaitUntilAsyncPropertyReleased(EStaticMeshAsyncProperties::RenderData) lock first followed by calling TryGetCookedData.

Also add a fix for memory leak of the FChaosDerivedDataCooker used in building, in the case of some configuration settings that disable loading it from DDC.
#rnx
#Synced-CL 51748961

diff --git a/Engine/Source/Runtime/Engine/Classes/PhysicsEngine/BodySetup.h b/Engine/Source/Runtime/Engine/Classes/PhysicsEngine/BodySetup.h
--- a/Engine/Source/Runtime/Engine/Classes/PhysicsEngine/BodySetup.h
+++ b/Engine/Source/Runtime/Engine/Classes/PhysicsEngine/BodySetup.h
@@ -268,7 +268,7 @@ class UBodySetup : public UBodySetupCore
 	 * This mutex protects the cook path in TryGetCookedData so we don't have multiple threads attempting to
 	 * write to the format data
 	 */
-	FCriticalSection CookedFormatDataMutex;
+	FRWLock CookedFormatDataMutex;
 #endif
 
 	/** GUID used to uniquely identify this setup so it can be found in the DDC */
diff --git a/Engine/Source/Runtime/Engine/Private/PhysicsEngine/BodySetup.cpp b/Engine/Source/Runtime/Engine/Private/PhysicsEngine/BodySetup.cpp
--- a/Engine/Source/Runtime/Engine/Private/PhysicsEngine/BodySetup.cpp
+++ b/Engine/Source/Runtime/Engine/Private/PhysicsEngine/BodySetup.cpp
@@ -1348,12 +1348,12 @@ UE_TRACE_EVENT_BEGIN(Cpu, BodySetupDDCFetch, NoSync)
 	UE_TRACE_EVENT_FIELD(UE::Trace::WideString, FetchReason)
 UE_TRACE_EVENT_END()
 
-template<typename DDCBuilderType>
-UBodySetup::EBodySetupCookResult TryGetDDCBuiltData(FByteBulkData* OutResult, DDCBuilderType& InBuilder, UBodySetup* InSetup, bool bInIsRuntime)
+UBodySetup::EBodySetupCookResult TryGetDDCBuiltData(FByteBulkData* OutResult, TUniquePtr<FChaosDerivedDataCooker>&& InBuilder, UBodySetup* InSetup, bool bInIsRuntime)
 {
 	TArray<uint8> OutData;
+	check(InBuilder);
 
-	if(InBuilder.CanBuild())
+	if (InBuilder->CanBuild())
 	{
 		COOK_STAT(FCookStats::FScopedStatsCounter Timer = BodySetupCookStats::UsageStats.TimeSyncWork());
 
@@ -1370,7 +1370,9 @@ UBodySetup::EBodySetupCookResult TryGetDDCBuiltData(FByteBulkData* OutResult, DD
 			UE_TRACE_LOG_SCOPED_T(Cpu, BodySetupDDCFetch, CpuChannel)
 				<< BodySetupDDCFetch.FetchReason(*FetchReason);
 #endif
-			bDDCHit = GetDerivedDataCacheRef().GetSynchronous(&InBuilder, OutData, &bDataWasBuilt);
+			// GetSynchronous takes ownership of and will delete the ReleasedBuilder
+			FChaosDerivedDataCooker* ReleasedBuilder = InBuilder.Release();
+			bDDCHit = GetDerivedDataCacheRef().GetSynchronous(ReleasedBuilder, OutData, &bDataWasBuilt);
 
 			// Only compute hit/miss if DDC was not skipped
 			COOK_STAT(Timer.AddHitOrMiss(!bDDCHit || bDataWasBuilt ? FCookStats::CallStats::EHitOrMiss::Miss : FCookStats::CallStats::EHitOrMiss::Hit, OutData.Num()));
@@ -1378,12 +1380,12 @@ UBodySetup::EBodySetupCookResult TryGetDDCBuiltData(FByteBulkData* OutResult, DD
 		else
 		{
 			bDataWasBuilt = true;
-			InBuilder.Build(OutData);
+			InBuilder->Build(OutData);
 			COOK_STAT(Timer.TrackCyclesOnly());
 		}
 	}
 
-	if(OutData.Num())
+	if (OutData.Num())
 	{
 		if(OutResult->IsLocked())
 		{
@@ -1403,13 +1405,6 @@ UBodySetup::EBodySetupCookResult TryGetDDCBuiltData(FByteBulkData* OutResult, DD
 	return UBodySetup::EBodySetupCookResult::NoError;
 }
 
-template<typename DDCBuilderType>
-void GetDDCBuiltData(FByteBulkData* OutResult, DDCBuilderType& InBuilder, UBodySetup* InSetup, bool bInIsRuntime)
-{
-	UBodySetup::EBodySetupCookResult Result = TryGetDDCBuiltData(OutResult, InBuilder, InSetup, bInIsRuntime);
-	check(Result != UBodySetup::EBodySetupCookResult::FailedToLock);
-}
-
 #endif //#if WITH_EDITOR
 
 FByteBulkData* UBodySetup::GetCookedData(FName Format)
@@ -1471,7 +1466,16 @@ UBodySetup::FBodySetupTryCookResult UBodySetup::TryGetCookedData(FName Format)
 	const bool bIsRuntime = IsRuntime(this);
 
 #if WITH_EDITOR
-	if(!UsedFormatData->Contains(Format))
+	bool bFormatDataContainsFormat;
+	{
+		FReadScopeLock ScopeLock(CookedFormatDataMutex);
+		bFormatDataContainsFormat = UsedFormatData->Contains(Format);
+		if (bFormatDataContainsFormat)
+		{
+			Result = &UsedFormatData->GetFormat(Format);
+		}
+	}
+	if (!bFormatDataContainsFormat)
 	{
 		SCOPE_CYCLE_COUNTER(STAT_PhysXCooking);
 
@@ -1485,35 +1489,40 @@ UBodySetup::FBodySetupTryCookResult UBodySetup::TryGetCookedData(FName Format)
 			return { nullptr, CookDataResult };
 		}
 
-		// Only one thread should build the physics data. Below calling GetFormat and TryGetDDCBuiltData will
-		// both mutate shared state so we protect this operation with a lock.
-		FScopeLock ScopeLock(&CookedFormatDataMutex);
+		// Create the builder and gather the data from our StaticMesh to build the DDC key.
+		// Gathering data from StaticMesh will block on the StaticMesh's compile if it is in progress, so do that
+		// gather outside of our lock.
+		// We do not want a FGCObject to be created to prevent garbage collection of our own UBodySetup*
+		// because the scope and lifetime is well defined and FGCObject can't be created on other threads
+		// during garbage collection, which would prevent this function from being run asynchronously.
+		const bool bUseRefHolder = false;
+		TUniquePtr<FChaosDerivedDataCooker> PhysicsDerivedCooker( new FChaosDerivedDataCooker(this, Format, bUseRefHolder));
+		PhysicsDerivedCooker->GatherKeyData();
+
+		// Now that we have waited for any static mesh compile to complete (or we are already on the static mesh compile
+		// thread), attempt to enter the write lock that we use to guard access to the physics data. Below calling GetFormat and
+		// TryGetDDCBuiltData will both mutate shared state.
+		FWriteScopeLock ScopeLock(CookedFormatDataMutex);
 
 		// Check the lock condition again, we only want the first thread in here to actually build the data,
 		// while it's building other threads should wait and proceed after the build completes.
-		// Double-checked locking avoids always taking the lock - we only do it for unbuilt data
-		if(!UsedFormatData->Contains(Format))
+		bFormatDataContainsFormat = UsedFormatData->Contains(Format);
+		if (bFormatDataContainsFormat)
 		{
-			// We do not want a FGCObject to be created to prevent garbage collection of our own UBodySetup*
-			// because the scope and lifetime is well defined and FGCObject can't be created on other threads
-			// during garbage collection, which would prevent this function from being run asynchronously.
-			const bool bUseRefHolder = false;
-			FChaosDerivedDataCooker* PhysicsDerivedCooker = new FChaosDerivedDataCooker(this, Format, bUseRefHolder);
-
+			// We were a waiting thread while another thread built the data, just take it from the container now
 			Result = &UsedFormatData->GetFormat(Format);
-			CookDataResult = TryGetDDCBuiltData(Result, *PhysicsDerivedCooker, this, bIsRuntime);
 		}
 		else
 		{
-			// We were a waiting thread while another thread built the data, just take it from the container now
 			Result = &UsedFormatData->GetFormat(Format);
+			CookDataResult = TryGetDDCBuiltData(Result, MoveTemp(PhysicsDerivedCooker), this, bIsRuntime);
 		}
 	}
-	else
-#endif // #if WITH_EDITOR
+#else // #if WITH_EDITOR
 	{
 		Result = &UsedFormatData->GetFormat(Format);
 	}
+#endif // #else !WITH_EDITOR
 
 	check(Result);
 
diff --git a/Engine/Source/Runtime/Engine/Private/PhysicsEngine/Experimental/ChaosDerivedData.cpp b/Engine/Source/Runtime/Engine/Private/PhysicsEngine/Experimental/ChaosDerivedData.cpp
--- a/Engine/Source/Runtime/Engine/Private/PhysicsEngine/Experimental/ChaosDerivedData.cpp
+++ b/Engine/Source/Runtime/Engine/Private/PhysicsEngine/Experimental/ChaosDerivedData.cpp
@@ -70,14 +70,25 @@ FString FChaosDerivedDataCooker::GetDebugContextString() const
 
 FString FChaosDerivedDataCooker::GetPluginSpecificCacheKeySuffix() const
 {
+	GatherKeyData();
+	return PluginSpecificCacheKeySuffix;
+}
+
+void FChaosDerivedDataCooker::GatherKeyData() const
+{
+	if (!PluginSpecificCacheKeySuffix.IsEmpty())
+	{
+		return;
+	}
+
 	FString SetupGeometryKey("INVALID");
 
-	if(Setup)
+	if (Setup)
 	{
 		Setup->GetGeometryDDCKey(SetupGeometryKey);
 	}
 
-	FString OutSuffix = FString::Printf(TEXT("%s_%s_REAL%d"),
+	PluginSpecificCacheKeySuffix = FString::Printf(TEXT("%s_%s_REAL%d"),
 										*RequestedFormat.ToString(),
 										*SetupGeometryKey,
 										(int)sizeof(Chaos::FReal));
@@ -86,10 +97,8 @@ FString FChaosDerivedDataCooker::GetPluginSpecificCacheKeySuffix() const
 	// Separate out arm keys as x64 and arm64 clang do not generate the same data for a given
 	// input. Add the arm specifically so that a) we avoid rebuilding the current DDC and
 	// b) we can remove it once we get arm64 to be consistent.
-	OutSuffix.Append(TEXT("_arm64"));
+	PluginSpecificCacheKeySuffix.Append(TEXT("_arm64"));
 #endif
-
-	return OutSuffix;
 }
 
 bool FChaosDerivedDataCooker::IsBuildThreadsafe() const
diff --git a/Engine/Source/Runtime/Engine/Private/PhysicsEngine/Experimental/ChaosDerivedData.h b/Engine/Source/Runtime/Engine/Private/PhysicsEngine/Experimental/ChaosDerivedData.h
--- a/Engine/Source/Runtime/Engine/Private/PhysicsEngine/Experimental/ChaosDerivedData.h
+++ b/Engine/Source/Runtime/Engine/Private/PhysicsEngine/Experimental/ChaosDerivedData.h
@@ -53,6 +53,7 @@ public:
 	//End FDerivedDataPluginInterface Interface
 
 	FChaosDerivedDataCooker(UBodySetup* InSetup, FName InFormat, bool bUseRefHolder = true);
+	void GatherKeyData() const;
 
 	bool CanBuild()
 	{
@@ -64,6 +65,7 @@ private:
 	TObjectPtr<UBodySetup> Setup;
 	FName RequestedFormat;
 	TUniquePtr<FGCObject> RefHolder;
+	mutable FString PluginSpecificCacheKeySuffix;
 };
 
 
