StreamingRenderAsset cleanup
- Removing unused "bWaitForMipFading" argument.
- Simplifying NumStreamedMips arguments by using an array view.

#Synced-CL 51783789

diff --git a/Engine/Source/Runtime/Engine/Private/Streaming/StreamingManagerTexture.cpp b/Engine/Source/Runtime/Engine/Private/Streaming/StreamingManagerTexture.cpp
--- a/Engine/Source/Runtime/Engine/Private/Streaming/StreamingManagerTexture.cpp
+++ b/Engine/Source/Runtime/Engine/Private/Streaming/StreamingManagerTexture.cpp
@@ -424,7 +424,7 @@ bool FRenderAssetStreamingManager::StreamOutRenderAssetData( int64 RequiredMemor
 			MemoryDropped += CurrentSize - MinimalSize; 
 			TempMemoryUsed += MinimalSize;
 
-			StreamingRenderAsset.UpdateStreamingStatus(false);
+			StreamingRenderAsset.UpdateStreamingStatus();
 
 			if (TempMemoryUsed >= MaxTempMemoryAllowed)
 			{
@@ -538,7 +538,7 @@ void FRenderAssetStreamingManager::TickFastResponseAssets()
 
 		check(RenderAsset == Asset.RenderAsset);
 
-		Asset.UpdateStreamingStatus(false);
+		Asset.UpdateStreamingStatus();
 
 		if (Asset.ResidentMips == Asset.RequestedMips && Asset.ResidentMips < Asset.MaxAllowedMips)
 		{
@@ -557,7 +557,7 @@ void FRenderAssetStreamingManager::TickFastResponseAssets()
 				}
 			}
 
-			Asset.UpdateStreamingStatus(false);
+			Asset.UpdateStreamingStatus();
 			TrackRenderAssetEvent(&Asset, RenderAsset, true, this);
 		}
 	}
@@ -597,16 +597,15 @@ void FRenderAssetStreamingManager::ProcessAddedRenderAssets()
 	// Add new textures or meshes.
 	TArray<FStreamingRenderAsset>& StreamingRenderAssets = GetStreamingRenderAssetsAsyncSafe();
 	StreamingRenderAssets.Reserve(StreamingRenderAssets.Num() + PendingStreamingRenderAssets.Num());
-	for (int32 Idx = 0; Idx < PendingStreamingRenderAssets.Num(); ++Idx)
+
+	for (UStreamableRenderAsset* Asset : PendingStreamingRenderAssets)
 	{
-		UStreamableRenderAsset* Asset = PendingStreamingRenderAssets[Idx];
 		// Could be null if it was removed after being added.
 		if (Asset)
 		{
 			Asset->StreamingIndex = StreamingRenderAssets.Num();
-			const int32* NumStreamedMips = nullptr;
-			const int32 NumLODGroups = GetNumStreamedMipsArray(Asset->GetRenderAssetType(), NumStreamedMips);
-			new (StreamingRenderAssets) FStreamingRenderAsset(Asset, Asset->SimpleStreamableAssetManagerIndex, NumStreamedMips, NumLODGroups, Settings);
+			TConstArrayView<int32> NumStreamedMips = GetNumStreamedMipsArray(Asset->GetRenderAssetType());
+			new (StreamingRenderAssets) FStreamingRenderAsset(Asset, Asset->SimpleStreamableAssetManagerIndex, NumStreamedMips, Settings);
 		}
 	}
 	PendingStreamingRenderAssets.Reset();
@@ -1395,10 +1394,8 @@ void FRenderAssetStreamingManager::UpdateIndividualRenderAsset( UStreamableRende
 	FStreamingRenderAsset* StreamingRenderAsset = GetStreamingRenderAsset(RenderAsset);
 	if (!StreamingRenderAsset) return;
 
-	const int32* NumStreamedMips;
-	const int32 NumLODGroups = GetNumStreamedMipsArray(StreamingRenderAsset->RenderAssetType, NumStreamedMips);
-
-	StreamingRenderAsset->UpdateDynamicData(NumStreamedMips, NumLODGroups, Settings, false);
+	const TConstArrayView<int32> NumStreamedMips = GetNumStreamedMipsArray(StreamingRenderAsset->RenderAssetType);
+	StreamingRenderAsset->UpdateDynamicData(NumStreamedMips, Settings);
 
 	if (StreamingRenderAsset->bForceFullyLoad) // Somewhat expected at this point.
 	{
@@ -1425,7 +1422,7 @@ bool FRenderAssetStreamingManager::FastForceFullyResident(UStreamableRenderAsset
 		const int32 StreamingIdx = RenderAsset->StreamingIndex;
 		FStreamingRenderAsset& Asset = StreamingRenderAssets[StreamingIdx];
 
-		Asset.UpdateStreamingStatus(false);
+		Asset.UpdateStreamingStatus();
 
 		if (Asset.ResidentMips < Asset.MaxAllowedMips)
 		{
@@ -1444,7 +1441,7 @@ bool FRenderAssetStreamingManager::FastForceFullyResident(UStreamableRenderAsset
  * @param StageIndex		Current stage index
  * @param NumUpdateStages	Number of texture update stages
  */
-void FRenderAssetStreamingManager::UpdateStreamingRenderAssets(int32 StageIndex, int32 NumUpdateStages, bool bWaitForMipFading, bool bAsync)
+void FRenderAssetStreamingManager::UpdateStreamingRenderAssets(int32 StageIndex, int32 NumUpdateStages, bool bAsync)
 {
 	TRACE_CPUPROFILER_EVENT_SCOPE(UpdateStreamingRenderAssets);
 	if ( StageIndex == 0 )
@@ -1506,7 +1503,7 @@ void FRenderAssetStreamingManager::UpdateStreamingRenderAssets(int32 StageIndex,
 			}
 		}
 
-		ParallelFor(TEXT("RenderAssetStreaming"), Packets.Num(), 1, [this, &Packets, &bWaitForMipFading](int32 PacketIndex)
+		ParallelFor(TEXT("RenderAssetStreaming"), Packets.Num(), 1, [this, &Packets](int32 PacketIndex)
 		{
 			FTaskTagScope Scope(ETaskTag::EParallelGameThread);
 
@@ -1522,10 +1519,9 @@ void FRenderAssetStreamingManager::UpdateStreamingRenderAssets(int32 StageIndex,
 
 				STAT(int32 PreviousResidentMips = StreamingRenderAsset.ResidentMips;)
 
-				const int32* NumStreamedMips;
-				const int32 NumLODGroups = GetNumStreamedMipsArray(StreamingRenderAsset.RenderAssetType, NumStreamedMips);
+				const TConstArrayView<int32> NumStreamedMips = GetNumStreamedMipsArray(StreamingRenderAsset.RenderAssetType);
 
-				StreamingRenderAsset.UpdateDynamicData(NumStreamedMips, NumLODGroups, Settings, bWaitForMipFading, &Packets[PacketIndex].LocalDeferredTickCBAssets); // We always use the Deferred CBs when doing the ParallelFor since those CBs are not thread safe.
+				StreamingRenderAsset.UpdateDynamicData(NumStreamedMips, Settings, &Packets[PacketIndex].LocalDeferredTickCBAssets); // We always use the Deferred CBs when doing the ParallelFor since those CBs are not thread safe.
 
 				// Make a list of each texture/mesh that can potentially require additional UpdateStreamingStatus
 				if (StreamingRenderAsset.RequestedMips != StreamingRenderAsset.ResidentMips)
@@ -1561,10 +1557,9 @@ void FRenderAssetStreamingManager::UpdateStreamingRenderAssets(int32 StageIndex,
 
 			STAT(int32 PreviousResidentMips = StreamingRenderAsset.ResidentMips;)
 
-			const int32* NumStreamedMips;
-			const int32 NumLODGroups = GetNumStreamedMipsArray(StreamingRenderAsset.RenderAssetType, NumStreamedMips);
+			const TConstArrayView<int32> NumStreamedMips = GetNumStreamedMipsArray(StreamingRenderAsset.RenderAssetType);
 
-			StreamingRenderAsset.UpdateDynamicData(NumStreamedMips, NumLODGroups, Settings, bWaitForMipFading, bAsync ? &DeferredTickCBAssets : nullptr);
+			StreamingRenderAsset.UpdateDynamicData(NumStreamedMips, Settings, bAsync ? &DeferredTickCBAssets : nullptr);
 
 			// Make a list of each texture/mesh that can potentially require additional UpdateStreamingStatus
 			if (StreamingRenderAsset.RequestedMips != StreamingRenderAsset.ResidentMips)
@@ -1888,17 +1883,11 @@ class FUpdateStreamingRenderAssetsTask
 	FRenderAssetStreamingManager* Manager;
 	int32 StageIdx;
 	int32 NumUpdateStages;
-	bool bWaitForMipFading;
 public:
-	FUpdateStreamingRenderAssetsTask(
-		FRenderAssetStreamingManager* InManager,
-		int32 InStageIdx,
-		int32 InNumUpdateStages,
-		bool bInWaitForMipFading)
+	FUpdateStreamingRenderAssetsTask(FRenderAssetStreamingManager* InManager, int32 InStageIdx, int32 InNumUpdateStages)
 		: Manager(InManager)
 		, StageIdx(InStageIdx)
 		, NumUpdateStages(InNumUpdateStages)
-		, bWaitForMipFading(bInWaitForMipFading)
 	{
 	}
 	static FORCEINLINE TStatId GetStatId()
@@ -1916,7 +1905,7 @@ public:
 	void DoTask(ENamedThreads::Type CurrentThread, const FGraphEventRef& MyCompletionGraphEvent)
 	{
 		FTaskTagScope Scope(ETaskTag::EParallelGameThread);
-		Manager->UpdateStreamingRenderAssets(StageIdx, NumUpdateStages, bWaitForMipFading, true);
+		Manager->UpdateStreamingRenderAssets(StageIdx, NumUpdateStages, true);
 	}
 };
 
@@ -1962,7 +1951,7 @@ void FRenderAssetStreamingManager::UpdateResourceStreaming( float DeltaTime, boo
 
 		// Update Thread Data
 		SetLastUpdateTime();
-		UpdateStreamingRenderAssets(0, 1, false);
+		UpdateStreamingRenderAssets(0, 1);
 
 		UpdatePendingStates(true);
 		PrepareAsyncTask(bProcessEverything || Settings.bStressTest);
@@ -2037,11 +2026,11 @@ void FRenderAssetStreamingManager::UpdateResourceStreaming( float DeltaTime, boo
 			}
 
 			StreamingRenderAssetsSyncEvent = TGraphTask<FUpdateStreamingRenderAssetsTask>::CreateTask(nullptr, ENamedThreads::GameThread)
-			.ConstructAndDispatchWhenReady(this, ProcessingStage - 1, NumRenderAssetProcessingStages, DeltaTime > 0.f);
+			.ConstructAndDispatchWhenReady(this, ProcessingStage - 1, NumRenderAssetProcessingStages);
 		}
 		else
 		{
-			UpdateStreamingRenderAssets(ProcessingStage - 1, NumRenderAssetProcessingStages, DeltaTime > 0.f);
+			UpdateStreamingRenderAssets(ProcessingStage - 1, NumRenderAssetProcessingStages);
 		}
 
 		IncrementalUpdate(1.f / (float)FMath::Max(NumRenderAssetProcessingStages - 1, 1), true); // -1 since we don't want to do anything at stage 0.
@@ -2063,7 +2052,7 @@ void FRenderAssetStreamingManager::UpdateResourceStreaming( float DeltaTime, boo
 		// Since this step is lightweight, tick each texture inflight here, to accelerate the state changes.
 		for (int32 TextureIndex : InflightRenderAssets)
 		{
-			StreamingRenderAssets[TextureIndex].UpdateStreamingStatus(DeltaTime > 0);
+			StreamingRenderAssets[TextureIndex].UpdateStreamingStatus();
 		}
 
 		TickFastResponseAssets();
@@ -2137,7 +2126,7 @@ int32 FRenderAssetStreamingManager::BlockTillAllRequestsFinished( float TimeLimi
 
 		for (FStreamingRenderAsset& StreamingRenderAsset : StreamingRenderAssets)
 		{
-			const bool bValid = StreamingRenderAsset.UpdateStreamingStatus(false).IsValid();
+			const bool bValid = StreamingRenderAsset.UpdateStreamingStatus().IsValid();
 			if (bValid && StreamingRenderAsset.RequestedMips != StreamingRenderAsset.ResidentMips)
 			{
 				++NumOfInFlights;
diff --git a/Engine/Source/Runtime/Engine/Private/Streaming/StreamingManagerTexture.h b/Engine/Source/Runtime/Engine/Private/Streaming/StreamingManagerTexture.h
--- a/Engine/Source/Runtime/Engine/Private/Streaming/StreamingManagerTexture.h
+++ b/Engine/Source/Runtime/Engine/Private/Streaming/StreamingManagerTexture.h
@@ -217,7 +217,7 @@ private:
 		 * @param NumUpdateStages	Number of texture/mesh update stages
 		 * @Param bAsync			Whether this is called on an async task
 		 */
-		void UpdateStreamingRenderAssets(int32 StageIndex, int32 NumStages, bool bWaitForMipFading, bool bAsync = false);
+		void UpdateStreamingRenderAssets(int32 StageIndex, int32 NumStages, bool bAsync = false);
 
 		/** Check visibility of fast response assets and initiate stream-in requests if necessary. */
 		void TickFastResponseAssets();
@@ -251,23 +251,19 @@ private:
 		 */
 		void StreamRenderAssets( bool bProcessEverything );
 
-		int32 GetNumStreamedMipsArray(EStreamableRenderAssetType AssetType, const int32*& OutArray)
+		TConstArrayView<int32> GetNumStreamedMipsArray(EStreamableRenderAssetType AssetType) const
 		{
 			switch (AssetType)
 			{
 			case EStreamableRenderAssetType::Texture:
-				OutArray = NumStreamedMips_Texture;
-				return TEXTUREGROUP_MAX;
+				return MakeArrayView(NumStreamedMips_Texture);
 			case EStreamableRenderAssetType::StaticMesh:
-				OutArray = NumStreamedMips_StaticMesh.GetData();
-				return NumStreamedMips_StaticMesh.Num();
+				return MakeArrayView(NumStreamedMips_StaticMesh);
 			case EStreamableRenderAssetType::SkeletalMesh:
-				OutArray = NumStreamedMips_SkeletalMesh.GetData();
-				return NumStreamedMips_SkeletalMesh.Num();
+				return MakeArrayView(NumStreamedMips_SkeletalMesh);
 			default:
 				check(false);
-				OutArray = nullptr;
-				return -1;
+				return TConstArrayView<int32>();
 			}
 		}
 
diff --git a/Engine/Source/Runtime/Engine/Private/Streaming/StreamingTexture.cpp b/Engine/Source/Runtime/Engine/Private/Streaming/StreamingTexture.cpp
--- a/Engine/Source/Runtime/Engine/Private/Streaming/StreamingTexture.cpp
+++ b/Engine/Source/Runtime/Engine/Private/Streaming/StreamingTexture.cpp
@@ -1,9 +1,5 @@
 // Copyright Epic Games, Inc. All Rights Reserved.
 
-/*=============================================================================
-StreamingTexture.cpp: Definitions of classes used for texture.
-=============================================================================*/
-
 #include "Streaming/StreamingTexture.h"
 #include "Streaming/StreamingManagerTexture.h"
 #include "Engine/StaticMesh.h"
@@ -30,15 +26,14 @@ static FAutoConsoleVariableRef CVarDefaultNoRefBias(
 FStreamingRenderAsset::FStreamingRenderAsset(
 	UStreamableRenderAsset* InRenderAsset,
 	TSharedPtr<int32, ESPMode::ThreadSafe>& InSimpleStreamableAssetManagerIndex,
-	const int32* NumStreamedMips,
-	int32 NumLODGroups,
+	TConstArrayView<int32> NumStreamedMips,
 	const FRenderAssetStreamingSettings& Settings)
 	: RenderAsset(InRenderAsset)
 	, SimpleStreamableAssetManagerIndex(InSimpleStreamableAssetManagerIndex)
 	, RenderAssetType(InRenderAsset->GetRenderAssetType())
 {
 	UpdateStaticData(Settings);
-	UpdateDynamicData(NumStreamedMips, NumLODGroups, Settings, false);
+	UpdateDynamicData(NumStreamedMips, Settings);
 
 	// In the editor, some paths can recreate the FStreamingRenderAsset, which could potentiallly trigger the unkown ref heuristic.
 	// To prevent this, we consider that the asset bindings where reset when creating the FStreamingRenderAsset.
@@ -171,13 +166,13 @@ void FStreamingRenderAsset::UpdateOptionalMipsState_Async()
 	}	
 }
 
-void FStreamingRenderAsset::UpdateDynamicData(const int32* NumStreamedMips, int32 NumLODGroups, const FRenderAssetStreamingSettings& Settings, bool bWaitForMipFading, TArray<UStreamableRenderAsset*>* DeferredTickCBAssets)
+void FStreamingRenderAsset::UpdateDynamicData(TConstArrayView<int32> NumStreamedMips, const FRenderAssetStreamingSettings& Settings, TArray<UStreamableRenderAsset*>* DeferredTickCBAssets)
 {
 	// Note that those values are read from the async task and must not be assigned temporary values!!
 	if (RenderAsset)
 	{
 		// Get the resource state after calling UpdateStreamingStatus() since it might have updated it.
-		const FStreamableRenderResourceState ResourceState = UpdateStreamingStatus(bWaitForMipFading, DeferredTickCBAssets);
+		const FStreamableRenderResourceState ResourceState = UpdateStreamingStatus(DeferredTickCBAssets);
 
 		// The last render time of this texture/mesh. Can be FLT_MAX when texture has no resource.
 		const float LastRenderTimeForTexture = RenderAsset->GetLastRenderTimeForStreaming();
@@ -239,7 +234,7 @@ void FStreamingRenderAsset::UpdateDynamicData(const int32* NumStreamedMips, int3
 	
 		check( MaxAllowedMips >= ResourceState.NumNonStreamingLODs );
 
-		check(LODGroup < NumLODGroups);
+		check(LODGroup < NumStreamedMips.Num());
 		if (NumStreamedMips[LODGroup] > 0)
 		{
 			MinAllowedMips = FMath::Clamp<int32>(ResourceState.MaxNumLODs - NumStreamedMips[LODGroup], ResourceState.NumNonStreamingLODs, MaxAllowedMips);
@@ -263,7 +258,7 @@ void FStreamingRenderAsset::UpdateDynamicData(const int32* NumStreamedMips, int3
 	}
 }
 
-FStreamableRenderResourceState FStreamingRenderAsset::UpdateStreamingStatus(bool bWaitForMipFading, TArray<UStreamableRenderAsset*>* DeferredTickCBAssets)
+FStreamableRenderResourceState FStreamingRenderAsset::UpdateStreamingStatus(TArray<UStreamableRenderAsset*>* DeferredTickCBAssets)
 {
 	FStreamableRenderResourceState ResourceState;
 
@@ -546,7 +541,7 @@ void FStreamingRenderAsset::CancelStreamingRequest()
 	if (RenderAsset)
 	{
 		RenderAsset->CancelPendingStreamingRequest();
-		UpdateStreamingStatus(false);
+		UpdateStreamingStatus();
 	}
 }
 
@@ -596,7 +591,7 @@ void FStreamingRenderAsset::StreamWantedMips_Internal(FRenderAssetStreamingManag
 				const bool bShouldPrioritizeAsyncIORequest = (bLocalForceFullyLoadHeuristic || bIsTerrainTexture || bLoadWithHigherPriority || IsMissingTooManyMips()) && LocalWantedMips <= LocalVisibleWantedMips;
 				RenderAsset->StreamIn(LocalWantedMips, bShouldPrioritizeAsyncIORequest);
 			}
-			UpdateStreamingStatus(false);
+			UpdateStreamingStatus();
 			TrackRenderAssetEvent(this, RenderAsset, bLocalForceFullyLoadHeuristic != 0, &Manager);
 		}
 	}
diff --git a/Engine/Source/Runtime/Engine/Private/Streaming/StreamingTexture.h b/Engine/Source/Runtime/Engine/Private/Streaming/StreamingTexture.h
--- a/Engine/Source/Runtime/Engine/Private/Streaming/StreamingTexture.h
+++ b/Engine/Source/Runtime/Engine/Private/Streaming/StreamingTexture.h
@@ -1,12 +1,7 @@
 // Copyright Epic Games, Inc. All Rights Reserved.
 
-/*=============================================================================
-StreamingTexture.h: Definitions of classes used for texture streaming.
-=============================================================================*/
-
 #pragma once
 
-#include "CoreMinimal.h"
 #include "Engine/StreamableRenderAsset.h"
 #include "Engine/TextureDefines.h"
 #include "StaticMeshResources.h"
@@ -32,18 +27,17 @@ struct FStreamingRenderAsset
 	FStreamingRenderAsset(
 		UStreamableRenderAsset* InRenderAsset,
 		TSharedPtr<int32, ESPMode::ThreadSafe>& InSimpleStreamableAssetManagerIndex,
-		const int32* NumStreamedMips,
-		int32 NumLODGroups,
+		TConstArrayView<int32> NumStreamedMips,
 		const FRenderAssetStreamingSettings& Settings);
 
 	/** Update data that should not change unless changing settings. */
 	void UpdateStaticData(const FRenderAssetStreamingSettings& Settings);
 
 	/** Update data that the engine could change through gameplay. */
-	void UpdateDynamicData(const int32* NumStreamedMips, int32 NumLODGroups, const FRenderAssetStreamingSettings& Settings, bool bWaitForMipFading, TArray<UStreamableRenderAsset*>* DeferredTickCBAssets = nullptr);
+	void UpdateDynamicData(TConstArrayView<int32> NumStreamedMips, const FRenderAssetStreamingSettings& Settings, TArray<UStreamableRenderAsset*>* DeferredTickCBAssets = nullptr);
 
 	/** Lightweight version of UpdateDynamicData. */
-	FStreamableRenderResourceState UpdateStreamingStatus(bool bWaitForMipFading, TArray<UStreamableRenderAsset*>* DeferredTickCBAssets = nullptr);
+	FStreamableRenderResourceState UpdateStreamingStatus(TArray<UStreamableRenderAsset*>* DeferredTickCBAssets = nullptr);
 
 	/**
 	 * Returns the amount of memory used by the texture/mesh given a specified number of mip-maps, in bytes.
