Fix Dereferencing of Null When Async PSO Compilation is Disabled

Running Lyra, UFXSystemAsset::IsLoadingPSOCache() would dereference null Task pointers when async compilation for PSOs was disabled due to UMaterialInstance::PrecachePSOs() returning a non-empty FGraphEventArray with a single nullptr. This avoids the issue by not considering FPSOPrecacheRequestResults with a a null AsyncCompileEvent to be valid GraphEvents for compilation and therefore it does not add them to the returned FGraphEventArray.

#rnx
#Synced-CL 51782356

diff --git a/Engine/Source/Runtime/Engine/Private/PSOPrecacheMaterial.cpp b/Engine/Source/Runtime/Engine/Private/PSOPrecacheMaterial.cpp
--- a/Engine/Source/Runtime/Engine/Private/PSOPrecacheMaterial.cpp
+++ b/Engine/Source/Runtime/Engine/Private/PSOPrecacheMaterial.cpp
@@ -312,7 +312,12 @@ public:
 				for (FPSOPrecacheRequestResult& Result : PrecacheResults)
 				{
 					check(Result.IsValid());
-					OutGraphEvents.Add(Result.AsyncCompileEvent);
+
+					// AsyncCompileEvent may be null if async compilation is disabled
+					if (Result.AsyncCompileEvent)
+					{
+						OutGraphEvents.Add(Result.AsyncCompileEvent);
+					}
 				}
 			}
 		}
