New material translator: added support for Materials using non tangent-space normals mode.

This required adding a special Normal type of aggregate attributes so the emitter knows how to default initilize the Normal attribute according to the tangent-space normal mode status.

#Synced-CL 51712679

diff --git a/Engine/Source/Runtime/Engine/Private/Materials/MaterialAggregate.cpp b/Engine/Source/Runtime/Engine/Private/Materials/MaterialAggregate.cpp
--- a/Engine/Source/Runtime/Engine/Private/Materials/MaterialAggregate.cpp
+++ b/Engine/Source/Runtime/Engine/Private/Materials/MaterialAggregate.cpp
@@ -14,38 +14,6 @@
 
 #if WITH_EDITOR
 
-EMaterialAggregateAttributeType MaterialValueTypeToMaterialAggregateAttributeType(EMaterialValueType Type)
-{
-	switch (Type)
-	{
-		case MCT_Float:
-			// FMaterialAttributeDefinitionMap::InitializeAttributeMap treats Float as Float1. This is a known bug,
-			// but fixing it would require resaving all materials because component masks are cached. This would cause the
-			// old trnaslator to generate an error about "not enough components".
-			// Fallthrough
-		case MCT_Float1: return EMaterialAggregateAttributeType::Float1;
-		case MCT_Float2: return EMaterialAggregateAttributeType::Float2;
-		case MCT_Float3: return EMaterialAggregateAttributeType::Float3;
-		case MCT_Float4: return EMaterialAggregateAttributeType::Float4;
-	
-		case MCT_UInt:
-		case MCT_UInt1: return EMaterialAggregateAttributeType::UInt1;
-		case MCT_UInt2: return EMaterialAggregateAttributeType::UInt2;
-		case MCT_UInt3: return EMaterialAggregateAttributeType::UInt3;
-		case MCT_UInt4: return EMaterialAggregateAttributeType::UInt4;
-
-		case MCT_Bool: return EMaterialAggregateAttributeType::Bool1;
-
-		case MCT_ShadingModel: return EMaterialAggregateAttributeType::ShadingModel;
-		case MCT_MaterialAttributes: return EMaterialAggregateAttributeType::MaterialAttributes;
-		case MCT_Substrate: return EMaterialAggregateAttributeType::SubstrateData;
-
-		default:
-			checkf(false, TEXT("This material value type '%d' cannot be expressed to a material aggregate attribute type."), (int32)Type);
-			return {};
-	}
-}
-
 EMaterialValueType FMaterialAggregateAttribute::ToMaterialValueType() const
 {
 	switch (Type)
@@ -65,6 +33,7 @@
 		case EMaterialAggregateAttributeType::Float3: return MCT_Float3;
 		case EMaterialAggregateAttributeType::Float4: return MCT_Float4;
 		
+		case EMaterialAggregateAttributeType::Normal: return MCT_Float3;
 		case EMaterialAggregateAttributeType::ShadingModel: return MCT_ShadingModel;
 		case EMaterialAggregateAttributeType::MaterialAttributes: return MCT_MaterialAttributes;
 		case EMaterialAggregateAttributeType::SubstrateData: return MCT_Substrate;
@@ -183,11 +152,53 @@
 namespace MaterialAttributesAggregate
 {
 
+static EMaterialAggregateAttributeType MaterialValueTypeToMaterialAggregateAttributeType(EMaterialValueType Type)
+{
+	switch (Type)
+	{
+		case MCT_Float:
+			// FMaterialAttributeDefinitionMap::InitializeAttributeMap treats Float as Float1. This is a known bug,
+			// but fixing it would require resaving all materials because component masks are cached. This would cause the
+			// old trnaslator to generate an error about "not enough components".
+			// Fallthrough
+		case MCT_Float1: return EMaterialAggregateAttributeType::Float1;
+		case MCT_Float2: return EMaterialAggregateAttributeType::Float2;
+		case MCT_Float3: return EMaterialAggregateAttributeType::Float3;
+		case MCT_Float4: return EMaterialAggregateAttributeType::Float4;
+	
+		case MCT_UInt:
+		case MCT_UInt1: return EMaterialAggregateAttributeType::UInt1;
+		case MCT_UInt2: return EMaterialAggregateAttributeType::UInt2;
+		case MCT_UInt3: return EMaterialAggregateAttributeType::UInt3;
+		case MCT_UInt4: return EMaterialAggregateAttributeType::UInt4;
+
+		case MCT_Bool: return EMaterialAggregateAttributeType::Bool1;
+
+		case MCT_ShadingModel: return EMaterialAggregateAttributeType::ShadingModel;
+		case MCT_MaterialAttributes: return EMaterialAggregateAttributeType::MaterialAttributes;
+		case MCT_Substrate: return EMaterialAggregateAttributeType::SubstrateData;
+
+		default:
+			checkf(false, TEXT("This material value type '%d' cannot be expressed to a material aggregate attribute type."), (int32)Type);
+			return {};
+	}
+}
+
+static EMaterialAggregateAttributeType GetPropertyType(EMaterialProperty Property)
+{
+	switch (Property)
+	{
+		case MP_Normal: return EMaterialAggregateAttributeType::Normal;
+		case MP_SubsurfaceColor: return EMaterialAggregateAttributeType::Float4;
+		default: return MaterialValueTypeToMaterialAggregateAttributeType(FMaterialAttributeDefinitionMap::GetValueType(Property));
+	}
+}
+
 static FMaterialAggregateAttribute MakeAttributeFromProperty(EMaterialProperty Property)
 { 
 	FMaterialAggregateAttribute Attribute;
 	Attribute.Name = (Property == MP_SubsurfaceColor) ? TEXT("Subsurface") : *FMaterialAttributeDefinitionMap::GetAttributeName(Property);
-	Attribute.Type = MaterialValueTypeToMaterialAggregateAttributeType((Property == MP_SubsurfaceColor) ? MCT_Float4 : FMaterialAttributeDefinitionMap::GetValueType(Property));
+	Attribute.Type = GetPropertyType(Property);
 	
 	//ShadingModel defaults to 0, i.e. unlit. This is overriden for the old material system via editor only data, but we can't access that here so we force the value to Lit(1.0)
 	Attribute.DefaultValue = (Property == MP_ShadingModel) ? FVector4f(1.0f, 0.0f, 0.0f, 0.0f) : FMaterialAttributeDefinitionMap::GetDefaultValue(Property); // Should that be changed for substrate?
diff --git a/Engine/Source/Runtime/Engine/Private/Materials/MaterialIREmitter.cpp b/Engine/Source/Runtime/Engine/Private/Materials/MaterialIREmitter.cpp
--- a/Engine/Source/Runtime/Engine/Private/Materials/MaterialIREmitter.cpp
+++ b/Engine/Source/Runtime/Engine/Private/Materials/MaterialIREmitter.cpp
@@ -518,6 +518,7 @@ FType FEmitter::GetMaterialAggregateAttributeType(const UMaterialAggregate* Aggr
 		case EMaterialAggregateAttributeType::UInt4: return FType::MakeVector(EScalarKind::Integer, 4);
 		case EMaterialAggregateAttributeType::Float1: return FType::MakeVector(EScalarKind::Float, 1);
 		case EMaterialAggregateAttributeType::Float2: return FType::MakeVector(EScalarKind::Float, 2);
+		case EMaterialAggregateAttributeType::Normal:
 		case EMaterialAggregateAttributeType::Float3: return FType::MakeVector(EScalarKind::Float, 3);
 		case EMaterialAggregateAttributeType::Float4: return FType::MakeVector(EScalarKind::Float, 4);
 		case EMaterialAggregateAttributeType::MaterialAttributes: return FType::MakeAggregate(MaterialAttributesAggregate::Get());
@@ -1528,26 +1529,41 @@ FValueRef FEmitter::Aggregate(const UMaterialAggregate* InAggregate)
 	return Aggregate(InAggregate, {}, TConstArrayView<FValueRef>{});
 }
 
-static FValueRef EmitAttributeDefaultValue(FEmitter& Emitter, UMaterialInterface* Material, const FMaterialAggregateAttribute& Attribute)
+// This is a necessary, sad special case. See 'EMaterialAggregateAttributeType::Normal' for more info.
+static FValueRef EmitNormalDefaultValue(FEmitter& Em)
+{
+	if (Em.GetBaseMaterial()->bTangentSpaceNormal)
+	{
+		return Em.ConstantFloat3({ 0.0f, 0.f, 1.f });
+	}
+	else
+	{
+		static const FName NAME_VertexNormal = TEXT("VertexNormal");
+		return Em.Extern<MIR::FExternFromMaterialDecl>(NAME_VertexNormal);
+	}
+}
+
+static FValueRef EmitAttributeDefaultValue(FEmitter& Em, UMaterialInterface* Material, const FMaterialAggregateAttribute& Attribute)
 {
 	switch (Attribute.Type)
 	{
-		case EMaterialAggregateAttributeType::Bool1: return Emitter.ConstantBool((bool)Attribute.DefaultValue.X);
-		case EMaterialAggregateAttributeType::Bool2: return Emitter.ConstantBool2((bool)Attribute.DefaultValue.X, (bool)Attribute.DefaultValue.Y);
-		case EMaterialAggregateAttributeType::Bool3: return Emitter.ConstantBool3((bool)Attribute.DefaultValue.X, (bool)Attribute.DefaultValue.Y, (bool)Attribute.DefaultValue.Z);
-		case EMaterialAggregateAttributeType::Bool4: return Emitter.ConstantBool4((bool)Attribute.DefaultValue.X, (bool)Attribute.DefaultValue.Y, (bool)Attribute.DefaultValue.Z, (bool)Attribute.DefaultValue.W);
-		case EMaterialAggregateAttributeType::UInt1: return Emitter.ConstantInt((TInteger)Attribute.DefaultValue.X);
-		case EMaterialAggregateAttributeType::UInt2: return Emitter.ConstantInt2({ (TInteger)Attribute.DefaultValue.X, (TInteger)Attribute.DefaultValue.Y });
-		case EMaterialAggregateAttributeType::UInt3: return Emitter.ConstantInt3({ (TInteger)Attribute.DefaultValue.X, (TInteger)Attribute.DefaultValue.Y, (TInteger)Attribute.DefaultValue.Z });
-		case EMaterialAggregateAttributeType::UInt4: return Emitter.ConstantInt4({ (TInteger)Attribute.DefaultValue.X, (TInteger)Attribute.DefaultValue.Y, (TInteger)Attribute.DefaultValue.Z, (TInteger)Attribute.DefaultValue.W });
-		case EMaterialAggregateAttributeType::Float1: return Emitter.ConstantFloat((TFloat)Attribute.DefaultValue.X);
-		case EMaterialAggregateAttributeType::Float2: return Emitter.ConstantFloat2({ (TFloat)Attribute.DefaultValue.X, (TFloat)Attribute.DefaultValue.Y });
-		case EMaterialAggregateAttributeType::Float3: return Emitter.ConstantFloat3({ (TFloat)Attribute.DefaultValue.X, (TFloat)Attribute.DefaultValue.Y, (TFloat)Attribute.DefaultValue.Z });
-		case EMaterialAggregateAttributeType::Float4: return Emitter.ConstantFloat4(UE::Math::TVector4<TFloat>{ Attribute.DefaultValue.X, Attribute.DefaultValue.Y, Attribute.DefaultValue.Z, Attribute.DefaultValue.W });
-		case EMaterialAggregateAttributeType::ShadingModel: return Emitter.ConstantInt((TInteger)Material->GetShadingModels().GetFirstShadingModel());
-		case EMaterialAggregateAttributeType::MaterialAttributes: return Emitter.Aggregate(MaterialAttributesAggregate::Get());
-		case EMaterialAggregateAttributeType::SubstrateData: return Emitter.SubstrateDefaultSlab();
-		case EMaterialAggregateAttributeType::Aggregate: return Emitter.Aggregate(Attribute.Aggregate);
+		case EMaterialAggregateAttributeType::Bool1: return Em.ConstantBool((bool)Attribute.DefaultValue.X);
+		case EMaterialAggregateAttributeType::Bool2: return Em.ConstantBool2((bool)Attribute.DefaultValue.X, (bool)Attribute.DefaultValue.Y);
+		case EMaterialAggregateAttributeType::Bool3: return Em.ConstantBool3((bool)Attribute.DefaultValue.X, (bool)Attribute.DefaultValue.Y, (bool)Attribute.DefaultValue.Z);
+		case EMaterialAggregateAttributeType::Bool4: return Em.ConstantBool4((bool)Attribute.DefaultValue.X, (bool)Attribute.DefaultValue.Y, (bool)Attribute.DefaultValue.Z, (bool)Attribute.DefaultValue.W);
+		case EMaterialAggregateAttributeType::UInt1: return Em.ConstantInt((TInteger)Attribute.DefaultValue.X);
+		case EMaterialAggregateAttributeType::UInt2: return Em.ConstantInt2({ (TInteger)Attribute.DefaultValue.X, (TInteger)Attribute.DefaultValue.Y });
+		case EMaterialAggregateAttributeType::UInt3: return Em.ConstantInt3({ (TInteger)Attribute.DefaultValue.X, (TInteger)Attribute.DefaultValue.Y, (TInteger)Attribute.DefaultValue.Z });
+		case EMaterialAggregateAttributeType::UInt4: return Em.ConstantInt4({ (TInteger)Attribute.DefaultValue.X, (TInteger)Attribute.DefaultValue.Y, (TInteger)Attribute.DefaultValue.Z, (TInteger)Attribute.DefaultValue.W });
+		case EMaterialAggregateAttributeType::Float1: return Em.ConstantFloat((TFloat)Attribute.DefaultValue.X);
+		case EMaterialAggregateAttributeType::Float2: return Em.ConstantFloat2({ (TFloat)Attribute.DefaultValue.X, (TFloat)Attribute.DefaultValue.Y });
+		case EMaterialAggregateAttributeType::Float3: return Em.ConstantFloat3({ (TFloat)Attribute.DefaultValue.X, (TFloat)Attribute.DefaultValue.Y, (TFloat)Attribute.DefaultValue.Z });
+		case EMaterialAggregateAttributeType::Float4: return Em.ConstantFloat4(UE::Math::TVector4<TFloat>{ Attribute.DefaultValue.X, Attribute.DefaultValue.Y, Attribute.DefaultValue.Z, Attribute.DefaultValue.W });
+		case EMaterialAggregateAttributeType::Normal: return EmitNormalDefaultValue(Em);
+		case EMaterialAggregateAttributeType::ShadingModel: return Em.ConstantInt((TInteger)Material->GetShadingModels().GetFirstShadingModel());
+		case EMaterialAggregateAttributeType::MaterialAttributes: return Em.Aggregate(MaterialAttributesAggregate::Get());
+		case EMaterialAggregateAttributeType::SubstrateData: return Em.SubstrateDefaultSlab();
+		case EMaterialAggregateAttributeType::Aggregate: return Em.Aggregate(Attribute.Aggregate);
 		default: UE_MIR_UNREACHABLE();
 	}
 }
diff --git a/Engine/Source/Runtime/Engine/Private/Materials/MaterialIRModuleBuilder.cpp b/Engine/Source/Runtime/Engine/Private/Materials/MaterialIRModuleBuilder.cpp
--- a/Engine/Source/Runtime/Engine/Private/Materials/MaterialIRModuleBuilder.cpp
+++ b/Engine/Source/Runtime/Engine/Private/Materials/MaterialIRModuleBuilder.cpp
@@ -775,7 +775,7 @@ struct FMaterialIRModuleBuilderImpl
 			{
 				PropertyValue = BuildSubstrateFrontMaterialIRFromRootNodeHiddenConversion();
 			}
-			else if (InputDesc.bUseConstant)
+			else if (Property != MP_Normal && InputDesc.bUseConstant)
 			{
 				// If input is marked to use constant, assign this output to the specified constant value.
 				PropertyValue = Emitter.ConstantFromShaderValue(InputDesc.ConstantValue);
diff --git a/Engine/Source/Runtime/Engine/Private/Materials/MaterialIRToHLSLTranslator.cpp b/Engine/Source/Runtime/Engine/Private/Materials/MaterialIRToHLSLTranslator.cpp
--- a/Engine/Source/Runtime/Engine/Private/Materials/MaterialIRToHLSLTranslator.cpp
+++ b/Engine/Source/Runtime/Engine/Private/Materials/MaterialIRToHLSLTranslator.cpp
@@ -2489,6 +2489,7 @@ static void GenerateMaterialAggregatesDeclarations(TMap<FString, FString>& OutPa
 
 				case EMaterialAggregateAttributeType::Float1: Printer << TEXTVIEW("MaterialFloat"); break;
 				case EMaterialAggregateAttributeType::Float2: Printer << TEXTVIEW("MaterialFloat2"); break;
+				case EMaterialAggregateAttributeType::Normal:
 				case EMaterialAggregateAttributeType::Float3: Printer << TEXTVIEW("MaterialFloat3"); break;
 				case EMaterialAggregateAttributeType::Float4: Printer << TEXTVIEW("MaterialFloat4"); break;
 
diff --git a/Engine/Source/Runtime/Engine/Public/Materials/MaterialAggregate.h b/Engine/Source/Runtime/Engine/Public/Materials/MaterialAggregate.h
--- a/Engine/Source/Runtime/Engine/Public/Materials/MaterialAggregate.h
+++ b/Engine/Source/Runtime/Engine/Public/Materials/MaterialAggregate.h
@@ -29,6 +29,12 @@ enum class EMaterialAggregateAttributeType
 	Float3,
 	Float4,
 
+	// This is needed because of the way materials and material attributes sadly work. The normal attribute
+	// default value is automatically assigned to (0, 0, 1) or VertexNormal depending on whether the material
+	// is in tangent or world space normals. If one day this mechanism changes or is removed, remove this entrys 
+	// and make Normals simple float3s.
+	Normal,
+
 	// Attribute is a shading mode. Used in by MaterialAttributes aggregate.
 	ShadingModel,
 
@@ -42,13 +48,6 @@ enum class EMaterialAggregateAttributeType
 	Aggregate,
 };
 
-#if WITH_EDITOR
-
-// Converts a EMaterialValueType to the equivalent EMaterialAggregateAttributeTypeKind, if possible.
-EMaterialAggregateAttributeType MaterialValueTypeToMaterialAggregateAttributeType(EMaterialValueType Type);
-
-#endif
-
 // Represents a single attribute within a material aggregate.
 USTRUCT(MinimalAPI)
 struct FMaterialAggregateAttribute
diff --git a/Engine/Source/Runtime/Engine/Public/Materials/MaterialIREmitter.h b/Engine/Source/Runtime/Engine/Public/Materials/MaterialIREmitter.h
--- a/Engine/Source/Runtime/Engine/Public/Materials/MaterialIREmitter.h
+++ b/Engine/Source/Runtime/Engine/Public/Materials/MaterialIREmitter.h
@@ -889,6 +889,7 @@ public:
 	
 	// Returns the source material interface. 
 	const UMaterialInterface* GetMaterialInterface() const { return MaterialInterface; }
+	const UMaterial* GetBaseMaterial() const { return BaseMaterial; }
 
 	//
 	const FStaticParameterSet* GetStaticParameterSet() const { return StaticParameterSet; }
