[Cooker/MTSAVE]
- Fix a few delegates that are required to be thread-safe during save
- Make GIgnoreDiffManager thread_local

#Synced-CL 51820749

diff --git a/Engine/Plugins/Runtime/RigVM/Source/RigVMDeveloper/Public/RigVMModel/RigVMNotifications.h b/Engine/Plugins/Runtime/RigVM/Source/RigVMDeveloper/Public/RigVMModel/RigVMNotifications.h
--- a/Engine/Plugins/Runtime/RigVM/Source/RigVMDeveloper/Public/RigVMModel/RigVMNotifications.h
+++ b/Engine/Plugins/Runtime/RigVM/Source/RigVMDeveloper/Public/RigVMModel/RigVMNotifications.h
@@ -71,7 +71,7 @@ enum class ERigVMGraphNotifType : uint8
 };
 
 // A delegate for subscribing / reacting to graph modifications.
-DECLARE_MULTICAST_DELEGATE_ThreeParams(FRigVMGraphModifiedEvent, ERigVMGraphNotifType /* type */, URigVMGraph* /* graph */, UObject* /* subject */);
+DECLARE_TS_MULTICAST_DELEGATE_ThreeParams(FRigVMGraphModifiedEvent, ERigVMGraphNotifType /* type */, URigVMGraph* /* graph */, UObject* /* subject */);
 
 // A dynamic delegate for subscribing / reacting to graph modifications (used for Python integration).
 DECLARE_DYNAMIC_MULTICAST_DELEGATE_ThreeParams(FRigVMGraphModifiedDynamicEvent, ERigVMGraphNotifType, NotifType, URigVMGraph*, Graph, UObject*, Subject);
\ No newline at end of file
diff --git a/Engine/Plugins/Runtime/StateTree/Source/StateTreeModule/Internal/StateTreeDelegatesInternal.h b/Engine/Plugins/Runtime/StateTree/Source/StateTreeModule/Internal/StateTreeDelegatesInternal.h
--- a/Engine/Plugins/Runtime/StateTree/Source/StateTreeModule/Internal/StateTreeDelegatesInternal.h
+++ b/Engine/Plugins/Runtime/StateTree/Source/StateTreeModule/Internal/StateTreeDelegatesInternal.h
@@ -28,7 +28,7 @@ namespace UE::StateTree::Delegates::Private
 	extern UE_API FOnCompileIfChanged OnCompileIfChanged;
 
 	/** A callback for UStateTree::GetAssetRegistryTags. */
-	DECLARE_DELEGATE_TwoParams(FOnStateTreeRequestAssetRegistryTags, TNotNull<const UStateTree*> /*StateTree*/, FAssetRegistryTagsContext& /*Context*/);
+	using FOnStateTreeRequestAssetRegistryTags = TTSDelegate<void(TNotNull<const UStateTree*> /*StateTree*/, FAssetRegistryTagsContext& /*Context*/)>;
 	extern UE_API FOnStateTreeRequestAssetRegistryTags OnRequestAssetRegistryTags;
 
 	/** A callback for UStateTree::AppendToClassSchema. */
diff --git a/Engine/Source/Editor/UMGEditor/Public/WidgetBlueprint.h b/Engine/Source/Editor/UMGEditor/Public/WidgetBlueprint.h
--- a/Engine/Source/Editor/UMGEditor/Public/WidgetBlueprint.h
+++ b/Engine/Source/Editor/UMGEditor/Public/WidgetBlueprint.h
@@ -35,8 +35,8 @@ class FWidgetBlueprintDelegates
 {
 public:
 	// delegate for generating widget asset registry tags.
-	DECLARE_MULTICAST_DELEGATE_TwoParams(FGetAssetTagsWithContext, const UWidgetBlueprint*, FAssetRegistryTagsContext);
-	DECLARE_MULTICAST_DELEGATE_TwoParams(FGetAssetTags, const UWidgetBlueprint*, TArray<UObject::FAssetRegistryTag>&);
+	DECLARE_TS_MULTICAST_DELEGATE_TwoParams(FGetAssetTagsWithContext, const UWidgetBlueprint*, FAssetRegistryTagsContext);
+	DECLARE_TS_MULTICAST_DELEGATE_TwoParams(FGetAssetTags, const UWidgetBlueprint*, TArray<UObject::FAssetRegistryTag>&);
 
 	// called by UWdgetBlueprint::GetAssetRegistryTags()
 	static UE_API FGetAssetTagsWithContext GetAssetTagsWithContext;
diff --git a/Engine/Source/Runtime/CoreUObject/Private/Serialization/ArchiveStackTrace.cpp b/Engine/Source/Runtime/CoreUObject/Private/Serialization/ArchiveStackTrace.cpp
--- a/Engine/Source/Runtime/CoreUObject/Private/Serialization/ArchiveStackTrace.cpp
+++ b/Engine/Source/Runtime/CoreUObject/Private/Serialization/ArchiveStackTrace.cpp
@@ -83,7 +83,7 @@ public:
 	}
 };
 
-static FIgnoreDiffManager GIgnoreDiffManager;
+thread_local FIgnoreDiffManager GIgnoreDiffManager;
 
 static const ANSICHAR* DebugDataStackMarker = "\r\nDebugDataStack:\r\n";
 
diff --git a/Engine/Source/Runtime/Engine/Classes/Engine/World.h b/Engine/Source/Runtime/Engine/Classes/Engine/World.h
--- a/Engine/Source/Runtime/Engine/Classes/Engine/World.h
+++ b/Engine/Source/Runtime/Engine/Classes/Engine/World.h
@@ -4432,8 +4432,8 @@ public:
 	DECLARE_MULTICAST_DELEGATE_TwoParams(FOnLevelChanged, ULevel*, UWorld*);
 
 	// delegate for generating world asset registry tags so project/game scope can add additional tags for filtering levels in their UI, etc
-	DECLARE_MULTICAST_DELEGATE_TwoParams(FWorldGetAssetTagsWithContext, const UWorld*, FAssetRegistryTagsContext);
-	DECLARE_MULTICAST_DELEGATE_TwoParams(FWorldGetAssetTags, const UWorld*, TArray<UObject::FAssetRegistryTag>&);
+	DECLARE_TS_MULTICAST_DELEGATE_TwoParams(FWorldGetAssetTagsWithContext, const UWorld*, FAssetRegistryTagsContext);
+	DECLARE_TS_MULTICAST_DELEGATE_TwoParams(FWorldGetAssetTags, const UWorld*, TArray<UObject::FAssetRegistryTag>&);
 
 	DECLARE_MULTICAST_DELEGATE_ThreeParams(FOnWorldTickStart, UWorld*, ELevelTick, float);
 	static UE_API FOnWorldTickStart OnWorldTickStart;
@@ -4490,7 +4490,7 @@ public:
 	static UE_API FOnWorldPIEEnded OnPIEEnded;
 
 	// Callback to add references when Serialize is called during SavePackage
-	DECLARE_MULTICAST_DELEGATE_TwoParams(FWorldCollectSaveReferencesEvent, UWorld*, FArchive&);
+	DECLARE_TS_MULTICAST_DELEGATE_TwoParams(FWorldCollectSaveReferencesEvent, UWorld*, FArchive&);
 	static UE_API FWorldCollectSaveReferencesEvent OnCollectSaveReferences;
 
 	static UE_API FWorldEvent OnPreRecreateScene;
diff --git a/Engine/Source/Runtime/UMG/Public/Blueprint/WidgetBlueprintGeneratedClass.h b/Engine/Source/Runtime/UMG/Public/Blueprint/WidgetBlueprintGeneratedClass.h
--- a/Engine/Source/Runtime/UMG/Public/Blueprint/WidgetBlueprintGeneratedClass.h
+++ b/Engine/Source/Runtime/UMG/Public/Blueprint/WidgetBlueprintGeneratedClass.h
@@ -57,9 +57,9 @@ class FWidgetBlueprintGeneratedClassDelegates
 {
 public:
 	// delegate for generating widget asset registry tags.
-	DECLARE_MULTICAST_DELEGATE_TwoParams(FGetAssetTagsWithContext, const UWidgetBlueprintGeneratedClass*, FAssetRegistryTagsContext);
-	DECLARE_MULTICAST_DELEGATE_TwoParams(FGetAssetTags, const UWidgetBlueprintGeneratedClass*, TArray<UObject::FAssetRegistryTag>&);
-	DECLARE_MULTICAST_DELEGATE_TwoParams(FCollectSaveOverrides, const UWidgetBlueprintGeneratedClass*, FObjectCollectSaveOverridesContext);
+	DECLARE_TS_MULTICAST_DELEGATE_TwoParams(FGetAssetTagsWithContext, const UWidgetBlueprintGeneratedClass*, FAssetRegistryTagsContext);
+	DECLARE_TS_MULTICAST_DELEGATE_TwoParams(FGetAssetTags, const UWidgetBlueprintGeneratedClass*, TArray<UObject::FAssetRegistryTag>&);
+	DECLARE_TS_MULTICAST_DELEGATE_TwoParams(FCollectSaveOverrides, const UWidgetBlueprintGeneratedClass*, FObjectCollectSaveOverridesContext);
 
 	// called by UWidgetBlueprintGeneratedClass::GetAssetRegistryTags()
 	static UMG_API FGetAssetTagsWithContext GetAssetTagsWithContext;
