New material translator:
- First implementation of UMaterialExpressionSubstrateConvertMaterialAttributes::Build(). Albeit not complete, this fixes material functions not previewing with the new translator when Substrate is on, which is the default now. What is missing is some of the normal tangent needs to be fixed for clear coat etc.

#rnx
#Synced-CL 51713787

diff --git a/Engine/Source/Runtime/Engine/Private/Materials/MaterialExpressionsToMIR.cpp b/Engine/Source/Runtime/Engine/Private/Materials/MaterialExpressionsToMIR.cpp
--- a/Engine/Source/Runtime/Engine/Private/Materials/MaterialExpressionsToMIR.cpp
+++ b/Engine/Source/Runtime/Engine/Private/Materials/MaterialExpressionsToMIR.cpp
@@ -6280,6 +6280,112 @@ void UMaterialExpressionSubstrateShadingModels::Build(MIR::FEmitter& Em)
 		bHasSSS, SubsurfaceProfile/*, EditorOnlyData = nullptr*/);
 }
 
+void UMaterialExpressionSubstrateConvertMaterialAttributes::Build(MIR::FEmitter& Em)
+{
+	FSubstrateTranslatorData* SubstrateTranslatorData = Em.GetSubstrateTranslatorData();
+
+	FValueRef AttrValue = Em.CheckIsAggregate(Em.Input(&MaterialAttributes), MaterialAttributesAggregate::Get());
+	UE_MIR_CHECKPOINT(Em);
+
+	// Subscript a standard property from the MaterialAttributes aggregate.
+	auto GetAttr = [&](EMaterialProperty Property) -> FValueRef
+	{
+		return Em.Subscript(AttrValue, MaterialAttributesAggregate::MaterialPropertyToAttributeIndex(Property));
+	};
+
+	FGuid PathUniqueId = SubstrateTranslatorData->SubstrateTreeStackGetPathUniqueId();
+	FSubstrateOperator& SubstrateOperator = SubstrateTranslatorData->SubstrateCompilationGetOperator(PathUniqueId);
+
+	const bool bHasAnisotropy = SubstrateOperator.Has(ESubstrateBsdfFeature::Anisotropy);
+
+	// Normal in world space
+	FValueRef NormalArg = GetAttr(MP_Normal);
+	NormalArg = Em.Cast(NormalArg, MIR::FType::MakeFloatVector(3));
+	NormalArg = Em.CastToFloatKind(NormalArg);
+	if (!NormalArg.IsValid())
+	{
+		NormalArg = MaterialToMIR::EmitVertexNormal(Em);
+	}
+	NormalArg = TransformNormalFromRequestedBasisToWorld(Em, Material.Get(), NormalArg);
+
+	// Tangent in world space. Always computed (FSubstrateShadingModelsDesc.Tangent must not be null),
+	// though bHasValidTangent gates whether the shared local basis actually uses it.
+	FValueRef TangentArg = GetAttr(MP_Tangent);
+	TangentArg = Em.Cast(TangentArg, MIR::FType::MakeFloatVector(3));
+	TangentArg = Em.CastToFloatKind(TangentArg);
+	if (!TangentArg.IsValid())
+	{
+		TangentArg = MaterialToMIR::EmitVertexTangent(Em);
+	}
+	TangentArg = TransformNormalFromRequestedBasisToWorld(Em, Material.Get(), TangentArg);
+	const bool bHasValidTangent = bHasAnisotropy && TangentArg.IsValid();
+
+	const FSubstrateRegisteredSharedLocalBasis NewRegisteredSharedLocalBasis = MIRSubstrateCompilationInfoCreateSharedLocalBasis2(SubstrateTranslatorData, &NormalArg, bHasValidTangent ? &TangentArg : nullptr);
+	SubstrateOperator.BSDFRegisteredSharedLocalBasis = NewRegisteredSharedLocalBasis;
+	FValueRef BasisIndexMacroArg = EmitSubstrateSharedLocalBasis(Em, NewRegisteredSharedLocalBasis, SubstrateTranslatorData->CurrentSubstrateCompilationContext);
+
+	// SSS profile
+	FValueRef SubsurfaceProfileIdArg = HasSSS() && SubsurfaceProfile
+		? CreateSubsurfaceProfileParameter(Em, SubsurfaceProfile)
+		: CreateDefaultSubsurfaceProfileParameter(Em);
+
+	// Opacity
+	FValueRef OpacityArg = GetAttr(MP_Opacity);
+	if (!OpacityArg.IsValid())
+	{
+		OpacityArg = Em.ConstantFloat(1.0f);
+	}
+
+	// Shading model: use the override constant unless the material has dynamic shading models
+	// or the override requests per-pixel selection from the expression.
+	const int32 ShadingModelCount = Material ? Material->GetShadingModels().CountShadingModels() : 1;
+	const bool bHasDynamicShadingModels = ShadingModelCount > 1;
+	FValueRef ShadingModelArg = (bHasDynamicShadingModels || ShadingModelOverride == MSM_FromMaterialExpression)
+		? GetAttr(MP_ShadingModel)
+		: Em.ConstantFloat(float(ShadingModelOverride));
+
+	// Water parameters from the node's own expression inputs
+	FValueRef WaterScatteringCoefficientsArg = CompileWithDefaultFloat3(Em, WaterScatteringCoefficients, 0.0f, 0.0f, 0.0f);
+	FValueRef WaterAbsorptionCoefficientsArg = CompileWithDefaultFloat3(Em, WaterAbsorptionCoefficients, 0.0f, 0.0f, 0.0f);
+	FValueRef WaterPhaseGArg                 = CompileWithDefaultFloat1(Em, WaterPhaseG, 0.0f);
+	FValueRef ColorScaleBehindWaterArg       = CompileWithDefaultFloat3(Em, ColorScaleBehindWater, 1.0f, 1.0f, 1.0f);
+
+	// Transmittance/thin-translucent: these come from custom attribute GUIDs not present in the
+	// standard aggregate property map, so fall back to the same defaults as the legacy translator.
+	FValueRef TransmittanceColorArg          = Em.ConstantFloat3(FVector3f(1.0f, 1.0f, 1.0f));
+	FValueRef ThinTranslucentSurfaceCoverageArg = Em.ConstantFloat(1.0f);
+
+	MIR::FSubstrateShadingModelsDesc SubstrateShadingModelsDesc = {
+		GetAttr(MP_BaseColor),
+		GetAttr(MP_Specular),
+		GetAttr(MP_Metallic),
+		GetAttr(MP_Roughness),
+		GetAttr(MP_Anisotropy),
+		GetAttr(MP_SubsurfaceColor),
+		SubsurfaceProfileIdArg,
+		GetAttr(MP_CustomData0),    // ClearCoat
+		GetAttr(MP_CustomData1),    // ClearCoatRoughness
+		GetAttr(MP_EmissiveColor),
+		OpacityArg,
+		TransmittanceColorArg,
+		ThinTranslucentSurfaceCoverageArg,
+		WaterScatteringCoefficientsArg,
+		WaterAbsorptionCoefficientsArg,
+		WaterPhaseGArg,
+		ColorScaleBehindWaterArg,
+		ShadingModelArg,
+		NormalArg,
+		TangentArg,
+		NormalArg,              // ClearCoatNormal: no separate coat normal from aggregate in MIR
+		NormalArg,              // CustomTangent: no separate custom tangent from aggregate in MIR
+		BasisIndexMacroArg,
+		BasisIndexMacroArg,     // ClearCoat_BasisIndexMacro: same basis as base layer
+		bHasDynamicShadingModels };
+
+	FValueRef Slab = Em.SubstrateShadingModels(SubstrateShadingModelsDesc);
+	FinalizeBSDFAndPromoteToOperatorIfNeeded(Em, SubstrateOperator, Slab);
+}
+
 void UMaterialExpressionSubstrateSimpleClearCoatBSDF::Build(MIR::FEmitter& Em)
 {
 	FSubstrateTranslatorData* SubstrateTranslatorData = Em.GetSubstrateTranslatorData();
diff --git a/Engine/Source/Runtime/Engine/Public/Materials/MaterialExpressionSubstrate.h b/Engine/Source/Runtime/Engine/Public/Materials/MaterialExpressionSubstrate.h
--- a/Engine/Source/Runtime/Engine/Public/Materials/MaterialExpressionSubstrate.h
+++ b/Engine/Source/Runtime/Engine/Public/Materials/MaterialExpressionSubstrate.h
@@ -1080,7 +1080,7 @@ class UMaterialExpressionSubstrateConvertMaterialAttributes : public UMaterialEx
 #if WITH_EDITOR
 	virtual FExpressionInput* GetInput(int32 InputIndex) override;
 	virtual int32 Compile(class FMaterialCompiler* Compiler, int32 OutputIndex) override;
-	// SUBSTRATE_TODO virtual void Build(MIR::FEmitter& Emitter) override;
+	virtual void Build(MIR::FEmitter& Emitter) override;
 	ENGINE_API static int32 CompileDefaultSlab(class FMaterialCompiler* Compiler, int32 EmissiveOverride);
 	ENGINE_API static int32 CompileDefaultSlab(class FMaterialCompiler* Compiler, FVector3f EmissiveOverride = FVector3f::Zero());
 	virtual void GetCaption(TArray<FString>& OutCaptions) const override;
