New material translator:
- Merged FHardwarePartialDerivative and FExternDifferential into unified FPartialDerivative; it now handles both hardware derivatives (ddx/ddy) and intrinsic partial derivatives of values that support it (externs, vertex interpolators)
- Unified LowerHardwarePartialDerivative/LowerExternDifferential into LowerPartialDerivative with extern and vertex interpolator support
- TextureSampleMipBias in vertex stage now emits AnalysisError directly instead of a derivative node
- Implemented UMaterialExpressionIfThenElse::Compile() in the old translator
- FunctionInput_MaterialAttributes now emits proper Aggregate output in MIR build
- Branch optimization: emit Select when both arms are constant primitives
- Added quotes around extern name in debug output

#Synced-CL 51823973

diff --git a/Engine/Source/Runtime/Engine/Private/Materials/MaterialExpressions.cpp b/Engine/Source/Runtime/Engine/Private/Materials/MaterialExpressions.cpp
--- a/Engine/Source/Runtime/Engine/Private/Materials/MaterialExpressions.cpp
+++ b/Engine/Source/Runtime/Engine/Private/Materials/MaterialExpressions.cpp
@@ -21822,6 +21822,33 @@ FName UMaterialExpressionSceneDepthWithoutWater::GetInputName(int32 InputIndex)
 #endif // WITH_EDITOR
 
 #if WITH_EDITOR
+int32 UMaterialExpressionIfThenElse::Compile(class FMaterialCompiler* Compiler, int32 OutputIndex)
+{
+	if (!Condition.GetTracedInput().Expression)
+	{
+		return Compiler->Errorf(TEXT("Missing IfThenElse Condition input"));
+	}
+	if (!True.GetTracedInput().Expression)
+	{
+		return Compiler->Errorf(TEXT("Missing IfThenElse True input"));
+	}
+	if (!False.GetTracedInput().Expression)
+	{
+		return Compiler->Errorf(TEXT("Missing IfThenElse False input"));
+	}
+
+	const int32 ConditionCode = Condition.Compile(Compiler);
+	const int32 TrueCode = True.Compile(Compiler);
+	const int32 FalseCode = False.Compile(Compiler);
+
+	if (ConditionCode == INDEX_NONE || TrueCode == INDEX_NONE || FalseCode == INDEX_NONE)
+	{
+		return INDEX_NONE;
+	}
+
+	return Compiler->DynamicBranch(ConditionCode, TrueCode, FalseCode);
+}
+
 void UMaterialExpressionIfThenElse::GetCaption(TArray<FString>& OutCaptions) const
 {
 	OutCaptions.Add(TEXT("If"));
diff --git a/Engine/Source/Runtime/Engine/Private/Materials/MaterialExpressionsToMIR.cpp b/Engine/Source/Runtime/Engine/Private/Materials/MaterialExpressionsToMIR.cpp
--- a/Engine/Source/Runtime/Engine/Private/Materials/MaterialExpressionsToMIR.cpp
+++ b/Engine/Source/Runtime/Engine/Private/Materials/MaterialExpressionsToMIR.cpp
@@ -367,11 +367,14 @@ void UMaterialExpressionFunctionInput::Build(MIR::FEmitter& Em)
 			OutputValue = Em.ConstantBool(PreviewValue.X != 0.0f);
 			break;
 
+		case FunctionInput_MaterialAttributes:
+			OutputValue = Em.Aggregate(MaterialAttributesAggregate::Get());
+			break;
+
 		case FunctionInput_Texture2D:
 		case FunctionInput_TextureCube:
 		case FunctionInput_Texture2DArray:
 		case FunctionInput_VolumeTexture:
-		case FunctionInput_MaterialAttributes:
 		case FunctionInput_TextureExternal:
 		case FunctionInput_Substrate:
 			Em.Error(TEXT("Function input of object type requires preview input to be provided."));
@@ -1194,33 +1197,15 @@ void UMaterialExpressionIf::Build(MIR::FEmitter& Em)
 	FValueRef AEqualsBValue = Em.TryInput(&AEqualsB);
 	FValueRef ALessThanBValue = Em.InputDefaultFloat(&ALessThanB, 0.f);
 
-	FValueRef OutputValue;
-
 	// Less than comparison -- if equals value isn't present (see below), AGreaterThanBValue will also be returned for the equal case.
 	FValueRef ALessThanBConditionValue = Em.LessThan(AValue, BValue);
-
-	if (ALessThanBConditionValue->Type.IsBoolScalar())
-	{
-		OutputValue = Em.Branch(ALessThanBConditionValue, ALessThanBValue, AGreaterThanBValue);
-	}
-	else
-	{
-		OutputValue = Em.Select(ALessThanBConditionValue, ALessThanBValue, AGreaterThanBValue);
-	}
+	FValueRef OutputValue = Em.Branch(ALessThanBConditionValue, ALessThanBValue, AGreaterThanBValue);
 
 	// Equals value is optional -- if present, generate an additional conditional.
 	if (AEqualsBValue.IsValid())
 	{
 		FValueRef AEqualsBConditionValue = EmitAlmostEquals(Em, AValue, BValue, EqualsThreshold);
-
-		if (AEqualsBConditionValue->Type.IsBoolScalar())
-		{
-			OutputValue = Em.Branch(AEqualsBConditionValue, AEqualsBValue, OutputValue);
-		}
-		else
-		{
-			OutputValue = Em.Select(AEqualsBConditionValue, AEqualsBValue, OutputValue);
-		}
+		OutputValue = Em.Branch(AEqualsBConditionValue, AEqualsBValue, OutputValue);
 	}
 
 	Em.Output(0, OutputValue);
diff --git a/Engine/Source/Runtime/Engine/Private/Materials/MaterialIR.cpp b/Engine/Source/Runtime/Engine/Private/Materials/MaterialIR.cpp
--- a/Engine/Source/Runtime/Engine/Private/Materials/MaterialIR.cpp
+++ b/Engine/Source/Runtime/Engine/Private/Materials/MaterialIR.cpp
@@ -38,7 +38,6 @@ const TCHAR* LexToString(EValueKind Kind)
 		case VK_TextureRead: return TEXT("TextureRead");
 		case VK_VTPageTableRead: return TEXT("VTPageTableRead");
 		case VK_Extern: return TEXT("Extern");
-		case VK_ExternDifferential: return TEXT("ExternDifferential");
 		case VK_SubstrateDefaultSlab: return TEXT("SubstrateDefaultSlab");
 		case VK_SubstrateSlab: return TEXT("SubstrateSlab");
 		case VK_SubstrateShadingModels: return TEXT("SubstrateShadingModel");
@@ -50,7 +49,7 @@ const TCHAR* LexToString(EValueKind Kind)
 		case VK_SubstrateSelect: return TEXT("SubstrateSelect");
 		case VK_SubstratePromoteToOperator: return TEXT("SubstratePromoteToOperator");
 		case VK_StageSwitch: return TEXT("StageSwitch");
-		case VK_HardwarePartialDerivative: return TEXT("HardwarePartialDerivative");
+		case VK_PartialDerivative: return TEXT("HardwarePartialDerivative");
 		case VK_SetVertexInterpolator: return TEXT("SetVertexInterpolator");
 		case VK_Call: return TEXT("Call");
 		case VK_CallParameterOutput: return TEXT("CallParameterOutput");
@@ -128,13 +127,12 @@ uint32 FValue::GetSizeInBytes() const
 		case VK_SubstrateSelect: return sizeof(FSubstrateSelect);
 		case VK_SubstratePromoteToOperator: return sizeof(FSubstratePromoteToOperator);
 		case VK_StageSwitch: return sizeof(FStageSwitch);
-		case VK_HardwarePartialDerivative: return sizeof(FHardwarePartialDerivative);
+		case VK_PartialDerivative: return sizeof(FPartialDerivative);
 		case VK_SetVertexInterpolator: return sizeof(FSetVertexInterpolator);
 		case VK_Call: return sizeof(FCall);
 		case VK_CallParameterOutput: return sizeof(FCallParameterOutput);
 		case VK_PreshaderParameter: return sizeof(FPreshaderParameter);
 		case VK_Extern: return static_cast<const FExtern*>(this)->Impl->ByteSize;
-		case VK_ExternDifferential: return sizeof(FExternDifferential);
 
 			/* invalid entries */
 		case VK_InstructionBegin:
@@ -214,12 +212,6 @@ TConstArrayView<FValue*> FValue::GetUses() const
 			return This->GetArguments();
 		}
 
-		case VK_ExternDifferential:
-		{
-			auto This = static_cast<const FExternDifferential*>(this);
-			return { &This->ExternArg, FExternDifferential::NumStaticUses };
-		}
-
 		case VK_SubstrateDefaultSlab:
 		{
 			auto This = static_cast<const FSubstrateDefaultSlab*>(this);
@@ -286,10 +278,10 @@ TConstArrayView<FValue*> FValue::GetUses() const
 			return { This->Args, FStageSwitch::NumStaticUses * NumStages };
 		}
 
-		case VK_HardwarePartialDerivative:
+		case VK_PartialDerivative:
 		{
-			auto This = static_cast<const FHardwarePartialDerivative*>(this);
-			return { &This->Arg, FHardwarePartialDerivative::NumStaticUses };
+			auto This = static_cast<const FPartialDerivative*>(this);
+			return { &This->Arg, FPartialDerivative::NumStaticUses };
 		}
 
 		case VK_SetVertexInterpolator:
diff --git a/Engine/Source/Runtime/Engine/Private/Materials/MaterialIRDebug.cpp b/Engine/Source/Runtime/Engine/Private/Materials/MaterialIRDebug.cpp
--- a/Engine/Source/Runtime/Engine/Private/Materials/MaterialIRDebug.cpp
+++ b/Engine/Source/Runtime/Engine/Private/Materials/MaterialIRDebug.cpp
@@ -342,7 +342,7 @@ struct FDebugDumpIRState
 			Out.Appendf(TEXT(" TextureIndex=%d"), PreshaderParameter->TextureIndex);
 			Out.Appendf(TEXT(" PreshaderOffset=%d"), PreshaderParameter->Analysis_PreshaderOffset - 1);			// One is added to preshader offsets, so zero can represent none, subtract that here.
 		}
-		else if (const FHardwarePartialDerivative* Derivative = Instr->As<FHardwarePartialDerivative>())
+		else if (const FPartialDerivative* Derivative = Instr->As<FPartialDerivative>())
 		{
 			Out.Append(Derivative->Axis == MIR::EDerivativeAxis::X ? TEXT(" \"ddx\"") : TEXT(" \"ddy\""));
 		}
@@ -363,10 +363,6 @@ struct FDebugDumpIRState
 			Out.Appendf(TEXT(" \"%s\" "), Extern->GetInfo().Name.GetData());
 			Extern->EmitDebugInfo(Out);
 		}
-		else if (const FExternDifferential* ExternDifferential = Instr->As<FExternDifferential>())
-		{
-			Out.Append(ExternDifferential->Axis == EDerivativeAxis::X ? TEXT("ddx") : TEXT("ddy"));
-		}
 
 		if (Instr->HasSubgraphProperties(MIR::EGraphProperties::PreshaderOut))
 		{
diff --git a/Engine/Source/Runtime/Engine/Private/Materials/MaterialIREmitter.cpp b/Engine/Source/Runtime/Engine/Private/Materials/MaterialIREmitter.cpp
--- a/Engine/Source/Runtime/Engine/Private/Materials/MaterialIREmitter.cpp
+++ b/Engine/Source/Runtime/Engine/Private/Materials/MaterialIREmitter.cpp
@@ -2970,6 +2970,12 @@ FValueRef FEmitter::Branch(FValueRef Condition, FValueRef True, FValueRef False)
 		return Poison();
 	}
 
+	// Constants don't require deferred evaluation inside scoped blocks. If so instead emit a lighter "select" instead.
+	if (CommonType.IsPrimitive() && True->IsConstant() && False->IsConstant())
+	{
+		return Select(Condition, True, False);
+	}
+
 	// Create the branch instruction.
 	FBranch Proto = MakePrototype<FBranch>(CommonType);
 	Proto.ConditionArg = Condition;
@@ -3688,15 +3694,6 @@ FValueRef FEmitter::TextureSampleBias(FValueRef Texture, FValueRef TexCoord, FVa
 	PrototypeCompute.TexCoordDdx = Operator(BO_Multiply, AnalyticalPartialDerivative(TexCoord, EDerivativeAxis::X), Exp2MipBias);
 	PrototypeCompute.TexCoordDdy = Operator(BO_Multiply, AnalyticalPartialDerivative(TexCoord, EDerivativeAxis::Y), Exp2MipBias);
 
-	// We throw a deliberate error downstream for this expression if referenced from a vertex shader.  We use the hardware
-	// derivative value type (disallowed in vertex shader, already tracked by the value analyzer) to communicate the error.
-	// We could alternately make a one-off type just to communicate this error, but this works, and is somewhat logical,
-	// given the root cause of this sampler type not working is due to missing derivative support in vertex shaders.
-	FHardwarePartialDerivative PrototypeVertex = MakePrototype<FHardwarePartialDerivative>(PrototypePixel.Type);
-	PrototypeVertex.Arg = ConstantFloat(0.0f);
-	PrototypeVertex.Axis = EDerivativeAxis::X;
-	PrototypeVertex.Source = EDerivativeSource::TextureSampleBias;
-
 	if (IsVirtualSamplerType(SamplerType))
 	{
 		PrototypePixel.VTPageTable   = VTPageTableLoadFromSamplerSource(*this, Texture, BaseAttributes, PrototypePixel.TexCoord, {}, {}, TMVM_MipBias, PrototypePixel.MipValue);
@@ -3704,7 +3701,7 @@ FValueRef FEmitter::TextureSampleBias(FValueRef Texture, FValueRef TexCoord, FVa
 	}
 
 	FStageSwitch StageSwitch = MakePrototype<FStageSwitch>(PrototypePixel.Type);
-	StageSwitch.Args[Stage_Vertex] = EmitPrototype(*this, PrototypeVertex);
+	StageSwitch.Args[Stage_Vertex] = AnalysisError(PrototypePixel.Type, TEXTVIEW("TextureSampleMipBias is only supported in the pixel shader"));
 	StageSwitch.Args[Stage_Pixel] = EmitPrototype(*this, PrototypePixel);
 	StageSwitch.Args[Stage_Compute] = EmitPrototype(*this, PrototypeCompute);
 
@@ -3763,10 +3760,10 @@ FValueRef FEmitter::PartialDerivative(FValueRef Value, EDerivativeAxis Axis)
 	}
 
 	// Make the hardware derivative instruction.
-	FHardwarePartialDerivative HwDerivativeProto = MakePrototype<FHardwarePartialDerivative>(Value->Type);
+	FPartialDerivative HwDerivativeProto = MakePrototype<FPartialDerivative>(Value->Type);
 	HwDerivativeProto.Arg = Value;
 	HwDerivativeProto.Axis = Axis;
-	HwDerivativeProto.Source = EDerivativeSource::Derivative;
+	HwDerivativeProto.bHardwareDerivative = true;
 	FValueRef HwDerivative = EmitPrototype(*this, HwDerivativeProto);
 
 	// Compute the analytical derivative for stages that don't support hardware derivatives.
@@ -4115,7 +4112,7 @@ FValueRef FEmitter::AnalyticalPartialDerivative(FValueRef Value, EDerivativeAxis
 			return Subscript(AnalyticalPartialDerivative(AsSubscript->Arg, Axis), AsSubscript->Index);
 		}
 
-		case VK_Scalar:
+		case VK_Scalar: 
 		{
 			FScalar* AsScalar = Value->As<FScalar>();
 			return Cast(AnalyticalPartialDerivative(AsScalar->Arg, Axis), AsScalar->Type);
@@ -4125,22 +4122,31 @@ FValueRef FEmitter::AnalyticalPartialDerivative(FValueRef Value, EDerivativeAxis
 		{
 			FExtern* AsExtern = Value->As<FExtern>();
 			const FExternInfo& Info = AsExtern->GetInfo();
-			if ((Info.Flags & MIR::EExternFlags::NoDifferentials) != MIR::EExternFlags::None)
+			const FType DerivativeType = !Info.DifferentialType.IsPoison() ? Info.DifferentialType : Info.Type;
+
+			if ((Info.Flags & MIR::EExternFlags::ZeroDifferentials) != MIR::EExternFlags::None)
+			{
+				return ConstantDefault(Info.DifferentialType);
+			}
+			else if ((Info.Flags & MIR::EExternFlags::NoDifferentials) != MIR::EExternFlags::None)
 			{
 				Errorf(TEXT("Extern '%s' does not support analytical partial derivatives."), Info.Name.GetData());
 				return Value.ToPoison();
 			}
-			FType DerivativeType = !Info.DifferentialType.IsPoison() ? Info.DifferentialType : Info.Type;
-			FExternDifferential DerivativeProto = MakePrototype<FExternDifferential>(DerivativeType);
-			DerivativeProto.ExternArg = AsExtern;
+
+			FPartialDerivative DerivativeProto = MakePrototype<FPartialDerivative>(DerivativeType);
+			DerivativeProto.Arg = Value;
 			DerivativeProto.Axis = Axis;
+
 			return EmitPrototype(*this, DerivativeProto);
 		}
 
 		case VK_GetVertexInterpolator:
 		{
-			FGetVertexInterpolator* AsVertexInterpolator = Value->As<FGetVertexInterpolator>();
-			return VertexInterpolator(AnalyticalPartialDerivative(AsVertexInterpolator->SetInstr->Arg, Axis));
+			FPartialDerivative DerivativeProto = MakePrototype<FPartialDerivative>(Value->Type);
+			DerivativeProto.Arg = Value;
+			DerivativeProto.Axis = Axis;
+			return EmitPrototype(*this, DerivativeProto);
 		}
 
 		// These values don't work with analytic derivatives, and force hardware derivatives (or zero if the shader model doesn't support compute shader derivatives).
@@ -4149,12 +4155,11 @@ FValueRef FEmitter::AnalyticalPartialDerivative(FValueRef Value, EDerivativeAxis
 		{
 			if (GetFeatureLevel() >= ERHIFeatureLevel::Type::SM6)
 			{
-				// Make the hardware derivative instruction.
-				FHardwarePartialDerivative HwDerivativeProto = MakePrototype<FHardwarePartialDerivative>(Value->Type);
-				HwDerivativeProto.Arg = Value;
-				HwDerivativeProto.Axis = Axis;
-				HwDerivativeProto.Source = EDerivativeSource::AnalyticDerivative;
-				return EmitPrototype(*this, HwDerivativeProto);
+				FPartialDerivative DerivativeProto = MakePrototype<FPartialDerivative>(Value->Type);
+				DerivativeProto.Arg = Value;
+				DerivativeProto.Axis = Axis;
+				DerivativeProto.bHardwareDerivative = true;
+				return EmitPrototype(*this, DerivativeProto);
 			}
 			else
 			{
@@ -4166,7 +4171,7 @@ FValueRef FEmitter::AnalyticalPartialDerivative(FValueRef Value, EDerivativeAxis
 		case VK_Constant:
 		case VK_PrimitiveUniform:
 		case VK_PreshaderParameter:
-		case VK_ExternDifferential:
+		case VK_PartialDerivative:
 			return Cast(ConstantZero(Value->Type.GetPrimitive().ScalarKind), Value->Type);
 
 		default:
diff --git a/Engine/Source/Runtime/Engine/Private/Materials/MaterialIRToHLSLTranslator.cpp b/Engine/Source/Runtime/Engine/Private/Materials/MaterialIRToHLSLTranslator.cpp
--- a/Engine/Source/Runtime/Engine/Private/Materials/MaterialIRToHLSLTranslator.cpp
+++ b/Engine/Source/Runtime/Engine/Private/Materials/MaterialIRToHLSLTranslator.cpp
@@ -455,9 +455,8 @@ struct FEntryPointCodeGen
 			case MIR::VK_TextureRead: LowerTextureRead(static_cast<const MIR::FTextureRead*>(Instr)); break;
 			case MIR::VK_VTPageTableRead: LowerVTPageTableRead(static_cast<const MIR::FVTPageTableRead*>(Instr)); break;
 			case MIR::VK_Extern: LowerExtern(static_cast<const MIR::FExtern*>(Instr)); break;
-			case MIR::VK_ExternDifferential: LowerExternDifferential(static_cast<const MIR::FExternDifferential*>(Instr)); break;
 			case MIR::VK_StageSwitch: LowerStageSwitch(static_cast<const MIR::FStageSwitch*>(Instr)); break;
-			case MIR::VK_HardwarePartialDerivative: LowerHardwarePartialDerivative(static_cast<const MIR::FHardwarePartialDerivative*>(Instr)); break;
+			case MIR::VK_PartialDerivative: LowerPartialDerivative(static_cast<const MIR::FPartialDerivative*>(Instr)); break;
 			case MIR::VK_SetVertexInterpolator: LowerSetVertexInterpolator(static_cast<const MIR::FSetVertexInterpolator*>(Instr)); break;
 			case MIR::VK_Call: LowerCall(static_cast<const MIR::FCall*>(Instr)); break;
 			case MIR::VK_CallParameterOutput: LowerCallOutput(static_cast<const MIR::FCallParameterOutput*>(Instr)); break;
@@ -524,20 +523,22 @@ struct FEntryPointCodeGen
 		Printer << MaterialParameterCollection->Analysis_CollectionIndex;
 	}
 
-	void LowerGetVertexInterpolator(const MIR::FGetVertexInterpolator* GetVertexInterpolator)
+	void LowerGetVertexInterpolator(const MIR::FGetVertexInterpolator* GetVertexInterpolator, const TCHAR* Suffix = TEXT(""))
 	{
 		if (GetVertexInterpolator->Type.GetPrimitive().NumComponents() > 1)
 		{
 			Printer << GetVertexInterpolator->Type << TEXT('(');
 		}
-		ComputeUserVertexInterpolatorAllocation(GetVertexInterpolator->SetInstr, [this](int32 Index, int32 RegisterIndex, FStringView RegisterSwizzle, FStringView)
+
+		ComputeUserVertexInterpolatorAllocation(GetVertexInterpolator->SetInstr, [this, Suffix](int32 Index, int32 RegisterIndex, FStringView RegisterSwizzle, FStringView)
 		{
 			if (Index > 0)
 			{
 				Printer << TEXTVIEW(", ");
 			}
-			Printer << TEXTVIEW("Parameters.TexCoords[") << RegisterIndex << TEXTVIEW("].") << RegisterSwizzle;
+			Printer << TEXTVIEW("Parameters.TexCoords") << Suffix << TEXT('[') << RegisterIndex << TEXTVIEW("].") << RegisterSwizzle;
 		});
+
 		if (GetVertexInterpolator->Type.GetPrimitive().NumComponents() > 1)
 		{
 			Printer << TEXT(')');
@@ -566,33 +567,6 @@ struct FEntryPointCodeGen
 		Extern->ToHLSL(ExternPrinter);
 	}
 
-	void LowerExternDifferential(const MIR::FExternDifferential* ExternDifferential)
-	{
-		const MIR::FExtern* ExternArg = ExternDifferential->GetExtern();
-		check(ExternArg); // FExternDifferential argument must always be an extern
-
-		// If extern has zero derivatives generate them directly
-		if ((ExternArg->GetInfo().Flags & MIR::EExternFlags::ZeroDifferentials) != MIR::EExternFlags::None)
-		{
-			check(ExternDifferential->Type.IsPrimitive());
-			if (ExternDifferential->Type.IsDouble())
-			{
-				// Example: WSPromote((MaterialFloat3)0)
-				Printer << TEXT("WSPromote((") << ExternDifferential->Type.GetPrimitive().ToScalarKind(MIR::EScalarKind::Float) << TEXT(")0)");
-			}
-			else
-			{
-				Printer << TEXT('(') << ExternDifferential->Type << TEXT(")0");
-			}
-			return;
-		}
-
-		// Otherwise invoke the extern to generate the differential code by itself
-		const MIR::EExternDifferential Differential = ExternDifferential->Axis == MIR::EDerivativeAxis::X ? MIR::EExternDifferential::ScreenSpaceX : MIR::EExternDifferential::ScreenSpaceY;
-		MIR::FExternPrinterHLSL ExternPrinter = MakeExternPrinter(Differential, ExternArg->GetArguments());
-		ExternArg->ToHLSL(ExternPrinter);
-	}
-
 	void LowerTextureUniform(const MIR::FTextureUniform* TextureUniform)
 	{
 		Printer << GetMaterialUBName() << TEXTVIEW(".") << LowerTextureMaterialType(TextureUniform->Texture->GetMaterialType(), true) << TEXTVIEW("_") << TextureUniform->Analysis_UniformIndex;
@@ -1716,16 +1690,35 @@ struct FEntryPointCodeGen
 		LowerValue(StageSwitch->Args[CurrentStage]);
 	}
 
-	void LowerHardwarePartialDerivative(const MIR::FHardwarePartialDerivative* StageSwitch)
+	void LowerPartialDerivative(const MIR::FPartialDerivative* PartialDerivative)
 	{
-		if (StageSwitch->Arg->Type.IsDouble())
+		if (PartialDerivative->bHardwareDerivative)
 		{
-			// Expression emitter assumes these to be LWC for analytic derivative evaluation, so we promote them on load
-			Printer << (StageSwitch->Axis == MIR::EDerivativeAxis::X ? TEXTVIEW("WSPromote(WSDdxDemote(") : TEXTVIEW("WSPromote(WSDdyDemote(")) << LowerValue(StageSwitch->Arg) << TEXTVIEW("))");
+			// This is differentiating a generic expression.
+			if (PartialDerivative->Arg->Type.IsDouble())
+			{
+				// Expression emitter assumes these to be LWC for analytic derivative evaluation, so we promote them on load
+				Printer << (PartialDerivative->Axis == MIR::EDerivativeAxis::X ? TEXTVIEW("WSPromote(WSDdxDemote(") : TEXTVIEW("WSPromote(WSDdyDemote(")) << LowerValue(PartialDerivative->Arg) << TEXTVIEW("))");
+			}
+			else
+			{
+				Printer << (PartialDerivative->Axis == MIR::EDerivativeAxis::X ? TEXTVIEW("DDX(") : TEXTVIEW("DDY(")) << LowerValue(PartialDerivative->Arg) << TEXTVIEW(")");
+			}
+		}
+		else if (const MIR::FExtern* ExternArg = PartialDerivative->Arg->As<MIR::FExtern>())
+		{
+			// Otherwise invoke the extern to generate the differential code by itself
+			const MIR::EExternDifferential Differential = PartialDerivative->Axis == MIR::EDerivativeAxis::X ? MIR::EExternDifferential::ScreenSpaceX : MIR::EExternDifferential::ScreenSpaceY;
+			MIR::FExternPrinterHLSL ExternPrinter = MakeExternPrinter(Differential, ExternArg->GetArguments());
+			ExternArg->ToHLSL(ExternPrinter);
+		}
+		else if (const MIR::FGetVertexInterpolator* InterpolatorArg = PartialDerivative->Arg->As<MIR::FGetVertexInterpolator>())
+		{
+			LowerGetVertexInterpolator(InterpolatorArg, PartialDerivative->Axis == MIR::EDerivativeAxis::X ? TEXT("_DDX") : TEXT("_DDY"));
 		}
 		else
 		{
-			Printer << (StageSwitch->Axis == MIR::EDerivativeAxis::X ? TEXTVIEW("DDX(") : TEXTVIEW("DDY(")) << LowerValue(StageSwitch->Arg) << TEXTVIEW(")");
+			UE_MIR_UNREACHABLE();
 		}
 	}
 
@@ -2143,8 +2136,6 @@ struct FEntryPointCodeGen
 				return false;
 			case MIR::VK_Extern:
 				return (static_cast<const MIR::FExtern*>(Instr)->GetInfo().Flags & MIR::EExternFlags::Inline) != MIR::EExternFlags::None;
-			case MIR::VK_ExternDifferential:
-				return IsFoldable(static_cast<const MIR::FExternDifferential*>(Instr)->GetExtern(), EntryPointIndex);
 
 			default:
 				return true;
@@ -2176,6 +2167,7 @@ static void GenerateStageEntryPoints(const FMaterialIRModule* Module, TMap<FStri
 
 	// Generate compute stage HLSL
 	Result = CodeGen.GenerateEntryPoint(MIR::Stage_Compute);
+	Result.NonNormalHLSL += TEXTVIEW("PixelMaterialInputs.Subsurface = 0;\n");
 	Result.NonNormalHLSL += CodeGen.Module->GetSubstrateData().PixelInputInitializerValuesHLSL;
 	OutParameters.Add(TEXT("calc_pixel_material_inputs_analytic_derivatives_normal"), MoveTemp(Result.NormalAttributeHLSL));
 	OutParameters.Add(TEXT("calc_pixel_material_inputs_analytic_derivatives_other_inputs"), MoveTemp(Result.NonNormalHLSL));
diff --git a/Engine/Source/Runtime/Engine/Private/Materials/MaterialIRValueAnalyzer.cpp b/Engine/Source/Runtime/Engine/Private/Materials/MaterialIRValueAnalyzer.cpp
--- a/Engine/Source/Runtime/Engine/Private/Materials/MaterialIRValueAnalyzer.cpp
+++ b/Engine/Source/Runtime/Engine/Private/Materials/MaterialIRValueAnalyzer.cpp
@@ -30,18 +30,11 @@ static void AnalysisError_Analyze(FMaterialIRValueAnalyzer& Analyzer, MIR::FAnal
 	Analyzer.Module->AddError(AnalysisError->Expression, AnalysisError->Message.ToString());
 }
 
-static void HardwarePartialDerivative_AnalyzeInStage(FMaterialIRValueAnalyzer& Analyzer, MIR::FHardwarePartialDerivative* PartialDerivative, MIR::EStage Stage)
+static void PartialDerivative_AnalyzeInStage(FMaterialIRValueAnalyzer& Analyzer, MIR::FPartialDerivative* PartialDerivative, MIR::EStage Stage)
 {
 	if (Stage == MIR::Stage_Vertex)
 	{
-		if (PartialDerivative->Source == MIR::EDerivativeSource::TextureSampleBias)
-		{
-			Analyzer.Module->AddError(nullptr, TEXT("(Node TextureSample) MipBias is only supported in the pixel shader"));
-		}
-		else
-		{
-			Analyzer.Module->AddError(nullptr, FString::Printf(TEXT("Invalid DD%c node used in vertex shader input!"), 'X' + (int32)PartialDerivative->Axis));
-		}
+		Analyzer.Module->AddError(nullptr, FString::Printf(TEXT("Screen-space partial derivative DD%c cannot be computed in the vertex stage."), 'X' + (int32)PartialDerivative->Axis));
 	}
 }
 
@@ -696,7 +689,7 @@ void FMaterialIRValueAnalyzer::AnalyzeInStage(MIR::FValue* Value, MIR::EStage St
 	{
 		EXPAND_CASE(Extern);
 		EXPAND_CASE(VTPageTableRead);
-		EXPAND_CASE(HardwarePartialDerivative);
+		EXPAND_CASE(PartialDerivative);
 		EXPAND_CASE(Call);
 		default: break;
 	}
diff --git a/Engine/Source/Runtime/Engine/Public/Materials/MaterialExpressionIfThenElse.h b/Engine/Source/Runtime/Engine/Public/Materials/MaterialExpressionIfThenElse.h
--- a/Engine/Source/Runtime/Engine/Public/Materials/MaterialExpressionIfThenElse.h
+++ b/Engine/Source/Runtime/Engine/Public/Materials/MaterialExpressionIfThenElse.h
@@ -28,6 +28,7 @@ public:
 	
 	//~ Begin UMaterialExpression Interface
 #if WITH_EDITOR
+	virtual int32 Compile(class FMaterialCompiler* Compiler, int32 OutputIndex) override;
 	virtual void Build(MIR::FEmitter& Emitter) override;
 	virtual void GetCaption(TArray<FString>& OutCaptions) const override;
 	virtual EMaterialValueType GetInputValueType(int32 InputIndex) override;
diff --git a/Engine/Source/Runtime/Engine/Public/Materials/MaterialIR.h b/Engine/Source/Runtime/Engine/Public/Materials/MaterialIR.h
--- a/Engine/Source/Runtime/Engine/Public/Materials/MaterialIR.h
+++ b/Engine/Source/Runtime/Engine/Public/Materials/MaterialIR.h
@@ -69,7 +69,6 @@ enum EValueKind
 	VK_TextureRead,
 	VK_VTPageTableRead,
 	VK_Extern,
-	VK_ExternDifferential,
 	VK_SubstrateDefaultSlab,
 	VK_SubstrateSlab,
 	VK_SubstrateShadingModels,
@@ -81,7 +80,7 @@ enum EValueKind
 	VK_SubstrateSelect,
 	VK_SubstratePromoteToOperator,
 	VK_StageSwitch,
-	VK_HardwarePartialDerivative,
+	VK_PartialDerivative,
 	VK_SetVertexInterpolator,
 	VK_Call,
 	VK_CallParameterOutput,
@@ -329,15 +328,6 @@ struct FConstant : TValue<VK_Constant>
 	}
 };
 
-// Specify the source expression type of a derivative, useful for generating actionable user facing errors if derivatives
-// are used where not allowed (vertex shader).
-enum class EDerivativeSource
-{
-	Derivative,				// DDX or DDY expression
-	TextureSampleBias,		// TextureSample expression with MipBias option selected
-	AnalyticDerivative,		// Generated during calculation of an analytic derivative
-};
-
 // Enumeration of builtin values.
 enum class EBuiltin
 {
@@ -625,19 +615,6 @@ struct FExtern : TInstruction<VK_Extern, 0>
 	void EmitDebugInfo(FString& Out) const { Impl->EmitDebugInfo(this, Out); }
 };
 
-// An auxiliary instruction used to generate the differential (ddx/ddy) of an Extern.
-struct FExternDifferential : TInstruction<VK_ExternDifferential, 1>
-{
-	// The extern to differentiate
-	FValue* ExternArg;
-
-	// The differential axis to generate for
-	EDerivativeAxis Axis;
-
-	// Convenience getter of the extern.
-	const FExtern* GetExtern() const { return MIR::As<FExtern>(ExternArg); }
-};
-
 // Operator enumeration.
 //
 // Note: If you modify this enum, update the implementations of the helper functions below.
@@ -945,19 +922,21 @@ struct FStageSwitch : TInstruction<VK_StageSwitch, 1>
 	void SetArgs(FValue* PixelStageArg, FValue* OtherStagesArg);
 };
 
-// Instruction that maps to hardware ddx()/ddy().
-//
-// Note: This is only available in stages that support hardware derivatives (e.g. pixel shader)
-struct FHardwarePartialDerivative : TInstruction<VK_HardwarePartialDerivative, 1>
+// This instruction represents the screen-space partial derivative of a given value along an axis.
+struct FPartialDerivative : TInstruction<VK_PartialDerivative, 1>
 {
-	// The value argument of ddx()/ddy()
+	// The value argument being differentiated.
+	// If bHardwareDerivative is true, this can be any (valid) value that supports calculating the partial derivatives
+	// by means of 2x2 quad differentiation.
+	// If bHardwareDerivative is false, this argument value kind supports generating the differential somehow. For example
+	// Arg may be an Extern or GetVertexInterpolator instruction.
 	FValue* Arg;
 
-	// The direction of partial derivative
-	EDerivativeAxis Axis;
-
-	// The source expression that generated this partial derivative (for error reporting)
-	EDerivativeSource Source;
+	// The direction of partial derivative.
+	EDerivativeAxis Axis : 1;
+	
+	// Whether this should use hardware differential calculation.
+	bool bHardwareDerivative : 1;
 };
 
 // Instruction executed in the vertex stage that sets the value of some interpolated data, available
diff --git a/Engine/Source/Runtime/Engine/Public/Materials/MaterialIRExtern.cpp b/Engine/Source/Runtime/Engine/Public/Materials/MaterialIRExtern.cpp
--- a/Engine/Source/Runtime/Engine/Public/Materials/MaterialIRExtern.cpp
+++ b/Engine/Source/Runtime/Engine/Public/Materials/MaterialIRExtern.cpp
@@ -139,7 +139,7 @@ void FExternFromMaterialDecl::ToHLSL(FExternPrinterHLSL& Printer) const
 
 void FExternFromMaterialDecl::EmitDebugInfo(FString& Out) const
 {
-	Out.Appendf(TEXT("Name=%s"), *Name.ToString());
+	Out.Appendf(TEXT("Name=\"%s\""), *Name.ToString());
 }
 
 //------------------------------------------------------------------------
