[FX] Fix null pointer crash in FXSystemAsset when r.PSOPrecache.UseBackgroundThreadForCollection=0.

#rnx
#Synced-CL 51745034

diff --git a/Engine/Source/Runtime/Engine/Private/Particles/ParticleSystem.cpp b/Engine/Source/Runtime/Engine/Private/Particles/ParticleSystem.cpp
--- a/Engine/Source/Runtime/Engine/Private/Particles/ParticleSystem.cpp
+++ b/Engine/Source/Runtime/Engine/Private/Particles/ParticleSystem.cpp
@@ -92,7 +92,7 @@ bool UFXSystemAsset::IsLoadingPSOCache(UMaterialInterface* Material) const
 	}
 	for (const FGraphEventRef& Task : MaterialPSORequests[Material])
 	{
-		if (!Task->IsCompleted())
+		if (Task && !Task->IsCompleted())
 		{
 			return true;
 		}
