Fixed issue where Instanced Skinned Mesh Components have invalid bounds due to uninitialized memory.

UAnimSequenceTransformProviderDataInstance::GetAnimationBounds was being called before async compilation had finished, returning uninitialized data from SequenceBounds.

This fix ensures that IsCompiling only returns true when compilation has finished, which prevents GetAnimationBounds from being called too soon.

#Synced-CL 51835690

diff --git a/Engine/Source/Runtime/Engine/Classes/Animation/AnimSequenceTransformProviderData.h b/Engine/Source/Runtime/Engine/Classes/Animation/AnimSequenceTransformProviderData.h
--- a/Engine/Source/Runtime/Engine/Classes/Animation/AnimSequenceTransformProviderData.h
+++ b/Engine/Source/Runtime/Engine/Classes/Animation/AnimSequenceTransformProviderData.h
@@ -548,6 +548,7 @@ public:
 	virtual uint32 GetSkinningDataOffset(int32 InstanceIndex, const FTransform& ComponentTransform, const FSkinnedMeshInstanceData& InstanceData) const override;
 	virtual bool HasAnimationBounds() const;
 	virtual bool GetAnimationBounds(uint32 AnimationIndex, FRenderBounds& OutBounds) const override;
+	virtual bool IsCompiling() const override;
 	virtual FTransformProviderRenderProxy* CreateRenderProxy(FInstancedSkinningSceneExtensionProxy* ExtensionProxy) const override;
 	virtual void BeginDestroy() override;
 	virtual void Serialize(FArchive& Ar) override;
diff --git a/Engine/Source/Runtime/Engine/Private/Animation/AnimSequenceTransformProviderData.cpp b/Engine/Source/Runtime/Engine/Private/Animation/AnimSequenceTransformProviderData.cpp
--- a/Engine/Source/Runtime/Engine/Private/Animation/AnimSequenceTransformProviderData.cpp
+++ b/Engine/Source/Runtime/Engine/Private/Animation/AnimSequenceTransformProviderData.cpp
@@ -937,6 +937,16 @@ bool UAnimSequenceTransformProviderDataInstance::GetAnimationBounds(uint32 Track
 	return false;
 }
 
+bool UAnimSequenceTransformProviderDataInstance::IsCompiling() const
+{
+    if (ParentProvider && ParentProvider->IsCompiling())
+    {
+        return true;
+    }
+
+    return Super::IsCompiling();
+}
+
 FTransformProviderRenderProxy* UAnimSequenceTransformProviderDataInstance::CreateRenderProxy(FInstancedSkinningSceneExtensionProxy* SceneExtensionProxy) const
 {
 	if (IsValidFor(SceneExtensionProxy))
