Added setter for Instanced Skinned Mesh Component's AnimationMinScreenSize, so that it can be set from code

#Synced-CL 51835690

diff --git a/Engine/Source/Runtime/Engine/Classes/Components/InstancedSkinnedMeshComponent.h b/Engine/Source/Runtime/Engine/Classes/Components/InstancedSkinnedMeshComponent.h
--- a/Engine/Source/Runtime/Engine/Classes/Components/InstancedSkinnedMeshComponent.h
+++ b/Engine/Source/Runtime/Engine/Classes/Components/InstancedSkinnedMeshComponent.h
@@ -274,6 +274,9 @@ public:
 	UFUNCTION(BlueprintCallable, Category = "Components|InstancedSkinnedMesh")
 	ENGINE_API FPrimitiveInstanceId GetInstanceId(int32 InstanceIndex) const;
 
+	UFUNCTION(BlueprintCallable, Category = "Components|InstancedSkinnedMesh")
+	ENGINE_API void SetAnimationMinScreenSize(float InAnimationMinScreenSize);
+
 	int32 GetNumCustomDataFloats() const { return NumCustomDataFloats; }
 	const TArray<float>& GetInstanceCustomData() const { return InstanceCustomData; }
 	const TArray<FSkinnedMeshInstanceData>& GetInstanceData() const { return InstanceData; }
diff --git a/Engine/Source/Runtime/Engine/Private/Components/InstancedSkinnedMeshComponent.cpp b/Engine/Source/Runtime/Engine/Private/Components/InstancedSkinnedMeshComponent.cpp
--- a/Engine/Source/Runtime/Engine/Private/Components/InstancedSkinnedMeshComponent.cpp
+++ b/Engine/Source/Runtime/Engine/Private/Components/InstancedSkinnedMeshComponent.cpp
@@ -667,6 +667,16 @@ void UInstancedSkinnedMeshComponent::SetCullDistances(int32 StartCullDistance, i
 	}
 }
 
+void UInstancedSkinnedMeshComponent::SetAnimationMinScreenSize(float InAnimationMinScreenSize)
+{
+	if (AnimationMinScreenSize != InAnimationMinScreenSize)
+	{
+		AnimationMinScreenSize = InAnimationMinScreenSize;
+
+		MarkRenderStateDirty();
+	}
+}
+
 void UInstancedSkinnedMeshComponent::PreApplyComponentInstanceData(struct FInstancedSkinnedMeshComponentInstanceData* InstancedMeshData)
 {
 #if WITH_EDITOR
