Fix long periods of very low framerate when opening level in editor

In a test sandbox project, reduces time to fully load a level from 25s to 12s.

#rnx
#Synced-CL 51776036

diff --git a/Engine/Source/Runtime/Engine/Private/AssetCompilingManager.cpp b/Engine/Source/Runtime/Engine/Private/AssetCompilingManager.cpp
--- a/Engine/Source/Runtime/Engine/Private/AssetCompilingManager.cpp
+++ b/Engine/Source/Runtime/Engine/Private/AssetCompilingManager.cpp
@@ -32,6 +32,7 @@ LLM_DEFINE_TAG(AssetCompilation, NAME_None, NAME_None, GET_STATFNAME(STAT_AssetC
 #include "DerivedDataBuildSchedulerQueue.h"
 #include "DerivedDataThreadPoolTask.h"
 #include "Features/IModularFeatures.h"
+#include "Framework/Application/SlateApplication.h"
 #include "Interfaces/Interface_AsyncCompilation.h"
 #include "Logging/StructuredLog.h"
 #endif
@@ -802,6 +803,23 @@ void FAssetCompilingManager::UpdateNumRemainingAssets()
 		FAssetCompilingManager::Get().OnAssetPostCompileEvent().Broadcast({});
 	}
 	LastNumRemainingAssets = NumRemainingAssets;
+
+#if WITH_EDITOR
+	// Hold a Slate responsive-mode request while there is pending work, so the editor ticks at full rate
+	// and the per-frame processing budget is consumed as fast as possible. Without this, Slate can
+	// throttle itself during load (no user input), causing asset compilation to take far longer than necessary.
+	if (NumRemainingAssets > 0)
+	{
+		if (!ResponsiveModeHandle.IsValid() && FSlateApplication::IsInitialized())
+		{
+			ResponsiveModeHandle = FSlateThrottleManager::Get().EnterResponsiveMode();
+		}
+	}
+	else if (ResponsiveModeHandle.IsValid())
+	{
+		FSlateThrottleManager::Get().LeaveResponsiveMode(ResponsiveModeHandle);
+	}
+#endif
 }
 
 #undef LOCTEXT_NAMESPACE
diff --git a/Engine/Source/Runtime/Engine/Public/AssetCompilingManager.h b/Engine/Source/Runtime/Engine/Public/AssetCompilingManager.h
--- a/Engine/Source/Runtime/Engine/Public/AssetCompilingManager.h
+++ b/Engine/Source/Runtime/Engine/Public/AssetCompilingManager.h
@@ -5,6 +5,9 @@
 #include "UObject/WeakObjectPtr.h"
 #include "UObject/Class.h"
 #include "UObject/ObjectMacros.h"
+#if WITH_EDITOR
+#include "Application/ThrottleManager.h"
+#endif
 #include "AssetCompilingManager.generated.h"
 
 class FQueuedThreadPool;
@@ -130,4 +133,9 @@ private:
 	FAssetPostCompileEvent AssetPostCompileEvent;
 
 	FPackageScopeEvent PackageScopeEvent;
+
+#if WITH_EDITOR
+	/** Held while there are pending assets to prevent Slate from throttling the editor tick rate during compilation. */
+	FThrottleRequest ResponsiveModeHandle;
+#endif
 };
