Fix stack overflow crash in the Material Editor when making loops in the graph w/ reroute nodes.

- StaticSwitchParameter plugged into a RerouteNode plugged back into the same StaticSwitchParameter.
- This created a cycle: Reroute::IsResultMaterialAttributes -> TraceInputsToRealExpression  -> StaticSwitchParameter::IsResultMaterialAttributes  -> back to Reroute.
- This affects any material expression that overrides IsResultMaterialAttributes and recurses into its inputs (StaticSwitch, If,  QualitySwitch, FeatureLevelSwitch, LandscapeLayerSwitch, etc).
- Add ContainsInputLoop() guards to all fourUMaterialExpressionRerouteBase methods that delegate to resolved expressions: IsResultMaterialAttributes, IsResultSubstrateMaterial, GatherSubstrateMaterialInfo, and SubstrateGenerateMaterialTopologyTree.
- Refactor the code to reduce duplication between all 4 implementations of that function.

[RN] major | Rendering - Materials and Shaders
#Synced-CL 51761382

diff --git a/Engine/Source/Runtime/Engine/Private/Materials/MaterialExpressions.cpp b/Engine/Source/Runtime/Engine/Private/Materials/MaterialExpressions.cpp
--- a/Engine/Source/Runtime/Engine/Private/Materials/MaterialExpressions.cpp
+++ b/Engine/Source/Runtime/Engine/Private/Materials/MaterialExpressions.cpp
@@ -18016,83 +18016,61 @@ EMaterialValueType UMaterialExpressionRerouteBase::GetOutputValueType(int32 Outp
 	return GetInputValueType(0);
 }
 
-bool UMaterialExpressionRerouteBase::IsResultMaterialAttributes(int32 OutputIndex)
+UMaterialExpression* UMaterialExpressionRerouteBase::ResolveRerouteTarget(int32 OutputIndex, int32& OutRealOutputIndex)
 {
 	FExpressionInput Input;
 	if (GetRerouteInput(Input))
 	{
-		// Most code checks to make sure that there aren't loops before going here. In our case, we rely on the fact that
-		// UMaterialExpressionReroute's implementation of TraceInputsToRealExpression is resistant to input loops.
 		if (Input.IsConnected() && Input.Expression != nullptr && OutputIndex == 0)
 		{
-			int32 RealExpressionOutputIndex = -1;
-			UMaterialExpression* RealExpression = TraceInputsToRealExpression(RealExpressionOutputIndex);
-			if (RealExpression != nullptr)
+			// Guard against cycles that pass through non-reroute expressions (e.g. Reroute -> StaticSwitch -> Reroute).
+			if (ContainsInputLoop())
 			{
-				return RealExpression->IsResultMaterialAttributes(RealExpressionOutputIndex);
+				return nullptr;
 			}
+
+			OutRealOutputIndex = -1;
+			return TraceInputsToRealExpression(OutRealOutputIndex);
 		}
 	}
+	return nullptr;
+}
 
+bool UMaterialExpressionRerouteBase::IsResultMaterialAttributes(int32 OutputIndex)
+{
+	int32 RealOutputIndex = -1;
+	if (UMaterialExpression* RealExpression = ResolveRerouteTarget(OutputIndex, RealOutputIndex))
+	{
+		return RealExpression->IsResultMaterialAttributes(RealOutputIndex);
+	}
 	return false;
 }
 
 bool UMaterialExpressionRerouteBase::IsResultSubstrateMaterial(int32 OutputIndex)
 {
-	FExpressionInput Input;
-	if (GetRerouteInput(Input))
+	int32 RealOutputIndex = -1;
+	if (UMaterialExpression* RealExpression = ResolveRerouteTarget(OutputIndex, RealOutputIndex))
 	{
-		// Most code checks to make sure that there aren't loops before going here. In our case, we rely on the fact that
-		// UMaterialExpressionReroute's implementation of TraceInputsToRealExpression is resistant to input loops.
-		if (Input.IsConnected() && Input.Expression != nullptr && OutputIndex == 0)
-		{
-			int32 RealExpressionOutputIndex = -1;
-			UMaterialExpression* RealExpression = TraceInputsToRealExpression(RealExpressionOutputIndex);
-			if (RealExpression != nullptr)
-			{
-				return RealExpression->IsResultSubstrateMaterial(RealExpressionOutputIndex);
-			}
-		}
+		return RealExpression->IsResultSubstrateMaterial(RealOutputIndex);
 	}
-
 	return false;
 }
 
 void UMaterialExpressionRerouteBase::GatherSubstrateMaterialInfo(FSubstrateMaterialInfo& SubstrateMaterialInfo, int32 OutputIndex)
 {
-	FExpressionInput Input;
-	if (GetRerouteInput(Input))
+	int32 RealOutputIndex = -1;
+	if (UMaterialExpression* RealExpression = ResolveRerouteTarget(OutputIndex, RealOutputIndex))
 	{
-		// Most code checks to make sure that there aren't loops before going here. In our case, we rely on the fact that
-		// UMaterialExpressionReroute's implementation of TraceInputsToRealExpression is resistant to input loops.
-		if (Input.IsConnected() && Input.Expression != nullptr && OutputIndex == 0)
-		{
-			int32 RealExpressionOutputIndex = -1;
-			UMaterialExpression* RealExpression = TraceInputsToRealExpression(RealExpressionOutputIndex);
-			if (RealExpression != nullptr)
-			{
-				RealExpression->GatherSubstrateMaterialInfo(SubstrateMaterialInfo, RealExpressionOutputIndex);
-			}
-		}
+		RealExpression->GatherSubstrateMaterialInfo(SubstrateMaterialInfo, RealOutputIndex);
 	}
 }
 
 FSubstrateOperator* UMaterialExpressionRerouteBase::SubstrateGenerateMaterialTopologyTree(struct FSubstrateTranslatorDataInterface& SubstrateTranslatorData, class UMaterialExpression* Parent, int32 OutputIndex)
 {
-	FExpressionInput Input;
-	if (GetRerouteInput(Input))
+	int32 RealOutputIndex = -1;
+	if (UMaterialExpression* RealExpression = ResolveRerouteTarget(OutputIndex, RealOutputIndex))
 	{
-		// Most code checks to make sure that there aren't loops before going here. In our case, we rely on the fact that
-		// UMaterialExpressionReroute's implementation of TraceInputsToRealExpression is resistant to input loops.
-		if (Input.IsConnected() && Input.Expression != nullptr && OutputIndex == 0)
-		{
-			int32 RealExpressionOutputIndex = -1;
-			UMaterialExpression* RealExpression = TraceInputsToRealExpression(RealExpressionOutputIndex);
-			if (RealExpression != nullptr)
-			{
-				RealExpression->SubstrateGenerateMaterialTopologyTree(SubstrateTranslatorData, Parent, RealExpressionOutputIndex);
-			}
-		}
+		return RealExpression->SubstrateGenerateMaterialTopologyTree(SubstrateTranslatorData, Parent, RealOutputIndex);
 	}
 	return nullptr;
 }
diff --git a/Engine/Source/Runtime/Engine/Public/Materials/MaterialExpressionRerouteBase.h b/Engine/Source/Runtime/Engine/Public/Materials/MaterialExpressionRerouteBase.h
--- a/Engine/Source/Runtime/Engine/Public/Materials/MaterialExpressionRerouteBase.h
+++ b/Engine/Source/Runtime/Engine/Public/Materials/MaterialExpressionRerouteBase.h
@@ -44,5 +44,9 @@ protected:
 	 virtual bool GetRerouteInput(FExpressionInput& OutInput) const { return false; }
 
 private:
+#if WITH_EDITOR
+	/** Resolve the reroute target expression with loop-detection guards. Returns nullptr if unconnected or a cycle is detected. */
+	UMaterialExpression* ResolveRerouteTarget(int32 OutputIndex, int32& OutRealOutputIndex);
+#endif
 	ENGINE_API FExpressionInput TraceInputsToRealExpressionInternal(TSet<FMaterialExpressionKey>& VisitedExpressions) const;
 };
