Localize anim sequence asset registry tags

#Synced-CL 51880625

diff --git a/Engine/Source/Runtime/Engine/Classes/Animation/AnimSequence.h b/Engine/Source/Runtime/Engine/Classes/Animation/AnimSequence.h
--- a/Engine/Source/Runtime/Engine/Classes/Animation/AnimSequence.h
+++ b/Engine/Source/Runtime/Engine/Classes/Animation/AnimSequence.h
@@ -388,6 +388,9 @@ public:
 	ENGINE_API virtual void GetAssetRegistryTags(FAssetRegistryTagsContext Context) const override;
 	UE_DEPRECATED(5.4, "Implement the version that takes FAssetRegistryTagsContext instead.")
 	ENGINE_API virtual void GetAssetRegistryTags(TArray<FAssetRegistryTag>& OutTags) const override;
+#if WITH_EDITOR
+	ENGINE_API virtual void GetAssetRegistryTagMetadata(TMap<FName, FAssetRegistryTagMetadata>& OutMetadata) const override;
+#endif // WITH_EDITOR
 	static ENGINE_API void AddReferencedObjects(UObject* This, FReferenceCollector& Collector);
 	//~ End UObject Interface
 
diff --git a/Engine/Source/Runtime/Engine/Private/Animation/AnimSequence.cpp b/Engine/Source/Runtime/Engine/Private/Animation/AnimSequence.cpp
--- a/Engine/Source/Runtime/Engine/Private/Animation/AnimSequence.cpp
+++ b/Engine/Source/Runtime/Engine/Private/Animation/AnimSequence.cpp
@@ -428,6 +428,13 @@ void UAnimSequence::GetAssetRegistryTags(TArray<FAssetRegistryTag>& OutTags) con
 	PRAGMA_ENABLE_DEPRECATION_WARNINGS;
 }
 
+static const FName AssetRegistryTag_CompressionRatio(TEXT("Compression Ratio"));
+static const FName AssetRegistryTag_SourceFrameRate(TEXT("Source Frame Rate"));
+static const FName AssetRegistryTag_NumberOfFrames(TEXT("Number of Frames"));
+static const FName AssetRegistryTag_NumberOfKeys(TEXT("Number of Keys"));
+static const FName AssetRegistryTag_CompressedSize(TEXT("Compressed Size (KB)"));
+static const FName AssetRegistryTag_TargetFrameRate(TEXT("Target Frame Rate"));
+
 void UAnimSequence::GetAssetRegistryTags(FAssetRegistryTagsContext Context) const
 {
 #if WITH_EDITORONLY_DATA
@@ -438,16 +445,16 @@ void UAnimSequence::GetAssetRegistryTags(FAssetRegistryTagsContext Context) cons
 
 	if (DataModelInterface.GetObject() && DataModelInterface->HasBeenPopulated())
 	{
-		Context.AddTag(FAssetRegistryTag(TEXT("Compression Ratio"), FString::Printf(TEXT("%.03f"), (double)GetApproxCompressedSize() / (double)GetUncompressedRawSize()), FAssetRegistryTag::TT_Numerical));
-		Context.AddTag(FAssetRegistryTag(TEXT("Source Frame Rate"), FString::Printf(TEXT("%.2f"), DataModelInterface->GetFrameRate().AsDecimal()), FAssetRegistryTag::TT_Numerical));
-		Context.AddTag(FAssetRegistryTag(TEXT("Number of Frames"), FString::Printf(TEXT("%.i"), DataModelInterface->GetNumberOfFrames()), FAssetRegistryTag::TT_Numerical));
-		Context.AddTag(FAssetRegistryTag(TEXT("Number of Keys"), FString::Printf(TEXT("%i"), DataModelInterface->GetNumberOfKeys()), FAssetRegistryTag::TT_Numerical));
+		Context.AddTag(FAssetRegistryTag(AssetRegistryTag_CompressionRatio, FString::Printf(TEXT("%.03f"), (double)GetApproxCompressedSize() / (double)GetUncompressedRawSize()), FAssetRegistryTag::TT_Numerical));
+		Context.AddTag(FAssetRegistryTag(AssetRegistryTag_SourceFrameRate, FString::Printf(TEXT("%.2f"), DataModelInterface->GetFrameRate().AsDecimal()), FAssetRegistryTag::TT_Numerical));
+		Context.AddTag(FAssetRegistryTag(AssetRegistryTag_NumberOfFrames, FString::Printf(TEXT("%.i"), DataModelInterface->GetNumberOfFrames()), FAssetRegistryTag::TT_Numerical));
+		Context.AddTag(FAssetRegistryTag(AssetRegistryTag_NumberOfKeys, FString::Printf(TEXT("%i"), DataModelInterface->GetNumberOfKeys()), FAssetRegistryTag::TT_Numerical));
 	}
 	else
 	{
-		Context.AddTag(FAssetRegistryTag(TEXT("Source Frame Rate"), TEXT("0.0"), FAssetRegistryTag::TT_Numerical));
-		Context.AddTag(FAssetRegistryTag(TEXT("Number of Frames"), TEXT("0"), FAssetRegistryTag::TT_Numerical));
-		Context.AddTag(FAssetRegistryTag(TEXT("Number of Keys"), TEXT("0"), FAssetRegistryTag::TT_Numerical));
+		Context.AddTag(FAssetRegistryTag(AssetRegistryTag_SourceFrameRate, TEXT("0.0"), FAssetRegistryTag::TT_Numerical));
+		Context.AddTag(FAssetRegistryTag(AssetRegistryTag_NumberOfFrames, TEXT("0"), FAssetRegistryTag::TT_Numerical));
+		Context.AddTag(FAssetRegistryTag(AssetRegistryTag_NumberOfKeys, TEXT("0"), FAssetRegistryTag::TT_Numerical));
 	}
 
 	// Output unique sync marker names we use
@@ -463,11 +470,25 @@ void UAnimSequence::GetAssetRegistryTags(FAssetRegistryTagsContext Context) cons
 	Context.AddTag(FAssetRegistryTag(USkeleton::AnimSyncMarkerTag, SyncMarkersBuilder.ToString(), FAssetRegistryTag::TT_Hidden));
 #endif
 
-	Context.AddTag(FAssetRegistryTag(TEXT("Compressed Size (KB)"), FString::Printf(TEXT("%.02f"), (float)GetApproxCompressedSize() / 1024.0f), FAssetRegistryTag::TT_Numerical));
-	Context.AddTag(FAssetRegistryTag(TEXT("Target Frame Rate"), FString::Printf(TEXT("%.2f"), PlatformTargetFrameRate.Default.AsDecimal()), FAssetRegistryTag::TT_Numerical));
+	Context.AddTag(FAssetRegistryTag(AssetRegistryTag_CompressedSize, FString::Printf(TEXT("%.02f"), (float)GetApproxCompressedSize() / 1024.0f), FAssetRegistryTag::TT_Numerical));
+	Context.AddTag(FAssetRegistryTag(AssetRegistryTag_TargetFrameRate, FString::Printf(TEXT("%.2f"), PlatformTargetFrameRate.Default.AsDecimal()), FAssetRegistryTag::TT_Numerical));
 	Super::GetAssetRegistryTags(Context);
 }
 
+#if WITH_EDITOR
+
+void UAnimSequence::GetAssetRegistryTagMetadata(TMap<FName, FAssetRegistryTagMetadata>& OutMetadata) const
+{
+	OutMetadata.Add(AssetRegistryTag_CompressionRatio, FAssetRegistryTagMetadata().SetDisplayName(LOCTEXT("CompressionRatio", "Compression Ratio")));
+	OutMetadata.Add(AssetRegistryTag_SourceFrameRate, FAssetRegistryTagMetadata().SetDisplayName(LOCTEXT("SourceFrameRate", "Source Frame Rate")));
+	OutMetadata.Add(AssetRegistryTag_NumberOfFrames, FAssetRegistryTagMetadata().SetDisplayName(LOCTEXT("NumberOfFrames", "Number of Frames")));
+	OutMetadata.Add(AssetRegistryTag_NumberOfKeys, FAssetRegistryTagMetadata().SetDisplayName(LOCTEXT("NumberOfKeys", "Number of Keys")));
+	OutMetadata.Add(AssetRegistryTag_CompressedSize, FAssetRegistryTagMetadata().SetDisplayName(LOCTEXT("CompressedSize", "Compressed Size (KB)")));
+	OutMetadata.Add(AssetRegistryTag_TargetFrameRate, FAssetRegistryTagMetadata().SetDisplayName(LOCTEXT("TargetFrameRate", "Target Frame Rate")));
+}
+
+#endif // WITH_EDITOR
+
 void UAnimSequence::AddReferencedObjects(UObject* This, FReferenceCollector& Collector)
 {
 	Super::AddReferencedObjects(This, Collector);
