NMT : fix for WorldPosition_DDX/DDY

#rnx
#Synced-CL 51883779

diff --git a/Engine/Source/Runtime/Engine/Private/Materials/MaterialExpressionsToMIR.cpp b/Engine/Source/Runtime/Engine/Private/Materials/MaterialExpressionsToMIR.cpp
--- a/Engine/Source/Runtime/Engine/Private/Materials/MaterialExpressionsToMIR.cpp
+++ b/Engine/Source/Runtime/Engine/Private/Materials/MaterialExpressionsToMIR.cpp
@@ -2197,7 +2197,7 @@ static FValueRef EmitWorldPosition(MIR::FEmitter& Em, EWorldPositionIncludedOffs
 
 				default:
 					// Expression emitter assumes these to be LWC for analytic derivative evaluation, so we promote them on load
-					Printer << TEXTVIEW("WSPromote(Parameters.WorldPosition_)") << MIR::ExternTag_DD;
+					Printer << TEXTVIEW("WSPromote(Parameters.WorldPosition_") << MIR::ExternTag_DD << TEXT(')');
 					break;
 			}
 		}
@@ -7521,4 +7521,4 @@ void UMaterialExpressionSubstrateThinFilm::Build(MIR::FEmitter& Em)
 	Em.Output(1, OutF90Arg);
 }
 
-#endif // WITH_EDITOR
\ No newline at end of file
+#endif // WITH_EDITOR
