[World Partition] Fix missing package mappings for cross-cell external actor references in LI instancing contexts

When loading Level Instance container actors via LoadActorsInternal, cross-cell external actor references were not mapped in the instancing context, causing PopulateInstancingContext warnings about unmapped packages with instanced outers.
Add an inline mapping function to dynamically resolve these references by matching external actor package paths and generating the correct instanced package name on-demand.

#Synced-CL 51685151

diff --git a/Engine/Source/Runtime/Engine/Private/WorldPartition/WorldPartitionLevelHelper.cpp b/Engine/Source/Runtime/Engine/Private/WorldPartition/WorldPartitionLevelHelper.cpp
--- a/Engine/Source/Runtime/Engine/Private/WorldPartition/WorldPartitionLevelHelper.cpp
+++ b/Engine/Source/Runtime/Engine/Private/WorldPartition/WorldPartitionLevelHelper.cpp
@@ -801,6 +801,24 @@ bool FWorldPartitionLevelHelper::LoadActorsInternal(FLoadActorsParams&& InParams
 				const bool bUniquePackage = !PackageObjectMapping.Package.ToString().Contains(FPackagePath::GetExternalActorsFolderName());
 				const FName ContainerPackageInstanceName(GetContainerPackage(PackageObjectMapping.ContainerID, PackageObjectMapping.ContainerPackage.ToString(), InParams.OuterWorld, bUniquePackage));
 				NewContext.AddPackageMapping(PackageObjectMapping.ContainerPackage, ContainerPackageInstanceName);
+
+				// Add dynamic mapping function to handle cross-cell external actor references within this sub-container.
+				// During linker load, references can pull in actors from other cells of the same level instance (e.g., editor-only properties). 
+				// This function generates the correct instanced package name on-demand for any external actor package under this
+				// container's path, avoiding PopulateInstancingContext warnings about unmapped packages with instanced outers.
+				const FString ExternalActorsPathStr = ULevel::GetExternalActorsPath(PackageObjectMapping.ContainerPackage.ToString());
+				const FString InstancedLevelPackageNameStr = ContainerPackageInstanceName.ToString();
+				
+				NewContext.AddPackageMappingFunc([ExternalActorsPathStr, InstancedLevelPackageNameStr](FName Original)
+				{
+					const FString OriginalStr = Original.ToString();
+					if (OriginalStr.StartsWith(ExternalActorsPathStr))
+					{
+						return FName(*ULevel::GetExternalActorPackageInstanceName(InstancedLevelPackageNameStr, OriginalStr));
+					}
+					return Original;
+				});
+
 				Context = &NewContext;
 			}
 		
