Add missing object references in ActiveSound::AddReferencedObjects

#rnx
#Synced-CL 51808504

diff --git a/Engine/Source/Runtime/Engine/Private/ActiveSound.cpp b/Engine/Source/Runtime/Engine/Private/ActiveSound.cpp
--- a/Engine/Source/Runtime/Engine/Private/ActiveSound.cpp
+++ b/Engine/Source/Runtime/Engine/Private/ActiveSound.cpp
@@ -326,6 +326,20 @@ void FActiveSound::AddReferencedObjects(FReferenceCollector& Collector)
 		Collector.AddReferencedObject(BusSendTuple.Value.SoundSourceBus);
 	}
 
+	for (TArray<FSoundSourceBusSendInfo>& SendArray : BusSendsOverride)
+	{
+		for (FSoundSourceBusSendInfo& BusSend : SendArray)
+		{
+			Collector.AddReferencedObject(BusSend.AudioBus);
+			Collector.AddReferencedObject(BusSend.SoundSourceBus);
+		}
+	}
+
+	for (FSoundSubmixSendInfo& SendInfo : SoundSubmixSendsOverride)
+	{
+		Collector.AddReferencedObject(SendInfo.SoundSubmix);
+	}
+
 	Collector.AddReferencedObjects(ModulationRouting.VolumeModulationDestination.Modulators);
 	Collector.AddReferencedObjects(ModulationRouting.PitchModulationDestination.Modulators);
 	Collector.AddReferencedObjects(ModulationRouting.HighpassModulationDestination.Modulators);
