Fix issues with PreviewPlatform when Active DeviceProfile != CookingDeviceProfile
The TargetPlatform setting must come from CookingDeviceProfile
Preview override must be configured on BOTH
this can happen due to  -deviceprofile=WindowsEditorSaveVRAM

#Synced-CL 51892244

diff --git a/Engine/Source/Editor/UnrealEd/Private/EditorEngine.cpp b/Engine/Source/Editor/UnrealEd/Private/EditorEngine.cpp
--- a/Engine/Source/Editor/UnrealEd/Private/EditorEngine.cpp
+++ b/Engine/Source/Editor/UnrealEd/Private/EditorEngine.cpp
@@ -6637,22 +6637,22 @@ static void PreviewPlatformChangePost()
 	{
 		// Need to update the cached TextureLODSettings in TargetPlatform with RegisterTextureLODSettings
 
-		const UTextureLODSettings * TLS_DP = UDeviceProfileManager::Get().GetActiveProfile()->GetTextureLODSettings();
+		UDeviceProfile * ADP = UDeviceProfileManager::Get().GetActiveProfile();
+		const UTextureLODSettings * TLS_ADP = ADP->GetTextureLODSettings();
 
 		ITargetPlatformManagerModule* TPM = GetTargetPlatformManager();
 		ITargetPlatform* TargetPlatform = TPM->GetRunningTargetPlatform();
-		//const UTextureLODSettings & TLS_TP = TargetPlatform->GetTextureLODSettings();
 		
-		// Alternate way to get DP : (should be used when cooking)
-		//  during interactive editor run, this should == TLS_DP
-		UTextureLODSettings * TLS_DP2 = UDeviceProfileManager::Get().FindProfile(TargetPlatform->CookingDeviceProfileName())->GetTextureLODSettings();
+		UDeviceProfile * CDP = UDeviceProfileManager::Get().FindProfile(TargetPlatform->CookingDeviceProfileName());
+		UTextureLODSettings * TLS_CDP = CDP->GetTextureLODSettings();
 
-		check( TLS_DP == TLS_DP2 );
+		// typically ADP == CDP , but not always
+		// when different, TextureLOD should come from CDP
+		// if PreviewPlatform is enabled, then BOTH ADP and CDP should return the override
 
-		TargetPlatform->RegisterTextureLODSettings(TLS_DP);
+		TargetPlatform->RegisterTextureLODSettings(TLS_CDP);
 
-		// after Register, TP->GetTextureLODSettings is equal to the DP way
-		check( &(TargetPlatform->GetTextureLODSettings()) == TLS_DP );
+		// after Register, TP->GetTextureLODSettings is equal to the CDP way
 	}
 
 	{
diff --git a/Engine/Source/Runtime/Engine/Private/DeviceProfiles/DeviceProfileManager.cpp b/Engine/Source/Runtime/Engine/Private/DeviceProfiles/DeviceProfileManager.cpp
--- a/Engine/Source/Runtime/Engine/Private/DeviceProfiles/DeviceProfileManager.cpp
+++ b/Engine/Source/Runtime/Engine/Private/DeviceProfiles/DeviceProfileManager.cpp
@@ -1372,6 +1372,23 @@ void UDeviceProfileManager::SetPreviewDeviceProfile(UDeviceProfile* DeviceProfil
 	PreviewDeviceProfile = DeviceProfile;
 	
 	ActiveDeviceProfile->SetPreviewTextureGroupRedirect(PreviewDeviceProfile);
+	
+#if WITH_EDITOR
+	{
+		ITargetPlatformManagerModule* TPM = GetTargetPlatformManager();
+		ITargetPlatform* TargetPlatform = TPM->GetRunningTargetPlatform();
+		
+		UDeviceProfile * CDP = UDeviceProfileManager::Get().FindProfile(
+			TargetPlatform->CookingDeviceProfileName());
+
+		// usually Active == CDP
+		// but if not then they must be kept in sync
+		if ( CDP != ActiveDeviceProfile )
+		{
+			CDP->SetPreviewTextureGroupRedirect(PreviewDeviceProfile);
+		}
+	}
+#endif
 
 	// walk over all cvars and apply them
 	IConsoleManager::Get().PreviewPlatformCVars(*DeviceProfile->DeviceType, DeviceProfile->GetName(), PreviewModeTag);
@@ -1410,6 +1427,23 @@ void UDeviceProfileManager::RestorePreviewDeviceProfile(FName PreviewModeTag)
 	PreviewDeviceProfile = nullptr;
 
 	ActiveDeviceProfile->SetPreviewTextureGroupRedirect(nullptr);
+	
+#if WITH_EDITOR
+	{
+		ITargetPlatformManagerModule* TPM = GetTargetPlatformManager();
+		ITargetPlatform* TargetPlatform = TPM->GetRunningTargetPlatform();
+		
+		UDeviceProfile * CDP = UDeviceProfileManager::Get().FindProfile(
+			TargetPlatform->CookingDeviceProfileName());
+
+		// usually Active == CDP
+		// but if not then they must be kept in sync
+		if ( CDP != ActiveDeviceProfile )
+		{
+			CDP->SetPreviewTextureGroupRedirect(nullptr);
+		}
+	}
+#endif
 
 	PreviewDeviceProfileScalabilityCVars.Empty();
 
