Prepare FWorldPartitionClassDescRegistry for upcoming parallel save by protecting ClassByPath, ParentClassMap, and RedirectClassMap with an FTransactionallySafeLock.
These maps are currently accessed without synchronization, which will race when save operations run concurrently on worker threads.

#Synced-CL 51944883

diff --git a/Engine/Source/Runtime/Engine/Private/WorldPartition/WorldPartitionClassDescRegistry.cpp b/Engine/Source/Runtime/Engine/Private/WorldPartition/WorldPartitionClassDescRegistry.cpp
--- a/Engine/Source/Runtime/Engine/Private/WorldPartition/WorldPartitionClassDescRegistry.cpp
+++ b/Engine/Source/Runtime/Engine/Private/WorldPartition/WorldPartitionClassDescRegistry.cpp
@@ -177,9 +177,14 @@ bool FWorldPartitionClassDescRegistry::IsInitialized() const
 
 void FWorldPartitionClassDescRegistry::PrefetchClassDescs(const TArray<FTopLevelAssetPath>& InClassPaths)
 {
-	TRACE_CPUPROFILER_EVENT_SCOPE(FWorldPartitionClassDescRegistry::PrefetchClassDescs);
-
 	check(IsInitialized());
+	UE::TScopeLock Lock(ClassDescLock);
+	PrefetchClassDescsInternal(InClassPaths);
+}
+
+void FWorldPartitionClassDescRegistry::PrefetchClassDescsInternal(const TArray<FTopLevelAssetPath>& InClassPaths)
+{
+	TRACE_CPUPROFILER_EVENT_SCOPE(FWorldPartitionClassDescRegistry::PrefetchClassDescs);
 
 	// This function makes asset registry request during the cook. The tracker will register those requests as dependencies. The issue
 	// is the requests are not always made by the same package so I end up with non deterministic dependencies.
@@ -192,7 +197,7 @@ void FWorldPartitionClassDescRegistry::PrefetchClassDescs(const TArray<FTopLevel
 	// First, remap classes and remove already known ones.
 	TArray<FTopLevelAssetPath> ClassPaths;
 	ClassPaths.Reserve(InClassPaths.Num());
-	Algo::Transform(InClassPaths, ClassPaths, [this](const FTopLevelAssetPath& ClassPath) { return RedirectClassPath(ClassPath); });
+	Algo::Transform(InClassPaths, ClassPaths, [this](const FTopLevelAssetPath& ClassPath) { return RedirectClassPathInternal(ClassPath); });
 	ClassPaths.SetNum(Algo::RemoveIf(ClassPaths, [this](const FTopLevelAssetPath& ClassPath)
 	{
 		const FString ClasPathString = ClassPath.ToString();
@@ -329,11 +334,11 @@ void FWorldPartitionClassDescRegistry::PrefetchClassDescs(const TArray<FTopLevel
 	{
 		TArray<FTopLevelAssetPath> ParentClassPaths;
 		ParentClassPathsMap.GenerateValueArray(ParentClassPaths);
-		PrefetchClassDescs(ParentClassPaths);
+		PrefetchClassDescsInternal(ParentClassPaths);
 
 		for (auto& [ClassPath, ParentClassPath] : ParentClassPathsMap)
 		{
-			ParentClassMap.Add(ClassPath, RedirectClassPath(ParentClassPath));
+			ParentClassMap.Add(ClassPath, RedirectClassPathInternal(ParentClassPath));
 		}
 	}
 
@@ -390,12 +395,15 @@ void FWorldPartitionClassDescRegistry::PrefetchClassDescs(const TArray<FTopLevel
 bool FWorldPartitionClassDescRegistry::IsRegisteredClass(const FTopLevelAssetPath& InClassPath) const
 {
 	check(IsInitialized());
-	const FTopLevelAssetPath ClassPath = RedirectClassPath(InClassPath);
+	UE::TScopeLock Lock(ClassDescLock);
+	const FTopLevelAssetPath ClassPath = RedirectClassPathInternal(InClassPath);
 	return ClassByPath.Contains(ClassPath);
 }
 
 bool FWorldPartitionClassDescRegistry::IsDerivedFrom(const FWorldPartitionActorDesc* InClassDesc, const FWorldPartitionActorDesc* InParentClassDesc) const
 {
+	UE::TScopeLock Lock(ClassDescLock);
+
 	const FTopLevelAssetPath ParentClassPath = FTopLevelAssetPath(InParentClassDesc->GetActorSoftPath().ToString());
 
 	FTopLevelAssetPath ClassPath = FTopLevelAssetPath(InClassDesc->GetActorSoftPath().ToString());
@@ -413,10 +421,8 @@ bool FWorldPartitionClassDescRegistry::IsDerivedFrom(const FWorldPartitionActorD
 	return false;
 }
 
-const FWorldPartitionActorDesc* FWorldPartitionClassDescRegistry::GetClassDescDefault(const FTopLevelAssetPath& InClassPath) const
+const FWorldPartitionActorDesc* FWorldPartitionClassDescRegistry::GetClassDescDefaultInternal(const FTopLevelAssetPath& InClassPath) const
 {
-	check(IsInitialized());
-
 	const FWorldPartitionActorDesc* Result = nullptr;
 	FTopLevelAssetPath ClassPath = InClassPath;
 
@@ -435,15 +441,26 @@ const FWorldPartitionActorDesc* FWorldPartitionClassDescRegistry::GetClassDescDe
 	return Result;
 }
 
+const FWorldPartitionActorDesc* FWorldPartitionClassDescRegistry::GetClassDescDefault(const FTopLevelAssetPath& InClassPath) const
+{
+	check(IsInitialized());
+	UE::TScopeLock Lock(ClassDescLock);
+	return GetClassDescDefaultInternal(InClassPath);
+}
+
 const FWorldPartitionActorDesc* FWorldPartitionClassDescRegistry::GetClassDescDefaultForActor(const FTopLevelAssetPath& InClassPath) const
 {
-	return GetClassDescDefault(RedirectClassPath(InClassPath));
+	check(IsInitialized());
+	UE::TScopeLock Lock(ClassDescLock);
+	return GetClassDescDefaultInternal(RedirectClassPathInternal(InClassPath));
 }
 
 const FWorldPartitionActorDesc* FWorldPartitionClassDescRegistry::GetClassDescDefaultForClass(const FTopLevelAssetPath& InClassPath) const
 {
-	const FTopLevelAssetPath ParentClassPath = ParentClassMap.FindRef(RedirectClassPath(InClassPath));
-	return GetClassDescDefault(ParentClassPath);
+	check(IsInitialized());
+	UE::TScopeLock Lock(ClassDescLock);
+	const FTopLevelAssetPath ParentClassPath = ParentClassMap.FindRef(RedirectClassPathInternal(InClassPath));
+	return GetClassDescDefaultInternal(ParentClassPath);
 }
 
 void FWorldPartitionClassDescRegistry::PrefetchClassDesc(UClass* InClass)
@@ -458,7 +475,7 @@ void FWorldPartitionClassDescRegistry::PrefetchClassDesc(UClass* InClass)
 			const UPackage* ParentClassPackage = ParentClass->GetPackage();
 			if ((ParentClassPackage != GetTransientPackage()) && !ParentClassPackage->HasAnyPackageFlags(PKG_NewlyCreated))
 			{
-				PrefetchClassDescs({ ParentClassPath });
+				PrefetchClassDescsInternal({ ParentClassPath });
 			}
 
 			if (ClassByPath.Contains(ParentClassPath))
@@ -520,49 +537,71 @@ void FWorldPartitionClassDescRegistry::RegisterClassDescriptorFromActorClass(con
 
 void FWorldPartitionClassDescRegistry::OnObjectPreSave(UObject* InObject, FObjectPreSaveContext InSaveContext)
 {
-	// Are we are saving a blueprint?
-	if (UBlueprint* Blueprint = Cast<UBlueprint>(InObject))
+	const FWorldPartitionActorDesc* UpdatedClassDesc = nullptr;
+
 	{
-		// Is this blueprint an asset (could be a level script)?
-		if (Cast<UPackage>(Blueprint->GetOuter()))
+		UE::TScopeLock Lock(ClassDescLock);
+
+		// Are we are saving a blueprint?
+		if (UBlueprint* Blueprint = Cast<UBlueprint>(InObject))
 		{
-			// Is this blueprint generated class valid?
-			if (UBlueprintGeneratedClass* BlueprintGeneratedClass = Cast<UBlueprintGeneratedClass>(Blueprint->GeneratedClass))
+			// Is this blueprint an asset (could be a level script)?
+			if (Cast<UPackage>(Blueprint->GetOuter()))
 			{
-				// Is it an actor derived blueprint?
-				if (BlueprintGeneratedClass->IsChildOf<AActor>())
+				// Is this blueprint generated class valid?
+				if (UBlueprintGeneratedClass* BlueprintGeneratedClass = Cast<UBlueprintGeneratedClass>(Blueprint->GeneratedClass))
 				{
-					PrefetchClassDesc(BlueprintGeneratedClass);
-					UpdateClassDescriptor(Blueprint, false);
+					// Is it an actor derived blueprint?
+					if (BlueprintGeneratedClass->IsChildOf<AActor>())
+					{
+						PrefetchClassDesc(BlueprintGeneratedClass);
+						UpdatedClassDesc = UpdateClassDescriptor(Blueprint, false);
+					}
 				}
 			}
 		}
+
+		if (InObject->IsA<AActor>())
+		{
+			if (UBlueprintGeneratedClass* BlueprintGeneratedClass = Cast<UBlueprintGeneratedClass>(InObject->GetClass()))
+			{
+				PrefetchClassDesc(BlueprintGeneratedClass);
+			}
+		}
 	}
 
-	if (InObject->IsA<AActor>())
+	if (UpdatedClassDesc)
 	{
-		if (UBlueprintGeneratedClass* BlueprintGeneratedClass = Cast<UBlueprintGeneratedClass>(InObject->GetClass()))
-		{			
-			PrefetchClassDesc(BlueprintGeneratedClass);
-		}
+		ClassDescriptorUpdatedEvent.Broadcast(UpdatedClassDesc);
 	}
 }
 
 void FWorldPartitionClassDescRegistry::OnObjectPropertyChanged(UObject* InObject, FPropertyChangedEvent& InPropertyChangedEvent)
 {
-	// We only want to handle BP class compiles, not property changes.
-	if (!InPropertyChangedEvent.Property)
+	const FWorldPartitionActorDesc* UpdatedClassDesc = nullptr;
+
 	{
-		if (UBlueprint* Blueprint = Cast<UBlueprint>(InObject))
+		UE::TScopeLock Lock(ClassDescLock);
+
+		// We only want to handle BP class compiles, not property changes.
+		if (!InPropertyChangedEvent.Property)
 		{
-			// The generated class is invalid in some situations, like renaming a blueprint, etc.
-			if (Blueprint->IsUpToDate() && Blueprint->GeneratedClass && Blueprint->GeneratedClass->IsChildOf<AActor>())
+			if (UBlueprint* Blueprint = Cast<UBlueprint>(InObject))
 			{
-				PrefetchClassDesc(Blueprint->GeneratedClass);
-				UpdateClassDescriptor(InObject, true);
+				// The generated class is invalid in some situations, like renaming a blueprint, etc.
+				if (Blueprint->IsUpToDate() && Blueprint->GeneratedClass && Blueprint->GeneratedClass->IsChildOf<AActor>())
+				{
+					PrefetchClassDesc(Blueprint->GeneratedClass);
+					UpdatedClassDesc = UpdateClassDescriptor(InObject, true);
+				}
 			}
 		}
 	}
+
+	if (UpdatedClassDesc)
+	{
+		ClassDescriptorUpdatedEvent.Broadcast(UpdatedClassDesc);
+	}
 }
 
 void FWorldPartitionClassDescRegistry::OnPluginLoadingPhaseComplete(ELoadingPhase::Type InLoadingPhase, bool bInPhaseSuccessful)
@@ -577,6 +616,7 @@ void FWorldPartitionClassDescRegistry::OnPluginLoadingPhaseComplete(ELoadingPhas
 void FWorldPartitionClassDescRegistry::OnAssetRemoved(const FAssetData& InAssetData)
 {
 	check(IsInitialized());
+	UE::TScopeLock Lock(ClassDescLock);
 
 	if (InAssetData.GetClass() && InAssetData.GetClass()->IsChildOf(UBlueprint::StaticClass()))
 	{
@@ -602,6 +642,7 @@ void FWorldPartitionClassDescRegistry::OnAssetRemoved(const FAssetData& InAssetD
 void FWorldPartitionClassDescRegistry::OnAssetRenamed(const FAssetData& InAssetData, const FString& InOldObjectPath)
 {
 	check(IsInitialized());
+	UE::TScopeLock Lock(ClassDescLock);
 
 	if (InAssetData.GetClass() && InAssetData.GetClass()->IsChildOf(UBlueprint::StaticClass()))
 	{
@@ -662,6 +703,7 @@ void FWorldPartitionClassDescRegistry::RegisterClass(UClass* Class)
 
 void FWorldPartitionClassDescRegistry::RegisterClasses()
 {
+	UE::TScopeLock Lock(ClassDescLock);
 	ValidateInternalState();
 
 	for (TObjectIterator<UClass> It; It; ++It)
@@ -676,7 +718,7 @@ void FWorldPartitionClassDescRegistry::RegisterClasses()
 	ValidateInternalState();
 }
 
-void FWorldPartitionClassDescRegistry::UpdateClassDescriptor(UObject* InObject, bool bOnlyIfExists)
+const FWorldPartitionActorDesc* FWorldPartitionClassDescRegistry::UpdateClassDescriptor(UObject* InObject, bool bOnlyIfExists)
 {
 	check(IsInitialized());
 
@@ -719,9 +761,12 @@ void FWorldPartitionClassDescRegistry::UpdateClassDescriptor(UObject* InObject,
 				ValidateInternalState();
 			}
 
-			ClassDescriptorUpdatedEvent.Broadcast((*ExistingClassDesc)->Get());
+			// Return the updated descriptor so the caller can broadcast after releasing the lock
+			return (*ExistingClassDesc)->Get();
 		}
 	}
+
+	return nullptr;
 }
 
 void FWorldPartitionClassDescRegistry::ValidateInternalState()
@@ -738,7 +783,7 @@ void FWorldPartitionClassDescRegistry::ValidateInternalState()
 #endif
 }
 
-FTopLevelAssetPath FWorldPartitionClassDescRegistry::RedirectClassPath(const FTopLevelAssetPath& InClassPath) const
+FTopLevelAssetPath FWorldPartitionClassDescRegistry::RedirectClassPathInternal(const FTopLevelAssetPath& InClassPath) const
 {
 	// Resolve native class redirectors
 	const FCoreRedirectObjectName OldClassName(InClassPath.GetAssetName(), NAME_None, InClassPath.GetPackageName());
@@ -753,4 +798,10 @@ FTopLevelAssetPath FWorldPartitionClassDescRegistry::RedirectClassPath(const FTo
 
 	return ResolvedClassPath;
 }
+
+FTopLevelAssetPath FWorldPartitionClassDescRegistry::RedirectClassPath(const FTopLevelAssetPath& InClassPath) const
+{
+	UE::TScopeLock Lock(ClassDescLock);
+	return RedirectClassPathInternal(InClassPath);
+}
 #endif
diff --git a/Engine/Source/Runtime/Engine/Public/WorldPartition/WorldPartitionClassDescRegistry.h b/Engine/Source/Runtime/Engine/Public/WorldPartition/WorldPartitionClassDescRegistry.h
--- a/Engine/Source/Runtime/Engine/Public/WorldPartition/WorldPartitionClassDescRegistry.h
+++ b/Engine/Source/Runtime/Engine/Public/WorldPartition/WorldPartitionClassDescRegistry.h
@@ -4,6 +4,7 @@
 #if WITH_EDITOR
 #include "CoreMinimal.h"
 #include "Misc/LazySingleton.h"
+#include "Misc/TransactionallySafeCriticalSection.h"
 #include "UObject/ObjectMacros.h"
 #include "WorldPartition/ActorDescList.h"
 #include "ModuleDescriptor.h"
@@ -66,12 +67,21 @@ private:
 	ENGINE_API void RegisterClass(UClass* Class);
 	ENGINE_API void RegisterClasses();
 
-	ENGINE_API void UpdateClassDescriptor(UObject* InObject, bool bOnlyIfExists);
+	// Returns the updated descriptor if one was modified (for deferred broadcast outside the lock), or nullptr.
+	ENGINE_API const FWorldPartitionActorDesc* UpdateClassDescriptor(UObject* InObject, bool bOnlyIfExists);
 
 	ENGINE_API void ValidateInternalState();
 
 	ENGINE_API FTopLevelAssetPath RedirectClassPath(const FTopLevelAssetPath& InClassPath) const;
 
+	// Internal variants that assume ClassDescLock is already held by the caller
+	const FWorldPartitionActorDesc* GetClassDescDefaultInternal(const FTopLevelAssetPath& InClassPath) const;
+	FTopLevelAssetPath RedirectClassPathInternal(const FTopLevelAssetPath& InClassPath) const;
+	void PrefetchClassDescsInternal(const TArray<FTopLevelAssetPath>& InClassPaths);
+
+	// Protects ClassByPath, ParentClassMap, and RedirectClassMap from concurrent access
+	mutable FTransactionallySafeCriticalSection ClassDescLock;
+
 	FNameClassDescMap ClassByPath;
 	FParentClassMap ParentClassMap;
 	FRedirectClassMap RedirectClassMap;
