Restored the behavior of the ObjectWorldPosition material node for non-instanced meshes to be consistent with the behavior before the ISM VF was combined into Local VF.

#rnx

diff --git a/Engine/Shaders/Private/MaterialTemplate.ush b/Engine/Shaders/Private/MaterialTemplate.ush
--- a/Engine/Shaders/Private/MaterialTemplate.ush
+++ b/Engine/Shaders/Private/MaterialTemplate.ush
@@ -1858,11 +1858,24 @@ FMaterialLWCData MakeMaterialLWCData(FMaterialVertexParameters Parameters)
 #else //DECAL_PRIMITIVE
 	
 #if USE_INSTANCING || USE_INSTANCE_CULLING || IS_MESHPARTICLE_FACTORY
+	// TODO: Rename PRIMITIVE_SCENE_DATA_FLAG_UNUSED_2
+	bool bEnableObjectWorldPositionFix = (GetPrimitiveData(Parameters).Flags & PRIMITIVE_SCENE_DATA_FLAG_UNUSED_2) != 0;
+	if (bEnableObjectWorldPositionFix)
+	{
+#if PRIMITIVE_HAS_TILEOFFSET_DATA && WSVECTOR_IS_TILEOFFSET
+		Result.ObjectWorldPosition = GetPrimitiveData(Parameters).ObjectWorldPositionTO;
+#else
+		Result.ObjectWorldPosition = DFToWS(GetPrimitiveData(Parameters).ObjectWorldPosition);
+#endif
+	}
+	else
+	{
 #if MATERIAL_LWC_ENABLED
-	Result.ObjectWorldPosition = WSGetOrigin(DFToWS(Parameters.InstanceLocalToWorld));
+		Result.ObjectWorldPosition = WSGetOrigin(DFToWS(Parameters.InstanceLocalToWorld));
 #else //MATERIAL_LWC_ENABLED
-	Result.ObjectWorldPosition = WSPromote(Parameters.InstanceLocalToWorld.PostTranslation + Parameters.InstanceLocalToWorld.M[3].xyz);
+		Result.ObjectWorldPosition = WSPromote(Parameters.InstanceLocalToWorld.PostTranslation + Parameters.InstanceLocalToWorld.M[3].xyz);
 #endif //MATERIAL_LWC_ENABLED
+	}
 #elif PRIMITIVE_HAS_TILEOFFSET_DATA && WSVECTOR_IS_TILEOFFSET
 	Result.ObjectWorldPosition = GetPrimitiveData(Parameters).ObjectWorldPositionTO;
 #else
diff --git a/Engine/Source/Runtime/Engine/Private/PrimitiveSceneProxy.cpp b/Engine/Source/Runtime/Engine/Private/PrimitiveSceneProxy.cpp
--- a/Engine/Source/Runtime/Engine/Private/PrimitiveSceneProxy.cpp
+++ b/Engine/Source/Runtime/Engine/Private/PrimitiveSceneProxy.cpp
@@ -117,6 +117,18 @@ static FAutoConsoleVariableRef CVarParallelGatherDynamicRayTracingInstances(
 	ECVF_RenderThreadSafe
 );
 
+bool GEnableObjectWorldPositionFix = true;
+static FAutoConsoleVariableRef CVarEnableObjectWorldPositionFix(
+	TEXT("r.EnableObjectWorldPositionFix"),
+	GEnableObjectWorldPositionFix,
+	TEXT("If enabled (default) any primitives where HasInstanceDataBuffers() is false will use the previous behavior for calculating ObjectWorldPosition."),
+	FConsoleVariableDelegate::CreateLambda([](IConsoleVariable* InVariable)
+	{
+		FGlobalComponentRecreateRenderStateContext Context;
+	}),
+	ECVF_RenderThreadSafe
+);
+
 bool IsParallelGatherDynamicRayTracingInstancesEnabled()
 {
 	return GParallelGatherDynamicRayTracingInstances &&
@@ -798,6 +810,7 @@ void FPrimitiveSceneProxy::BuildUniformShaderParametersInternal(FPrimitiveUnifor
 			.MeshPaintTextureDescriptor(GetMeshPaintTextureDescriptor())
 			.MaterialCacheDescriptor(GetMaterialCacheDescriptor())
 			.IsFirstPerson(IsFirstPerson())
+			.EnableObjectWorldPositionFix(!HasInstanceDataBuffers() && GEnableObjectWorldPositionFix)
 			.AnimationMinScreenSize(GetAnimationMinScreenSize());
 
 	if (IsNaniteMesh())
diff --git a/Engine/Source/Runtime/Engine/Public/PrimitiveUniformShaderParametersBuilder.h b/Engine/Source/Runtime/Engine/Public/PrimitiveUniformShaderParametersBuilder.h
--- a/Engine/Source/Runtime/Engine/Public/PrimitiveUniformShaderParametersBuilder.h
+++ b/Engine/Source/Runtime/Engine/Public/PrimitiveUniformShaderParametersBuilder.h
@@ -62,6 +62,7 @@ public:
 		bRayTracingHasGroupId						= false;
 		bHasPerClusterDisplacementFallbackRaster	= false;
 		bIsFirstPerson								= false;
+		bEnableObjectWorldPositionFix				= false;
 
 		Parameters.MaxWPOExtent						= 0.0f;
 		Parameters.MinMaterialDisplacement			= 0.0f;
@@ -138,6 +139,7 @@ public:
 	PRIMITIVE_UNIFORM_BUILDER_FLAG_METHOD(bool,			RayTracingHasGroupId);
 	PRIMITIVE_UNIFORM_BUILDER_FLAG_METHOD(bool,			HasPerClusterDisplacementFallbackRaster);
 	PRIMITIVE_UNIFORM_BUILDER_FLAG_METHOD(bool,			IsFirstPerson);
+	PRIMITIVE_UNIFORM_BUILDER_FLAG_METHOD(bool,			EnableObjectWorldPositionFix);
 
 	PRIMITIVE_UNIFORM_BUILDER_METHOD(uint32,			InstanceSceneDataOffset);
 	PRIMITIVE_UNIFORM_BUILDER_METHOD(uint32,			NumInstanceSceneDataEntries);
@@ -438,6 +440,8 @@ public:
 		Parameters.Flags |= bRayTracingHasGroupId ? PRIMITIVE_SCENE_DATA_FLAG_RAYTRACING_HAS_GROUPID : 0u;
 		Parameters.Flags |= bHasPerClusterDisplacementFallbackRaster ? PRIMITIVE_SCENE_DATA_FLAG_PER_CLUSTER_DISPLACEMENT_FALLBACK_RASTER : 0u;
 		Parameters.Flags |= bIsFirstPerson ? PRIMITIVE_SCENE_DATA_FLAG_IS_FIRST_PERSON : 0u;
+		// TODO: Rename the flag later - not done now to avoid affecting more shaders than needed
+		Parameters.Flags |= bEnableObjectWorldPositionFix ? PRIMITIVE_SCENE_DATA_FLAG_UNUSED_2 : 0u;
 		
 		Parameters.VisibilityFlags = 0;
 		Parameters.VisibilityFlags |= bCastHiddenShadow ? PRIMITIVE_VISIBILITY_FLAG_CAST_HIDDEN_SHADOW : 0u;
@@ -507,4 +511,5 @@ private:
 	uint32 bRayTracingHasGroupId : 1;
 	uint32 bHasPerClusterDisplacementFallbackRaster : 1;
 	uint32 bIsFirstPerson : 1;
+	uint32 bEnableObjectWorldPositionFix : 1;
 };
