Handle null Component in AppendCollisionDataFromEnvironment for ISM overlaps

When FOverlapResult::Component is null (e.g. FastGeo ISM without surrogate) or its BodyInstance is invalid (ISM with surrogate), recover the FBodyInstance directly from the physics proxy via PhysicsObject.
This allows cloth simulation to collide with ISM geometry.

Also remove colliding with other cloth colliders from environmental collisions since this feature doesn't work correctly for Chaos Cloth (and will be entirely deprecated in a separate CL).

#Synced-CL 51899829
#Synced-CL 51899829

diff --git a/Engine/Source/Runtime/Engine/Private/PhysicsEngine/EnvironmentalCollisions.cpp b/Engine/Source/Runtime/Engine/Private/PhysicsEngine/EnvironmentalCollisions.cpp
--- a/Engine/Source/Runtime/Engine/Private/PhysicsEngine/EnvironmentalCollisions.cpp
+++ b/Engine/Source/Runtime/Engine/Private/PhysicsEngine/EnvironmentalCollisions.cpp
@@ -6,6 +6,7 @@
 #include "Chaos/ImplicitFwd.h"
 #include "Chaos/ImplicitObjectScaled.h"
 #include "Chaos/PhysicsObjectInterface.h"
+#include "Physics/Experimental/PhysScene_Chaos.h"
 #include "Components/PrimitiveComponent.h"
 #include "Components/SceneComponent.h"
 #include "Components/SkeletalMeshComponent.h"
@@ -42,8 +43,6 @@ void FEnvironmentalCollisions::AppendCollisionDataFromEnvironment(const USceneCo
 	FCollisionObjectQueryParams ObjectParams;
 
 	ObjectParams.AddObjectTypesToQuery(ECollisionChannel::ECC_WorldStatic);
-	// to collide with other clothing objects
-	ObjectParams.AddObjectTypesToQuery(ECollisionChannel::ECC_PhysicsBody);
 
 	FCollisionQueryParams Params(SCENE_QUERY_STAT(ClothOverlapComponents), false);
 	const FBoxSphereBounds PaddedBounds = SceneComponent->Bounds.ExpandBy(EnvironmentCollisionConfig.Padding);
@@ -53,324 +52,318 @@ void FEnvironmentalCollisions::AppendCollisionDataFromEnvironment(const USceneCo
 	for (int32 OverlapIdx = 0; OverlapIdx < Overlaps.Num(); ++OverlapIdx)
 	{
 		const TWeakObjectPtr<UPrimitiveComponent>& Component = Overlaps[OverlapIdx].Component;
-		if (Component.IsValid())
-		{
-			ECollisionChannel Channel = Component->GetCollisionObjectType();
+		const Chaos::FPhysicsObjectHandle PhysicsObject = Overlaps[OverlapIdx].PhysicsObject;
+
+		const UPrimitiveComponent* const ResolvedComponent = Component.Get();
 
-			if (Channel == ECollisionChannel::ECC_WorldStatic)
+		// Resolve the FBodyInstance for this overlap.
+		// When Component is valid with a valid body instance (the common case), the existing component-based path is used.
+		// When Component is null (e.g. FastGeo ISM without surrogate) or the Component's body instance is invalid (ISM with surrogate), fall back to the
+		// FBodyInstance obtained directly from the physics proxy.
+		const bool bUseResolvedComponent = ResolvedComponent && ResolvedComponent->BodyInstance.IsValidBodyInstance();
+
+		FBodyInstance* FallbackBodyInstance = nullptr;
+		if (!bUseResolvedComponent && PhysicsObject)
+		{
+			const Chaos::FConstPhysicsObjectHandle ConstHandle = PhysicsObject;
+			IPhysicsProxyBase* Proxy = Chaos::FPhysicsObjectInterface::GetProxy({ &ConstHandle, 1 });
+			if (Proxy)
 			{
-				// Static world geo
-				if (!Component->BodyInstance.IsValidBodyInstance())
+				FallbackBodyInstance = SceneComponent->GetWorld()->GetPhysicsScene()->GetBodyInstanceFromProxy(Proxy);
+			}
+		}
+
+		const FBodyInstance* const BodyInstToUse = bUseResolvedComponent ? &ResolvedComponent->BodyInstance : FallbackBodyInstance;
+
+		if (!BodyInstToUse || !BodyInstToUse->IsValidBodyInstance())
+		{
+			continue;
+		}
+
+		const ECollisionChannel Channel = bUseResolvedComponent ? ResolvedComponent->GetCollisionObjectType() : FallbackBodyInstance->GetObjectType();
+
+		if (Channel == ECollisionChannel::ECC_WorldStatic)
+		{
+			// Static world geo
+			bool bSuccessfulRead = false;
+
+			FPhysicsActorHandle ActorRef = BodyInstToUse->GetPhysicsActorWeldRoot();
+
+			const FString OwnerDebugName = ResolvedComponent && ResolvedComponent->GetOwner()
+				? ResolvedComponent->GetOwner()->GetFName().ToString()
+				: TEXT("Unknown (ISM instance)");
+
+			FTransform ClothComponentTransform = SceneComponent->GetComponentTransform();
+			ClothComponentTransform.RemoveScaling();  // The environment collision shape doesn't need the scale of the cloth skeletal mesh applied to it (but it does need the source scale from its component transform)
+
+			FPhysicsCommand::ExecuteRead(ActorRef, [&](const FPhysicsActorHandle& Actor)
 				{
-					// Move to next component, this one has no valid physics
-					continue;
-				}
+					// Get the transform directly from the PhysicsActor when not using the resolved component
+					const FTransform ComponentToClothBaseTransform = bUseResolvedComponent ? ResolvedComponent->GetComponentTransform() * ClothComponentTransform.Inverse() : FPhysicsInterface::GetTransform_AssumesLocked(Actor) * ClothComponentTransform.Inverse();
 
-				bool bSuccessfulRead = false;
+					TArray<FPhysicsShapeHandle> AllShapes;
+					const int32 NumSyncShapes = BodyInstToUse->GetAllShapes_AssumesLocked(AllShapes);
 
-				FPhysicsActorHandle ActorRef = Component->BodyInstance.GetPhysicsActorWeldRoot();
-				FPhysicsCommand::ExecuteRead(ActorRef, [&](const FPhysicsActorHandle& Actor)
+					if (NumSyncShapes == 0 || NumSyncShapes > EnvironmentCollisionConfig.MaxShapes)
 					{
-						TArray<FPhysicsShapeHandle> AllShapes;
-						const int32 NumSyncShapes = Component->BodyInstance.GetAllShapes_AssumesLocked(AllShapes);
+						// Either no shapes or too complicated to consider
+						return;
+					}
+
+					using namespace Chaos;
+
+					const FReal Thickness = (FReal)EnvironmentCollisionConfig.Thickness;
 
-						if (NumSyncShapes == 0 || NumSyncShapes > EnvironmentCollisionConfig.MaxShapes)
+					auto AddSphere = [&CollisionData](const FImplicitSphere3& ImplicitSphere, const FTransform& ComponentToClothTransform, const FVec3& Scale = FVec3::OneVector)
 						{
-							// Either no shapes or too complicated to consider
-							return;
-						}
+							CollisionData.Spheres.Emplace(
+								ImplicitSphere.GetRadiusf() * Scale.X,  // Assumes uniform scale
+								ComponentToClothTransform.TransformPosition(FVec3(ImplicitSphere.GetCenterf())));
+						};
 
-						using namespace Chaos;
+					auto AddBox = [&CollisionData, Thickness](const FImplicitBox3& ImplicitBox, const FTransform& ComponentToClothTransform, const FVec3& Scale = FVec3::OneVector)
+						{
+							CollisionData.Boxes.Emplace(
+								ComponentToClothTransform.TransformPosition(ImplicitBox.GetCenter()),
+								ComponentToClothTransform.GetRotation(),
+								ImplicitBox.Extents() * Scale * (FReal)0.5 + Thickness);
+						};
 
-						const FReal Thickness = (FReal)EnvironmentCollisionConfig.Thickness;
+					auto AddCapsule = [&CollisionData](const FImplicitCapsule3& ImplicitCapsule, const FTransform& ComponentToClothTransform, const FVec3& Scale = FVec3::OneVector)
+						{
+							const int32 BaseSphereIndex = CollisionData.Spheres.Num();
 
-						auto AddSphere = [&CollisionData](const FImplicitSphere3& ImplicitSphere, const FTransform& ComponentToClothTransform, const FVec3& Scale = FVec3::OneVector)
-							{
-								CollisionData.Spheres.Emplace(
-									ImplicitSphere.GetRadiusf() * Scale.X,  // Assumes uniform scale
-									ComponentToClothTransform.TransformPosition(FVec3(ImplicitSphere.GetCenterf())));
-							};
+							const FReal Radius = ImplicitCapsule.GetRadiusf() * Scale.X;  // Assumes uniform scale
+							const FReal HalfHeight = ImplicitCapsule.GetHeightf() * (FReal)0.5;
+							const FVector TransformedAxis = ComponentToClothTransform.TransformVectorNoScale(FVec3(ImplicitCapsule.GetAxisf()));
+							const FVector HalfSegment = TransformedAxis * HalfHeight * Scale.X;  // Assumes uniform scale
+							const FVector TransformedCenter = ComponentToClothTransform.TransformPosition(FVec3(ImplicitCapsule.GetCenterf()));
 
-						auto AddBox = [&CollisionData, Thickness](const FImplicitBox3& ImplicitBox, const FTransform& ComponentToClothTransform, const FVec3& Scale = FVec3::OneVector)
-							{
-								CollisionData.Boxes.Emplace(
-									ComponentToClothTransform.TransformPosition(ImplicitBox.GetCenter()),
-									ComponentToClothTransform.GetRotation(),
-									ImplicitBox.Extents() * Scale * (FReal)0.5 + Thickness);
-							};
+							CollisionData.Spheres.Emplace(Radius, TransformedCenter + HalfSegment);
+							CollisionData.Spheres.Emplace(Radius, TransformedCenter - HalfSegment);
 
-						auto AddCapsule = [&CollisionData](const FImplicitCapsule3& ImplicitCapsule, const FTransform& ComponentToClothTransform, const FVec3& Scale = FVec3::OneVector)
-							{
-								const int32 BaseSphereIndex = CollisionData.Spheres.Num();
+							CollisionData.SphereConnections.Emplace(BaseSphereIndex, BaseSphereIndex + 1);
+						};
 
-								const FReal Radius = ImplicitCapsule.GetRadiusf() * Scale.X;  // Assumes uniform scale
-								const FReal HalfHeight = ImplicitCapsule.GetHeightf() * (FReal)0.5;
-								const FVector TransformedAxis = ComponentToClothTransform.TransformVectorNoScale(FVec3(ImplicitCapsule.GetAxisf()));
-								const FVector HalfSegment = TransformedAxis * HalfHeight * Scale.X;  // Assumes uniform scale
-								const FVector TransformedCenter = ComponentToClothTransform.TransformPosition(FVec3(ImplicitCapsule.GetCenterf()));
+					auto AddConvex = [&CollisionData, Thickness](const FImplicitConvex3& ImplicitConvex, const FMatrix& ComponentToClothMatrix, const FVec3* const Scale = nullptr)
+						{
+							TArray<FClothCollisionPrim_ConvexFace> Faces;
+							const int32 NumPlanes = ImplicitConvex.NumPlanes();
+							Faces.SetNum(NumPlanes);
 
-								CollisionData.Spheres.Emplace(Radius, TransformedCenter + HalfSegment);
-								CollisionData.Spheres.Emplace(Radius, TransformedCenter - HalfSegment);
+							TArray<FVector> SurfacePoints;
+							const int32 NumSurfacePoints = ImplicitConvex.NumVertices();
+							SurfacePoints.SetNumUninitialized(NumSurfacePoints);
 
-								CollisionData.SphereConnections.Emplace(BaseSphereIndex, BaseSphereIndex + 1);
-							};
+							TArray<TArray<int32, TInlineAllocator<4>>, TInlineAllocator<16>> PointFaces;
+							PointFaces.SetNum(NumSurfacePoints);
 
-						auto AddConvex = [&CollisionData, Thickness](const FImplicitConvex3& ImplicitConvex, const FMatrix& ComponentToClothMatrix, const FVec3* const Scale = nullptr)
+							for (int32 FaceIndex = 0; FaceIndex < NumPlanes; ++FaceIndex)
 							{
-								TArray<FClothCollisionPrim_ConvexFace> Faces;
-								const int32 NumPlanes = ImplicitConvex.NumPlanes();
-								Faces.SetNum(NumPlanes);
+								const TPlaneConcrete<FReal, 3>& Plane = ImplicitConvex.GetPlane(FaceIndex);
+								const FVec3& Normal = Plane.Normal();
+								const FVec3 Base = Plane.X() + Normal * (FReal)Thickness;
 
-								TArray<FVector> SurfacePoints;
-								const int32 NumSurfacePoints = ImplicitConvex.NumVertices();
-								SurfacePoints.SetNumUninitialized(NumSurfacePoints);
+								Faces[FaceIndex].Plane = FPlane(FVector(Base), FVector(Normal)).TransformBy(ComponentToClothMatrix);
 
-								TArray<TArray<int32, TInlineAllocator<4>>, TInlineAllocator<16>> PointFaces;
-								PointFaces.SetNum(NumSurfacePoints);
+								const int32 NumFaceIndices = ImplicitConvex.NumPlaneVertices(FaceIndex);
+								Faces[FaceIndex].Indices.SetNumUninitialized(NumFaceIndices);
 
-								for (int32 FaceIndex = 0; FaceIndex < NumPlanes; ++FaceIndex)
+								for (int32 Index = 0; Index < NumFaceIndices; ++Index)
 								{
-									const TPlaneConcrete<FReal, 3>& Plane = ImplicitConvex.GetPlane(FaceIndex);
-									const FVec3& Normal = Plane.Normal();
-									const FVec3 Base = Plane.X() + Normal * (FReal)Thickness;
-
-									Faces[FaceIndex].Plane = FPlane(FVector(Base), FVector(Normal)).TransformBy(ComponentToClothMatrix);
+									const int32 PointIndex = ImplicitConvex.GetPlaneVertex(FaceIndex, Index);
+									Faces[FaceIndex].Indices[Index] = PointIndex;
 
-									const int32 NumFaceIndices = ImplicitConvex.NumPlaneVertices(FaceIndex);
-									Faces[FaceIndex].Indices.SetNumUninitialized(NumFaceIndices);
-
-									for (int32 Index = 0; Index < NumFaceIndices; ++Index)
+									if (ensure(PointFaces.IsValidIndex(PointIndex)))
 									{
-										const int32 PointIndex = ImplicitConvex.GetPlaneVertex(FaceIndex, Index);
-										Faces[FaceIndex].Indices[Index] = PointIndex;
-
-										if (ensure(PointFaces.IsValidIndex(PointIndex)))
-										{
-											PointFaces[PointIndex].Add(FaceIndex);
-										}
+										PointFaces[PointIndex].Add(FaceIndex);
 									}
 								}
+							}
 
-								for (int32 PointIndex = 0; PointIndex < NumSurfacePoints; ++PointIndex)
+							for (int32 PointIndex = 0; PointIndex < NumSurfacePoints; ++PointIndex)
+							{
+								if (ensure(PointFaces[PointIndex].Num() >= 3))
 								{
-									if (ensure(PointFaces[PointIndex].Num() >= 3))
-									{
-										const int32 Index0 = PointFaces[PointIndex][0];
-										const int32 Index1 = PointFaces[PointIndex][1];
-										const int32 Index2 = PointFaces[PointIndex][2];
-
-										if (!FMath::IntersectPlanes3(SurfacePoints[PointIndex], Faces[Index0].Plane, Faces[Index1].Plane, Faces[Index2].Plane))
-										{
-											SurfacePoints[PointIndex] = ComponentToClothMatrix.TransformPosition(FVector(ImplicitConvex.GetVertex(PointIndex)));
-										}
-									}
-									else
+									const int32 Index0 = PointFaces[PointIndex][0];
+									const int32 Index1 = PointFaces[PointIndex][1];
+									const int32 Index2 = PointFaces[PointIndex][2];
+
+									if (!FMath::IntersectPlanes3(SurfacePoints[PointIndex], Faces[Index0].Plane, Faces[Index1].Plane, Faces[Index2].Plane))
 									{
 										SurfacePoints[PointIndex] = ComponentToClothMatrix.TransformPosition(FVector(ImplicitConvex.GetVertex(PointIndex)));
 									}
 								}
+								else
+								{
+									SurfacePoints[PointIndex] = ComponentToClothMatrix.TransformPosition(FVector(ImplicitConvex.GetVertex(PointIndex)));
+								}
+							}
 
-								CollisionData.Convexes.Emplace(MoveTemp(Faces), MoveTemp(SurfacePoints));
-							};
-
-						FTransform ClothComponentTransform = SceneComponent->GetComponentTransform();
-						ClothComponentTransform.RemoveScaling();  // The environment collision shape doesn't need the scale of the cloth skeletal mesh applied to it (but it does need the source scale from its component transform)
-						const FTransform ComponentToClothBaseTransform = Component->GetComponentTransform() * ClothComponentTransform.Inverse();
+							CollisionData.Convexes.Emplace(MoveTemp(Faces), MoveTemp(SurfacePoints));
+						};
 
-						bool bHasSimpleCollision = false;
+					bool bHasSimpleCollision = false;
 
-						for (FPhysicsShapeHandle& ShapeHandle : AllShapes)
-						{
-							FTransform ComponentToClothTransform = ComponentToClothBaseTransform;
+					for (FPhysicsShapeHandle& ShapeHandle : AllShapes)
+					{
+						FTransform ComponentToClothTransform = ComponentToClothBaseTransform;
 
-							const FImplicitObject* ImplicitObject = &ShapeHandle.GetGeometry();
-							EImplicitObjectType ImplicitType = ImplicitObject->GetType();
+						const FImplicitObject* ImplicitObject = &ShapeHandle.GetGeometry();
+						EImplicitObjectType ImplicitType = ImplicitObject->GetType();
 
-							// Transformed implicits
-							if (ImplicitType == ImplicitObjectType::Transformed)
-							{
-								const TImplicitObjectTransformed<FReal, 3>& ImplicitTransformed = ImplicitObject->GetObjectChecked<TImplicitObjectTransformed<FReal, 3>>();
-								ImplicitObject = ImplicitTransformed.GetTransformedObject();
-								ImplicitType = ImplicitObject->GetType();
+						// Transformed implicits
+						if (ImplicitType == ImplicitObjectType::Transformed)
+						{
+							const TImplicitObjectTransformed<FReal, 3>& ImplicitTransformed = ImplicitObject->GetObjectChecked<TImplicitObjectTransformed<FReal, 3>>();
+							ImplicitObject = ImplicitTransformed.GetTransformedObject();
+							ImplicitType = ImplicitObject->GetType();
 
-								ComponentToClothTransform = ImplicitTransformed.GetTransform() * ComponentToClothTransform;
-								UE_LOG(LogSkeletalMesh, Verbose, TEXT("Found transformed environmental collision"));
-							}
+							ComponentToClothTransform = ImplicitTransformed.GetTransform() * ComponentToClothTransform;
+							UE_LOG(LogSkeletalMesh, Verbose, TEXT("Found transformed environmental collision"));
+						}
 
-							switch (ImplicitType)
-							{
-								// Base implicits
-							case ImplicitObjectType::Sphere:
-							{
-								const FImplicitSphere3& ImplicitSphere = ImplicitObject->GetObjectChecked<FImplicitSphere3>();;
-								AddSphere(ImplicitSphere, ComponentToClothTransform);
-							}
-							UE_LOG(LogSkeletalMesh, Verbose, TEXT("Found Sphere cloth environmental collision in [%s]"), *Component->GetOwner()->GetFName().ToString());
+						switch (ImplicitType)
+						{
+							// Base implicits
+						case ImplicitObjectType::Sphere:
+						{
+							const FImplicitSphere3& ImplicitSphere = ImplicitObject->GetObjectChecked<FImplicitSphere3>();;
+							AddSphere(ImplicitSphere, ComponentToClothTransform);
+						}
+						UE_LOG(LogSkeletalMesh, Verbose, TEXT("Found Sphere cloth environmental collision in [%s]"), *OwnerDebugName);
+						bHasSimpleCollision = true;
+						break;
+						case ImplicitObjectType::Box:
+						{
+							const FImplicitBox3& ImplicitBox = ImplicitObject->GetObjectChecked<FImplicitBox3>();
+							AddBox(ImplicitBox, ComponentToClothTransform);
+						}
+						UE_LOG(LogSkeletalMesh, Verbose, TEXT("Found Box cloth environmental collision in [%s]"), *OwnerDebugName);
+						bHasSimpleCollision = true;
+						break;
+						case ImplicitObjectType::Capsule:
+						{
+							const FImplicitCapsule3& ImplicitCapsule = ImplicitObject->GetObjectChecked<FImplicitCapsule3>();
+							AddCapsule(ImplicitCapsule, ComponentToClothTransform);
+							UE_LOG(LogSkeletalMesh, Verbose, TEXT("Found Capsule cloth environmental collision in [%s]"), *OwnerDebugName);
 							bHasSimpleCollision = true;
-							break;
-							case ImplicitObjectType::Box:
-							{
-								const FImplicitBox3& ImplicitBox = ImplicitObject->GetObjectChecked<FImplicitBox3>();
-								AddBox(ImplicitBox, ComponentToClothTransform);
-							}
-							UE_LOG(LogSkeletalMesh, Verbose, TEXT("Found Box cloth environmental collision in [%s]"), *Component->GetOwner()->GetFName().ToString());
+						}
+						break;
+						case ImplicitObjectType::Convex:
+						{
+							const FImplicitConvex3& ImplicitConvex = ImplicitObject->GetObjectChecked<FImplicitConvex3>();
+							AddConvex(ImplicitConvex, ComponentToClothTransform.ToMatrixNoScale());
+							UE_LOG(LogSkeletalMesh, Verbose, TEXT("Found Convex cloth environmental collision in [%s]"), *OwnerDebugName);
 							bHasSimpleCollision = true;
-							break;
-							case ImplicitObjectType::Capsule:
-							{
-								const FImplicitCapsule3& ImplicitCapsule = ImplicitObject->GetObjectChecked<FImplicitCapsule3>();
-								AddCapsule(ImplicitCapsule, ComponentToClothTransform);
-								UE_LOG(LogSkeletalMesh, Verbose, TEXT("Found Capsule cloth environmental collision in [%s]"), *Component->GetOwner()->GetFName().ToString());
-								bHasSimpleCollision = true;
-							}
-							break;
-							case ImplicitObjectType::Convex:
-							{
-								const FImplicitConvex3& ImplicitConvex = ImplicitObject->GetObjectChecked<FImplicitConvex3>();
-								AddConvex(ImplicitConvex, ComponentToClothTransform.ToMatrixNoScale());
-								UE_LOG(LogSkeletalMesh, Verbose, TEXT("Found Convex cloth environmental collision in [%s]"), *Component->GetOwner()->GetFName().ToString());
-								bHasSimpleCollision = true;
-							}
-							break;
-
-							// Instanced implicits
-							case ImplicitObjectType::IsInstanced | ImplicitObjectType::Sphere:
-							{
-								const TImplicitObjectInstanced<FImplicitSphere3>& ImplicitInstanced = ImplicitObject->GetObjectChecked<TImplicitObjectInstanced<FImplicitSphere3>>();
-								check(ImplicitInstanced.Object());
-								const FImplicitSphere3& ImplicitSphere = *ImplicitInstanced.GetInstancedObject();
-								AddSphere(ImplicitSphere, ComponentToClothTransform);
-								UE_LOG(LogSkeletalMesh, Verbose, TEXT("Found Instanced Sphere cloth environmental collision in [%s]"), *Component->GetOwner()->GetFName().ToString());
-								bHasSimpleCollision = true;
-							}
-							break;
-							case ImplicitObjectType::IsInstanced | ImplicitObjectType::Box:
-							{
-								const TImplicitObjectInstanced<FImplicitBox3>& ImplicitInstanced = ImplicitObject->GetObjectChecked<TImplicitObjectInstanced<FImplicitBox3>>();
-								check(ImplicitInstanced.Object());
-								const FImplicitBox3& ImplicitBox = *ImplicitInstanced.GetInstancedObject();
-								AddBox(ImplicitBox, ComponentToClothTransform);
-								UE_LOG(LogSkeletalMesh, Verbose, TEXT("Found Instanced Box cloth environmental collision in [%s]"), *Component->GetOwner()->GetFName().ToString());
-								bHasSimpleCollision = true;
-							}
-							break;
-							case ImplicitObjectType::IsInstanced | ImplicitObjectType::Capsule:
-							{
-								const TImplicitObjectInstanced<FImplicitCapsule3>& ImplicitInstanced = ImplicitObject->GetObjectChecked<TImplicitObjectInstanced<FImplicitCapsule3>>();
-								check(ImplicitInstanced.Object());
-								const FImplicitCapsule3& ImplicitCapsule = *ImplicitInstanced.GetInstancedObject();
-								AddCapsule(ImplicitCapsule, ComponentToClothTransform);
-								UE_LOG(LogSkeletalMesh, Verbose, TEXT("Found Instanced Capsule cloth environmental collision in [%s]"), *Component->GetOwner()->GetFName().ToString());
-								bHasSimpleCollision = true;
-							}
-							break;
-							case ImplicitObjectType::IsInstanced | ImplicitObjectType::Convex:
-							{
-								const TImplicitObjectInstanced<FImplicitConvex3>& ImplicitInstanced = ImplicitObject->GetObjectChecked<TImplicitObjectInstanced<FImplicitConvex3>>();
-								check(ImplicitInstanced.Object());
-								const FImplicitConvex3& ImplicitConvex = *ImplicitInstanced.GetInstancedObject();
-								AddConvex(ImplicitConvex, ComponentToClothTransform.ToMatrixNoScale());
-								UE_LOG(LogSkeletalMesh, Verbose, TEXT("Found Instanced Convex cloth environmental collision in [%s]"), *Component->GetOwner()->GetFName().ToString());
-								bHasSimpleCollision = true;
-							}
-							break;
-
-							// Scaled implicits
-							case ImplicitObjectType::IsScaled | ImplicitObjectType::Sphere:
-							{
-								const TImplicitObjectScaled<FImplicitSphere3>& ImplicitScaled = ImplicitObject->GetObjectChecked<TImplicitObjectScaled<FImplicitSphere3>>();
-								check(ImplicitScaled.Object());
-								const FImplicitSphere3& ImplicitSphere = *ImplicitScaled.GetUnscaledObject();
-								ensure(FVector::DistSquared(ComponentToClothTransform.GetScale3D(), FVector(ImplicitScaled.GetScale())) < UE_KINDA_SMALL_NUMBER);
-								AddSphere(ImplicitSphere, ComponentToClothTransform, ImplicitScaled.GetScale());
-								UE_LOG(LogSkeletalMesh, Verbose, TEXT("Found Scaled Sphere cloth environmental collision in [%s]"), *Component->GetOwner()->GetFName().ToString());
-								bHasSimpleCollision = true;
-							}
-							break;
-							case ImplicitObjectType::IsScaled | ImplicitObjectType::Box:
-							{
-								const TImplicitObjectScaled<FImplicitBox3>& ImplicitScaled = ImplicitObject->GetObjectChecked<TImplicitObjectScaled<FImplicitBox3>>();
-								check(ImplicitScaled.Object());
-								const FImplicitBox3& ImplicitBox = *ImplicitScaled.GetUnscaledObject();
-								ensure(FVector::DistSquared(ComponentToClothTransform.GetScale3D(), FVector(ImplicitScaled.GetScale())) < UE_KINDA_SMALL_NUMBER);
-								AddBox(ImplicitBox, ComponentToClothTransform, ImplicitScaled.GetScale());
-								UE_LOG(LogSkeletalMesh, Verbose, TEXT("Found Scaled Box cloth environmental collision in [%s]"), *Component->GetOwner()->GetFName().ToString());
-								bHasSimpleCollision = true;
-							}
-							break;
-							case ImplicitObjectType::IsScaled | ImplicitObjectType::Capsule:
-							{
-								const TImplicitObjectScaled<FImplicitCapsule3>& ImplicitScaled = ImplicitObject->GetObjectChecked<TImplicitObjectScaled<FImplicitCapsule3>>();
-								check(ImplicitScaled.Object());
-								const FImplicitCapsule3& ImplicitCapsule = *ImplicitScaled.GetUnscaledObject();
-								ensure(FVector::DistSquared(ComponentToClothTransform.GetScale3D(), FVector(ImplicitScaled.GetScale())) < UE_KINDA_SMALL_NUMBER);
-								AddCapsule(ImplicitCapsule, ComponentToClothTransform, ImplicitScaled.GetScale());
-								UE_LOG(LogSkeletalMesh, Verbose, TEXT("Found Scaled Capsule cloth environmental collision in [%s]"), *Component->GetOwner()->GetFName().ToString());
-								bHasSimpleCollision = true;
-							}
-							break;
-							case ImplicitObjectType::IsScaled | ImplicitObjectType::Convex:
-							{
-								const TImplicitObjectScaled<FImplicitConvex3>& ImplicitScaled = ImplicitObject->GetObjectChecked<TImplicitObjectScaled<FImplicitConvex3>>();
-								check(ImplicitScaled.Object());
-								const FImplicitConvex3& ImplicitConvex = *ImplicitScaled.GetUnscaledObject();
-								ensure(FVector::DistSquared(ComponentToClothTransform.GetScale3D(), FVector(ImplicitScaled.GetScale())) < UE_KINDA_SMALL_NUMBER);
-								AddConvex(ImplicitConvex, ComponentToClothTransform.ToMatrixWithScale());
-								UE_LOG(LogSkeletalMesh, Verbose, TEXT("Found Scaled Convex cloth environmental collision in [%s]"), *Component->GetOwner()->GetFName().ToString());
-								bHasSimpleCollision = true;
-							}
-							break;
-
-							// Triangle mesh
-							case ImplicitObjectType::TriangleMesh:
-							case ImplicitObjectType::IsInstanced | ImplicitObjectType::TriangleMesh:
-							case ImplicitObjectType::IsScaled | ImplicitObjectType::TriangleMesh:
-								// TODO: We could eventually want to collide cloth against triangle meshes,
-								//       however the concept of simple vs complex shape might need to be clarified
-								//       as it currently iterates over all shape to discard the triangle mesh ones.
-								UE_LOG(LogSkeletalMesh, Verbose, TEXT("Found unusable Triangle Mesh cloth environmental collision in [%s]"), !Component->GetOwner() ? TEXT("Unknown") : *Component->GetOwner()->GetFName().ToString());
-								break;
-
-							default:
-								UE_LOG(LogSkeletalMesh, Verbose, TEXT("Found unsupported collision type during environmental collision with the cloth in [%s]"), !Component->GetOwner() ? TEXT("Unknown") : *Component->GetOwner()->GetFName().ToString());
-								break;
-							}
 						}
-						bSuccessfulRead = true;
-					});
-			}
-			else if (Channel == ECollisionChannel::ECC_PhysicsBody)
-			{
-				// Possibly a skeletal mesh, extract it's clothing collisions if necessary
-				USkeletalMeshComponent* SkelComp = Cast<USkeletalMeshComponent>(Component.Get());
+						break;
 
-				if (SkelComp && SkelComp->GetSkeletalMeshAsset())
-				{
-					if ((UPrimitiveComponent*)SkelComp == (UPrimitiveComponent*)SceneComponent)
-					{
-						// Same mesh, move to next component
-						continue;
-					}
-
-					const TArray<FClothingSimulationInstance>& ClothingSimulationInstances = SkelComp->GetClothingSimulationInstances();
-					if (ClothingSimulationInstances.Num())
-					{
-						// append skeletal component collisions
-						FClothCollisionData SkelCollisionData;
-						constexpr bool bIncludeExternal = false;
-
-						ClothingSimulationInstances[0].GetCollisions(SkelCollisionData, bIncludeExternal);
+						// Instanced implicits
+						case ImplicitObjectType::IsInstanced | ImplicitObjectType::Sphere:
+						{
+							const TImplicitObjectInstanced<FImplicitSphere3>& ImplicitInstanced = ImplicitObject->GetObjectChecked<TImplicitObjectInstanced<FImplicitSphere3>>();
+							check(ImplicitInstanced.Object());
+							const FImplicitSphere3& ImplicitSphere = *ImplicitInstanced.GetInstancedObject();
+							AddSphere(ImplicitSphere, ComponentToClothTransform);
+							UE_LOG(LogSkeletalMesh, Verbose, TEXT("Found Instanced Sphere cloth environmental collision in [%s]"), *OwnerDebugName);
+							bHasSimpleCollision = true;
+						}
+						break;
+						case ImplicitObjectType::IsInstanced | ImplicitObjectType::Box:
+						{
+							const TImplicitObjectInstanced<FImplicitBox3>& ImplicitInstanced = ImplicitObject->GetObjectChecked<TImplicitObjectInstanced<FImplicitBox3>>();
+							check(ImplicitInstanced.Object());
+							const FImplicitBox3& ImplicitBox = *ImplicitInstanced.GetInstancedObject();
+							AddBox(ImplicitBox, ComponentToClothTransform);
+							UE_LOG(LogSkeletalMesh, Verbose, TEXT("Found Instanced Box cloth environmental collision in [%s]"), *OwnerDebugName);
+							bHasSimpleCollision = true;
+						}
+						break;
+						case ImplicitObjectType::IsInstanced | ImplicitObjectType::Capsule:
+						{
+							const TImplicitObjectInstanced<FImplicitCapsule3>& ImplicitInstanced = ImplicitObject->GetObjectChecked<TImplicitObjectInstanced<FImplicitCapsule3>>();
+							check(ImplicitInstanced.Object());
+							const FImplicitCapsule3& ImplicitCapsule = *ImplicitInstanced.GetInstancedObject();
+							AddCapsule(ImplicitCapsule, ComponentToClothTransform);
+							UE_LOG(LogSkeletalMesh, Verbose, TEXT("Found Instanced Capsule cloth environmental collision in [%s]"), *OwnerDebugName);
+							bHasSimpleCollision = true;
+						}
+						break;
+						case ImplicitObjectType::IsInstanced | ImplicitObjectType::Convex:
+						{
+							const TImplicitObjectInstanced<FImplicitConvex3>& ImplicitInstanced = ImplicitObject->GetObjectChecked<TImplicitObjectInstanced<FImplicitConvex3>>();
+							check(ImplicitInstanced.Object());
+							const FImplicitConvex3& ImplicitConvex = *ImplicitInstanced.GetInstancedObject();
+							AddConvex(ImplicitConvex, ComponentToClothTransform.ToMatrixNoScale());
+							UE_LOG(LogSkeletalMesh, Verbose, TEXT("Found Instanced Convex cloth environmental collision in [%s]"), *OwnerDebugName);
+							bHasSimpleCollision = true;
+						}
+						break;
 
-						// AppendUnique is quite expensive, so let's only call it for subsequent clothing simulation instances
-						for (int32 Index = 1 ; Index < ClothingSimulationInstances.Num(); ++Index)
+						// Scaled implicits
+						case ImplicitObjectType::IsScaled | ImplicitObjectType::Sphere:
+						{
+							const TImplicitObjectScaled<FImplicitSphere3>& ImplicitScaled = ImplicitObject->GetObjectChecked<TImplicitObjectScaled<FImplicitSphere3>>();
+							check(ImplicitScaled.Object());
+							const FImplicitSphere3& ImplicitSphere = *ImplicitScaled.GetUnscaledObject();
+							ensure(FVector::DistSquared(ComponentToClothTransform.GetScale3D(), FVector(ImplicitScaled.GetScale())) < UE_KINDA_SMALL_NUMBER);
+							AddSphere(ImplicitSphere, ComponentToClothTransform, ImplicitScaled.GetScale());
+							UE_LOG(LogSkeletalMesh, Verbose, TEXT("Found Scaled Sphere cloth environmental collision in [%s]"), *OwnerDebugName);
+							bHasSimpleCollision = true;
+						}
+						break;
+						case ImplicitObjectType::IsScaled | ImplicitObjectType::Box:
+						{
+							const TImplicitObjectScaled<FImplicitBox3>& ImplicitScaled = ImplicitObject->GetObjectChecked<TImplicitObjectScaled<FImplicitBox3>>();
+							check(ImplicitScaled.Object());
+							const FImplicitBox3& ImplicitBox = *ImplicitScaled.GetUnscaledObject();
+							ensure(FVector::DistSquared(ComponentToClothTransform.GetScale3D(), FVector(ImplicitScaled.GetScale())) < UE_KINDA_SMALL_NUMBER);
+							AddBox(ImplicitBox, ComponentToClothTransform, ImplicitScaled.GetScale());
+							UE_LOG(LogSkeletalMesh, Verbose, TEXT("Found Scaled Box cloth environmental collision in [%s]"), *OwnerDebugName);
+							bHasSimpleCollision = true;
+						}
+						break;
+						case ImplicitObjectType::IsScaled | ImplicitObjectType::Capsule:
 						{
-							ClothingSimulationInstances[Index].AppendUniqueCollisions(SkelCollisionData, bIncludeExternal);
+							const TImplicitObjectScaled<FImplicitCapsule3>& ImplicitScaled = ImplicitObject->GetObjectChecked<TImplicitObjectScaled<FImplicitCapsule3>>();
+							check(ImplicitScaled.Object());
+							const FImplicitCapsule3& ImplicitCapsule = *ImplicitScaled.GetUnscaledObject();
+							ensure(FVector::DistSquared(ComponentToClothTransform.GetScale3D(), FVector(ImplicitScaled.GetScale())) < UE_KINDA_SMALL_NUMBER);
+							AddCapsule(ImplicitCapsule, ComponentToClothTransform, ImplicitScaled.GetScale());
+							UE_LOG(LogSkeletalMesh, Verbose, TEXT("Found Scaled Capsule cloth environmental collision in [%s]"), *OwnerDebugName);
+							bHasSimpleCollision = true;
+						}
+						break;
+						case ImplicitObjectType::IsScaled | ImplicitObjectType::Convex:
+						{
+							const TImplicitObjectScaled<FImplicitConvex3>& ImplicitScaled = ImplicitObject->GetObjectChecked<TImplicitObjectScaled<FImplicitConvex3>>();
+							check(ImplicitScaled.Object());
+							const FImplicitConvex3& ImplicitConvex = *ImplicitScaled.GetUnscaledObject();
+							ensure(FVector::DistSquared(ComponentToClothTransform.GetScale3D(), FVector(ImplicitScaled.GetScale())) < UE_KINDA_SMALL_NUMBER);
+							AddConvex(ImplicitConvex, ComponentToClothTransform.ToMatrixWithScale());
+							UE_LOG(LogSkeletalMesh, Verbose, TEXT("Found Scaled Convex cloth environmental collision in [%s]"), *OwnerDebugName);
+							bHasSimpleCollision = true;
 						}
+						break;
+
+						// Triangle mesh
+						case ImplicitObjectType::TriangleMesh:
+						case ImplicitObjectType::IsInstanced | ImplicitObjectType::TriangleMesh:
+						case ImplicitObjectType::IsScaled | ImplicitObjectType::TriangleMesh:
+							// TODO: We could eventually want to collide cloth against triangle meshes,
+							//       however the concept of simple vs complex shape might need to be clarified
+							//       as it currently iterates over all shape to discard the triangle mesh ones.
+							UE_LOG(LogSkeletalMesh, Verbose, TEXT("Found unusable Triangle Mesh cloth environmental collision in [%s]"), *OwnerDebugName);
+							break;
 
-						CollisionData.Append(SkelCollisionData);
+						default:
+							UE_LOG(LogSkeletalMesh, Verbose, TEXT("Found unsupported collision type during environmental collision with the cloth in [%s]"), *OwnerDebugName);
+							break;
+						}
 					}
-				}
-			}
+					bSuccessfulRead = true;
+				});
 		}
 	}
 }
