[StaticMesh] Expose GetNumTexCoords, GetNumNaniteVertices, and GetNumNaniteTriangles accessor functions to blueprint

#Synced-CL 51901069

diff --git a/Engine/Source/Runtime/Engine/Classes/Engine/StaticMesh.h b/Engine/Source/Runtime/Engine/Classes/Engine/StaticMesh.h
--- a/Engine/Source/Runtime/Engine/Classes/Engine/StaticMesh.h
+++ b/Engine/Source/Runtime/Engine/Classes/Engine/StaticMesh.h
@@ -2155,16 +2155,19 @@ public:
 	/**
 	 * Returns the number of tex coords for the specified LOD.
 	 */
+	UFUNCTION(BlueprintPure, Category = StaticMesh)
 	ENGINE_API int32 GetNumTexCoords(int32 LODIndex) const;
 
 	/**
 	 * Returns the number of vertices of the Nanite representation of this mesh.
 	 */
+	UFUNCTION(BlueprintPure, Category = StaticMesh)
 	ENGINE_API int32 GetNumNaniteVertices() const;
 
 	/**
 	 * Returns the number of triangles of the Nanite representation of this mesh.
 	 */
+	UFUNCTION(BlueprintPure, Category = StaticMesh)
 	ENGINE_API int32 GetNumNaniteTriangles() const;
 
 	/**
