[ISKM] Remove Experimental Warning From Instanced Skinned Mesh Component
- Also Fix Implicit Cast Warnings for AnimBank Usage

#rnx
#Synced-CL 51901791

diff --git a/Engine/Source/Runtime/Engine/Classes/Animation/AnimBank.h b/Engine/Source/Runtime/Engine/Classes/Animation/AnimBank.h
--- a/Engine/Source/Runtime/Engine/Classes/Animation/AnimBank.h
+++ b/Engine/Source/Runtime/Engine/Classes/Animation/AnimBank.h
@@ -222,9 +222,9 @@ public:
 	inline float GetSequenceCount() const
 	{
 #if WITH_EDITORONLY_DATA
-		return Sequences.Num();
+		return static_cast<float>(Sequences.Num());
 #else
-		return GetData().Entries.Num();
+		return static_cast<float>(GetData().Entries.Num());
 #endif
 	}
 
diff --git a/Engine/Source/Runtime/Engine/Classes/Animation/AnimSequenceTransformProviderData.h b/Engine/Source/Runtime/Engine/Classes/Animation/AnimSequenceTransformProviderData.h
--- a/Engine/Source/Runtime/Engine/Classes/Animation/AnimSequenceTransformProviderData.h
+++ b/Engine/Source/Runtime/Engine/Classes/Animation/AnimSequenceTransformProviderData.h
@@ -177,7 +177,7 @@ public:
 		if (bAutoPlay)
 		{
 			const double Position = (WorldTime - Auto.ReferenceTimestamp) * Auto.PlayRate;
-			return GetPosition(GetLoopMode(), Position, SequenceLength);
+			return GetPosition(GetLoopMode(), static_cast<float>(Position), static_cast<float>(SequenceLength));
 		}
 		return Manual.Position;
 	}
@@ -222,7 +222,7 @@ public:
 	{
 		check(IsAutoPlay());
 		const double Position = (WorldTime - Auto.ReferenceTimestamp) * Auto.PlayRate;
-		Auto.ReferenceTimestamp = GetTimestamp(WorldTime, Position, NewPlayRate);
+		Auto.ReferenceTimestamp = GetTimestamp(WorldTime, static_cast<float>(Position), NewPlayRate);
 		Auto.PlayRate = NewPlayRate;
 		PatchFlags |= (uint32)EAnimSequenceTrackPatchFlags::PlayRate;
 	}
@@ -256,12 +256,12 @@ private:
 	{
 		if (LoopMode == EAnimSequenceTrackLoopMode::Loop)
 		{
-			if (SequenceLength >= 0.0)
+			if (SequenceLength >= 0.0f)
 			{
-				Position = SequenceLength == 0.0 ? 0.0 : FMath::Fmod(Position, SequenceLength);
+				Position = SequenceLength == 0.0f ? 0.0f : FMath::Fmod(Position, SequenceLength);
 
 				// Returns a negative range if input is negative.
-				if (Position < 0.0)
+				if (Position < 0.0f)
 				{
 					Position += SequenceLength;
 				}
@@ -269,7 +269,7 @@ private:
 		}
 		else
 		{
-			Position = FMath::Clamp(Position, 0.0, SequenceLength);
+			Position = FMath::Clamp(Position, 0.0f, SequenceLength);
 		}
 		return Position;
 	}
diff --git a/Engine/Source/Runtime/Engine/Classes/Components/InstancedSkinnedMeshComponent.h b/Engine/Source/Runtime/Engine/Classes/Components/InstancedSkinnedMeshComponent.h
--- a/Engine/Source/Runtime/Engine/Classes/Components/InstancedSkinnedMeshComponent.h
+++ b/Engine/Source/Runtime/Engine/Classes/Components/InstancedSkinnedMeshComponent.h
@@ -51,7 +51,7 @@ struct FSkinnedMeshInstanceData
 };
 
 UCLASS(ClassGroup=Rendering, hidecategories=(Object,Activation,Collision,"Components|Activation",Physics), editinlinenew, meta=(BlueprintSpawnableComponent), Blueprintable, MinimalAPI)
-class UE_EXPERIMENTAL(5.6, "This class is currently extremely experimental and should not be used at this time.") UInstancedSkinnedMeshComponent : public USkinnedMeshComponent
+class UInstancedSkinnedMeshComponent : public USkinnedMeshComponent
 {
 	GENERATED_UCLASS_BODY()
 	
