Marked UAnimNotify as EditInlineNew

#Synced-CL 51945476

diff --git a/Engine/Source/Runtime/Engine/Classes/Animation/AnimNotifies/AnimNotify.h b/Engine/Source/Runtime/Engine/Classes/Animation/AnimNotifies/AnimNotify.h
--- a/Engine/Source/Runtime/Engine/Classes/Animation/AnimNotifies/AnimNotify.h
+++ b/Engine/Source/Runtime/Engine/Classes/Animation/AnimNotifies/AnimNotify.h
@@ -47,7 +47,7 @@ public:
 	{}
 };
 
-UCLASS(abstract, Blueprintable, const, hidecategories=Object, collapsecategories, MinimalAPI)
+UCLASS(abstract, editinlinenew, Blueprintable, const, hidecategories=Object, collapsecategories, MinimalAPI)
 class UAnimNotify : public UObject
 {
 	GENERATED_UCLASS_BODY()
