Fix FAccessTag crash when PointerHash returns zero for non-null pointers
PointerHash can return 0 for ~1 in 2^32 non-null pointers. FAccessTag uses the hash as a WriterTag, where 0 is the "no writer" sentinel in CurrentWriterTag.
Remap 0 to ~0u in the FAccessTag writer constructor.

#Synced-CL 51624647

diff --git a/Engine/Source/Runtime/Engine/Public/InstanceDataSceneProxy.h b/Engine/Source/Runtime/Engine/Public/InstanceDataSceneProxy.h
--- a/Engine/Source/Runtime/Engine/Public/InstanceDataSceneProxy.h
+++ b/Engine/Source/Runtime/Engine/Public/InstanceDataSceneProxy.h
@@ -146,8 +146,10 @@ public:
 			Writer
 		};
 		FAccessTag() : WriterTag(0), Kind(EKind::Reader) { }
-		FAccessTag(uint32 InWriterTag) : WriterTag(InWriterTag) , Kind(EKind::Writer) { }
-
+		
+		// WriterTag must be non-zero because 0 is the "no writer" sentinel in CurrentWriterTag
+		FAccessTag(uint32 InWriterTag) : WriterTag(InWriterTag ? InWriterTag : ~0u), Kind(EKind::Writer) {}
+		
 		uint32 WriterTag;
 		EKind Kind;
 	};
