Neural profile: PostEditChangeProperty should be toolonly

#rnx
#Synced-CL 51944828

diff --git a/Engine/Source/Runtime/Engine/Classes/Engine/NeuralProfile.h b/Engine/Source/Runtime/Engine/Classes/Engine/NeuralProfile.h
--- a/Engine/Source/Runtime/Engine/Classes/Engine/NeuralProfile.h
+++ b/Engine/Source/Runtime/Engine/Classes/Engine/NeuralProfile.h
@@ -149,8 +149,10 @@ public:
 	UPROPERTY()
 	FGuid Guid;
 
-	void BeginDestroy();
-	virtual void PostEditChangeProperty(struct FPropertyChangedEvent& PropertyChangedEvent);
+	virtual void BeginDestroy() override;
+#if WITH_EDITOR
+	virtual void PostEditChangeProperty(struct FPropertyChangedEvent& PropertyChangedEvent) override;
+#endif
 };
 
 namespace NeuralProfile
diff --git a/Engine/Source/Runtime/Engine/Private/Rendering/NeuralProfile.cpp b/Engine/Source/Runtime/Engine/Private/Rendering/NeuralProfile.cpp
--- a/Engine/Source/Runtime/Engine/Private/Rendering/NeuralProfile.cpp
+++ b/Engine/Source/Runtime/Engine/Private/Rendering/NeuralProfile.cpp
@@ -246,9 +246,10 @@ void UNeuralProfile::BeginDestroy()
 
 	Super::BeginDestroy();
 }
-
+#if WITH_EDITOR
 void UNeuralProfile::PostEditChangeProperty(struct FPropertyChangedEvent& PropertyChangedEvent)
 {
+	Super::PostEditChangeProperty(PropertyChangedEvent);
 	const FNeuralProfileStruct LocalSettings = this->Settings;
 	UNeuralProfile* LocalProfile = this;
 	GetRendererModule().InvalidatePathTracedOutput();
@@ -275,7 +276,7 @@ void UNeuralProfile::PostEditChangeProperty(struct FPropertyChangedEvent& Proper
 			});
 	}
 }
-
+#endif
 /////////////////////////////////////////////////////////////////////////////////////////////////////
 //// Public API
 //
