Clean up Subsystem headers to remove redundant empty functions and improve comments

#Synced-CL 51908825

diff --git a/Engine/Source/Editor/EditorSubsystem/Private/EditorSubsystem.cpp b/Engine/Source/Editor/EditorSubsystem/Private/EditorSubsystem.cpp
deleted file mode 100644
diff --git a/Engine/Source/Editor/EditorSubsystem/Public/EditorSubsystem.h b/Engine/Source/Editor/EditorSubsystem/Public/EditorSubsystem.h
--- a/Engine/Source/Editor/EditorSubsystem/Public/EditorSubsystem.h
+++ b/Engine/Source/Editor/EditorSubsystem/Public/EditorSubsystem.h
@@ -3,15 +3,9 @@
 #pragma once
 
 #include "Subsystems/Subsystem.h"
-#include "UObject/ObjectMacros.h"
-#include "UObject/UObjectGlobals.h"
 
 #include "EditorSubsystem.generated.h"
 
-#define UE_API EDITORSUBSYSTEM_API
-
-class UObject;
-
 /**
  * UEditorSubsystem
  * Base class for auto instanced and initialized systems that share the lifetime of the Editor
@@ -22,14 +16,9 @@ class UObject;
  * If the subsystem collection is created post module load then the instances will be created at
  * collection initialization time.
  */
-
 UCLASS(MinimalAPI, Abstract)
 class UEditorSubsystem : public UDynamicSubsystem
 {
 	GENERATED_BODY()
-
-public:
-	UE_API UEditorSubsystem();
 };
 
-#undef UE_API
diff --git a/Engine/Source/Runtime/Engine/Private/Subsystems/EngineSubsystem.cpp b/Engine/Source/Runtime/Engine/Private/Subsystems/EngineSubsystem.cpp
deleted file mode 100644
diff --git a/Engine/Source/Runtime/Engine/Private/Subsystems/GameInstanceSubsystem.cpp b/Engine/Source/Runtime/Engine/Private/Subsystems/GameInstanceSubsystem.cpp
--- a/Engine/Source/Runtime/Engine/Private/Subsystems/GameInstanceSubsystem.cpp
+++ b/Engine/Source/Runtime/Engine/Private/Subsystems/GameInstanceSubsystem.cpp
@@ -5,12 +5,6 @@
 
 #include UE_INLINE_GENERATED_CPP_BY_NAME(GameInstanceSubsystem)
 
-UGameInstanceSubsystem::UGameInstanceSubsystem()
-	: USubsystem()
-{
-
-}
-
 UGameInstance* UGameInstanceSubsystem::GetGameInstance() const
 {
 	return Cast<UGameInstance>(GetOuter());
diff --git a/Engine/Source/Runtime/Engine/Private/Subsystems/LocalPlayerSubsystem.cpp b/Engine/Source/Runtime/Engine/Private/Subsystems/LocalPlayerSubsystem.cpp
deleted file mode 100644
diff --git a/Engine/Source/Runtime/Engine/Private/Subsystems/Subsystem.cpp b/Engine/Source/Runtime/Engine/Private/Subsystems/Subsystem.cpp
--- a/Engine/Source/Runtime/Engine/Private/Subsystems/Subsystem.cpp
+++ b/Engine/Source/Runtime/Engine/Private/Subsystems/Subsystem.cpp
@@ -5,18 +5,7 @@
 
 #include UE_INLINE_GENERATED_CPP_BY_NAME(Subsystem)
 
-USubsystem::USubsystem()
-{
-
-}
-
 int32 USubsystem::GetFunctionCallspace(UFunction* Function, FFrame* Stack)
 {
 	return GEngine->GetGlobalFunctionCallspace(Function, this, Stack);
 }
-
-UDynamicSubsystem::UDynamicSubsystem()
-	: USubsystem()
-{
-
-}
diff --git a/Engine/Source/Runtime/Engine/Private/Subsystems/WorldSubsystem.cpp b/Engine/Source/Runtime/Engine/Private/Subsystems/WorldSubsystem.cpp
--- a/Engine/Source/Runtime/Engine/Private/Subsystems/WorldSubsystem.cpp
+++ b/Engine/Source/Runtime/Engine/Private/Subsystems/WorldSubsystem.cpp
@@ -2,19 +2,9 @@
 
 #include "Subsystems/WorldSubsystem.h"
 #include "Engine/World.h"
-#include "Subsystems/Subsystem.h"
-#include "Streaming/StreamingWorldSubsystemInterface.h"
 
 #include UE_INLINE_GENERATED_CPP_BY_NAME(WorldSubsystem)
 
-// ----------------------------------------------------------------------------------
-
-UWorldSubsystem::UWorldSubsystem()
-	: USubsystem()
-{
-
-}
-
 UWorld& UWorldSubsystem::GetWorldRef() const
 {
 	return *CastChecked<UWorld>(GetOuter(), ECastCheckedType::NullChecked);
@@ -41,11 +31,6 @@ bool UWorldSubsystem::DoesSupportWorldType(const EWorldType::Type WorldType) con
 	return WorldType == EWorldType::Game || WorldType == EWorldType::Editor || WorldType == EWorldType::PIE;
 }
 
-void UWorldSubsystem::Initialize(FSubsystemCollectionBase& Collection)
-{
-	Super::Initialize(Collection);
-}
-
 void UWorldSubsystem::PostInitialize()
 {
 	ensureAlwaysMsgf(!bHasCalledPostInitialize, TEXT("PostInitialize has already been called for subsystem %s"), *GetClass()->GetName());
@@ -67,16 +52,6 @@ void UWorldSubsystem::EnsureHasCalledBeginPlay() const
 	ensureAlwaysMsgf(bHasCalledBeginPlay, TEXT("OnWorldBeginPlay has not been called for subsystem %s, check for missing Super::OnWorldBeginPlay call."), *GetClass()->GetName());
 }
 
-PRAGMA_DISABLE_DEPRECATION_WARNINGS
-void UWorldSubsystem::UpdateStreamingState()
-{
-	if (IStreamingWorldSubsystemInterface* StreamingWorldSubsystem = Cast<IStreamingWorldSubsystemInterface>(this))
-	{
-		StreamingWorldSubsystem->OnUpdateStreamingState();
-	}
-}
-PRAGMA_ENABLE_DEPRECATION_WARNINGS
-
 // ----------------------------------------------------------------------------------
 
 UTickableWorldSubsystem::UTickableWorldSubsystem()
diff --git a/Engine/Source/Runtime/Engine/Public/Subsystems/EngineSubsystem.h b/Engine/Source/Runtime/Engine/Public/Subsystems/EngineSubsystem.h
--- a/Engine/Source/Runtime/Engine/Public/Subsystems/EngineSubsystem.h
+++ b/Engine/Source/Runtime/Engine/Public/Subsystems/EngineSubsystem.h
@@ -16,13 +16,8 @@
  * If the subsystem collection is created post module load then the instances will be created at
  * collection initialization time.
  */
-
 UCLASS(Abstract, MinimalAPI)
 class UEngineSubsystem : public UDynamicSubsystem
 {
 	GENERATED_BODY()
-
-public:
-	ENGINE_API UEngineSubsystem();
-
 };
diff --git a/Engine/Source/Runtime/Engine/Public/Subsystems/GameInstanceSubsystem.h b/Engine/Source/Runtime/Engine/Public/Subsystems/GameInstanceSubsystem.h
--- a/Engine/Source/Runtime/Engine/Public/Subsystems/GameInstanceSubsystem.h
+++ b/Engine/Source/Runtime/Engine/Public/Subsystems/GameInstanceSubsystem.h
@@ -18,8 +18,10 @@ class UGameInstanceSubsystem : public USubsystem
 	GENERATED_BODY()
 
 public:
-	ENGINE_API UGameInstanceSubsystem();
-
+	/**
+	 * Returns a pointer to the UGameInstance this subsystem is contained within.
+	 * @note This will return null for default objects
+	 */
 	ENGINE_API UGameInstance* GetGameInstance() const;
 
 };
diff --git a/Engine/Source/Runtime/Engine/Public/Subsystems/LocalPlayerSubsystem.h b/Engine/Source/Runtime/Engine/Public/Subsystems/LocalPlayerSubsystem.h
--- a/Engine/Source/Runtime/Engine/Public/Subsystems/LocalPlayerSubsystem.h
+++ b/Engine/Source/Runtime/Engine/Public/Subsystems/LocalPlayerSubsystem.h
@@ -19,8 +19,6 @@ class ULocalPlayerSubsystem : public USubsystem
 	GENERATED_BODY()
 
 public:
-	ENGINE_API ULocalPlayerSubsystem();
-
 	template<typename LocalPlayerType = ULocalPlayer>
 	LocalPlayerType* GetLocalPlayer() const
 	{
@@ -34,5 +32,5 @@ public:
 	}
 
 	/** Callback for when the player controller is changed on this subsystem's owning local player */
-	ENGINE_API virtual void PlayerControllerChanged(APlayerController* NewPlayerController);
+	virtual void PlayerControllerChanged(APlayerController* NewPlayerController) {}
 };
diff --git a/Engine/Source/Runtime/Engine/Public/Subsystems/Subsystem.h b/Engine/Source/Runtime/Engine/Public/Subsystems/Subsystem.h
--- a/Engine/Source/Runtime/Engine/Public/Subsystems/Subsystem.h
+++ b/Engine/Source/Runtime/Engine/Public/Subsystems/Subsystem.h
@@ -9,8 +9,8 @@
 
 class FSubsystemCollectionBase;
 
-/** Subsystems are auto instanced classes that share the lifetime of certain engine constructs
- * 
+/** 
+ *	Subsystems are auto instanced classes that share the lifetime of certain engine constructs
  *	Currently supported Subsystem lifetimes are:
  *		Engine		 -> inherit UEngineSubsystem
  *		Editor		 -> inherit UEditorSubsystem
@@ -18,7 +18,6 @@ class FSubsystemCollectionBase;
  *		World		 -> inherit UWorldSubsystem
  *		LocalPlayer	 -> inherit ULocalPlayerSubsystem
  *
- *
  *	Normal Example:
  *		class UMySystem : public UGameInstanceSubsystem
  *	Which can be accessed by:
@@ -29,7 +28,6 @@ class FSubsystemCollectionBase;
  *		UGameInstance* GameInstance = ...;
  *		UMyGameSubsystem* MySubsystem = UGameInstance::GetSubsystem<MyGameSubsystem>(GameInstance);
  *
- *
  *	You can get also define interfaces that can have multiple implementations.
  *	Interface Example :
  *      MySystemInterface
@@ -40,7 +38,6 @@ class FSubsystemCollectionBase;
  *	Which can be accessed by:
  *		UGameInstance* GameInstance = ...;
  *		const TArray<UMyGameSubsystem*>& MySubsystems = GameInstance->GetSubsystemArray<MyGameSubsystem>();
- *
  */
 
 UCLASS(Abstract, MinimalAPI)
@@ -49,14 +46,12 @@ class USubsystem : public UObject
 	GENERATED_BODY()
 
 public:
-	ENGINE_API USubsystem();
-
-	/** Override to control if the Subsystem should be created at all.
+	/**
+	 * Override to control if the Subsystem should be created at all.
 	 * For example you could only have your system created on servers.
-	 * It's important to note that if using this is becomes very important to null check whenever getting the Subsystem.
+	 * For subsystems that can return false, you must null check whenever getting the Subsystem.
 	 *
 	 * Note: This function is called on the CDO prior to instances being created!
-	 *
 	 */
 	virtual bool ShouldCreateSubsystem(UObject* Outer) const { return true; }
 
@@ -71,14 +66,12 @@ public:
 
 private:
 	friend class FSubsystemCollectionBase;
-	FSubsystemCollectionBase* InternalOwningSubsystem;
-
+	FSubsystemCollectionBase* InternalOwningSubsystem = nullptr;
 };
 
-
-/** Dynamic Subsystems auto populate/depopulate existing collections when modules are loaded and unloaded
- *
- * Only UEngineSubsystems and UEditorSubsystems allow for dynamic loading.
+/** 
+ * Dynamic Subsystems auto populate/depopulate existing collections when modules are loaded and unloaded.
+ * This abstract base class is used for singletons like UEngineSubsystem that are automatically created.
  * 
  * If instances of your subsystem aren't being created it maybe that the module they are in isn't being explicitly loaded,
  * make sure there is a LoadModule("ModuleName") to load the module.
@@ -87,7 +80,4 @@ UCLASS(Abstract, MinimalAPI)
 class UDynamicSubsystem : public USubsystem
 {
 	GENERATED_BODY()
-
-public:
-	ENGINE_API UDynamicSubsystem();
 };
diff --git a/Engine/Source/Runtime/Engine/Public/Subsystems/WorldSubsystem.h b/Engine/Source/Runtime/Engine/Public/Subsystems/WorldSubsystem.h
--- a/Engine/Source/Runtime/Engine/Public/Subsystems/WorldSubsystem.h
+++ b/Engine/Source/Runtime/Engine/Public/Subsystems/WorldSubsystem.h
@@ -10,6 +10,7 @@
 /**
  * UWorldSubsystem
  * Base class for auto instanced and initialized systems that share the lifetime of a UWorld
+ * FWorldSubsystemCollection handles calling PostInitialize and OnWorldBeginPlay for already initialized worlds
  */
 UCLASS(Abstract, MinimalAPI)
 class UWorldSubsystem : public USubsystem
@@ -17,8 +18,10 @@ class UWorldSubsystem : public USubsystem
 	GENERATED_BODY()
 
 public:
-	ENGINE_API UWorldSubsystem();
-
+	/**
+	 * Returns a pointer to the UWorld this subsystem is contained within.
+	 * @note This will return null for default objects
+	 */
 	ENGINE_API virtual UWorld* GetWorld() const override final;
 
 	/**
@@ -27,14 +30,9 @@ public:
 	 */
 	ENGINE_API UWorld& GetWorldRef() const;
 
+	/** Calls DoesSupportWorldType to check if this subsystem should be created */
 	ENGINE_API virtual bool ShouldCreateSubsystem(UObject* Outer) const override;
 
-	/** 
-	* Initializes the world subsystem.  
-	* Will execute PostInitialize if the world has already been Initialize
-	* Will execute OnWorldBeginPlay if the world has already begun play*/
-	ENGINE_API virtual void Initialize(FSubsystemCollectionBase& Collection) override;
-
 	/** Called once all UWorldSubsystems have been initialized */
 	ENGINE_API virtual void PostInitialize();
 
@@ -56,14 +54,15 @@ public:
 	/** Returns true once the subsystem has executed OnWorldBeginPlay **/
 	bool HasCalledBeginPlay() const { return bHasCalledBeginPlay; }
 
+	/** Called during initialization to ensure the correct callbacks have happened */
 	ENGINE_API void EnsureHasCalledPostInitialize() const;
-
 	ENGINE_API void EnsureHasCalledBeginPlay() const;
 
-	UE_DEPRECATED(5.5, "Implement IStreamingWorldSubsystemInterface instead.")
-	ENGINE_API virtual void UpdateStreamingState();
-
 protected:
+	/** 
+	 * Override this to change what types of worlds support this subsystem.
+	 * By default, world subsystems support game, PIE, and editor worlds but not preview worlds.
+	 */
 	ENGINE_API virtual bool DoesSupportWorldType(const EWorldType::Type WorldType) const;
 
 	bool bHasCalledPostInitialize : 1 = false;
