Fix for crash building morph targets with legacy FBX importer when UseT0AsRefPose is set
#Synced-CL 51920915

diff --git a/Engine/Source/Runtime/Engine/Private/Rendering/SkeletalMeshLODImporterData.cpp b/Engine/Source/Runtime/Engine/Private/Rendering/SkeletalMeshLODImporterData.cpp
--- a/Engine/Source/Runtime/Engine/Private/Rendering/SkeletalMeshLODImporterData.cpp
+++ b/Engine/Source/Runtime/Engine/Private/Rendering/SkeletalMeshLODImporterData.cpp
@@ -25,10 +25,13 @@ DEFINE_LOG_CATEGORY_STATIC(LogSkeletalMeshLODImporterData, Log, All);
 
 void FSkeletalMeshImportData::CopyDataNeedByMorphTargetImport(FSkeletalMeshImportData& Other) const
 {
-	//The points array is the only data we need to compute the morph target in the skeletalmesh build
 	Other.Points = Points;
 	//PointToRawMap should not be save when saving morph target data, we only need it temporary to gather the point from the fbx shape
 	Other.PointToRawMap = PointToRawMap;
+
+	Other.Faces = Faces;
+	Other.bHasNormals = bHasNormals;
+	Other.bHasTangents = bHasTangents;
 }
 
 void FSkeletalMeshImportData::KeepAlternateSkinningBuildDataOnly()
