Avoid crashes due to missing CachedExpressionData on MaterialInstanceDynamics when cooking.

- This was happening while gathering references and we would serialize a MID triggering the cached expression check (Trying to save cooked material instance %s, missing CachedExpressionData)
- Since we don't need CachedExpressionData on MIDs during a cook, we can safely skip these checks when gathering references.
- These MID references were being serialized via a FImportExportCollector archive.

#rnx

diff --git a/Engine/Source/Runtime/Engine/Private/Materials/MaterialInstance.cpp b/Engine/Source/Runtime/Engine/Private/Materials/MaterialInstance.cpp
--- a/Engine/Source/Runtime/Engine/Private/Materials/MaterialInstance.cpp
+++ b/Engine/Source/Runtime/Engine/Private/Materials/MaterialInstance.cpp
@@ -3295,8 +3295,11 @@ void UMaterialInstance::Serialize(FArchive& Ar)
 			}
 			else if (!bAllowMissingCachedData)
 			{
-				// ClassDefault object is expected to be missing cached data, but in all other cases it should have been created when the material was loaded, in PostLoad
-				checkf(HasAllFlags(RF_ClassDefaultObject), TEXT("Trying to save cooked material instance %s, missing CachedExpressionData"), *GetName());
+				if (!Ar.IsObjectReferenceCollector())
+				{
+					// ClassDefault object is expected to be missing cached data, but in all other cases it should have been created when the material was loaded, in PostLoad
+					checkf(HasAllFlags(RF_ClassDefaultObject), TEXT("Trying to save cooked material instance %s, missing CachedExpressionData"), *GetName());
+				}
 			}
 		}
 
