Fix non-deterministic extern deduplication caused by hashing uninitialized padding bytes.

Extern definitions are hashed by their raw byte pattern for deduplication. When a definition struct has padding bytes, those bytes could contain arbitrary stack data, causing identical definitions to hash differently and fail to deduplicate.

The fix copies the definition into a zero-initialized buffer member-by-member. Types that satisfy std::has_unique_object_representations_v (no padding) are assigned directly. Types that don't must implement a CopyTo() method that copies each member individually, leaving padding bytes zeroed. Empty types skip copying entirely.

#rnx
#Synced-CL 51827487

diff --git a/Engine/Source/Runtime/Engine/Private/Materials/MaterialExpressionsToMIR.cpp b/Engine/Source/Runtime/Engine/Private/Materials/MaterialExpressionsToMIR.cpp
--- a/Engine/Source/Runtime/Engine/Private/Materials/MaterialExpressionsToMIR.cpp
+++ b/Engine/Source/Runtime/Engine/Private/Materials/MaterialExpressionsToMIR.cpp
@@ -1515,6 +1515,13 @@ struct FScreenTextureExtern
 				break;
 		}
 	}
+
+	void CopyTo(FScreenTextureExtern& Other) const
+	{
+		Other.Kind = Kind;
+		Other.SceneTextureId = SceneTextureId;
+		Other.UserSceneTexture = UserSceneTexture;
+	}
 };
 
 // If DefaultOffset is not null, Coordinates are treated as an offset (or DefaultOffset if unset), rather than absolute coordinates.  Clamping is
@@ -3279,6 +3286,12 @@ static FValueRef BuildViewProperty(MIR::FEmitter& Em, EMaterialExposedViewProper
 		{
 			Printer << (bInvProperty ? PropertyMeta->InvPropertyCode : PropertyMeta->PropertyCode);
 		}
+
+		void CopyTo(FViewPropertyExtern& Other) const
+		{
+			Other.PropertyMeta = PropertyMeta;
+			Other.bInvProperty = bInvProperty;
+		}
 	};
 
 	const FMaterialExposedViewPropertyMeta& PropertyMeta = MaterialExternalCodeRegistry::Get().GetExternalViewPropertyCode(InProperty);
@@ -3467,6 +3480,12 @@ static FValueRef EmitActorOrObjectPosition(MIR::FEmitter& Em, EActorOrObjectPosi
 		{
 			Out.Appendf(TEXT("IsActorPosition=%d, bIsAbsolute=%d"), Mode == EActorOrObjectPosition::Actor, bIsAbsolute);
 		}
+
+		void CopyTo(FActorAndObjectPositionExtern& Other) const
+		{
+			Other.Mode = Mode;
+			Other.bIsAbsolute = bIsAbsolute;
+		}
 	};
 
 	FValueRef Value = Em.Extern<FActorAndObjectPositionExtern>({ Mode, Origin == EPositionOrigin::Absolute });
@@ -4645,6 +4664,12 @@ void UMaterialExpressionDynamicParameter::Build(MIR::FEmitter& Em)
 		{
 			Out.Appendf(TEXT("Default=(%f, %f, %f, %f) ParameterIndex=%d"), Default.X, Default.Y, Default.Z, Default.W, ParameterIndex);
 		}
+
+		void CopyTo(FDynamicParameter& Other) const
+		{
+			Other.Default = Default;
+			Other.ParameterIndex = ParameterIndex;
+		}
 	};
 
 	FValueRef Result = Em.Extern<FDynamicParameter>({
diff --git a/Engine/Source/Runtime/Engine/Public/Materials/MaterialIREmitter.h b/Engine/Source/Runtime/Engine/Public/Materials/MaterialIREmitter.h
--- a/Engine/Source/Runtime/Engine/Public/Materials/MaterialIREmitter.h
+++ b/Engine/Source/Runtime/Engine/Public/Materials/MaterialIREmitter.h
@@ -1077,11 +1077,28 @@ FValueRef FEmitter::Extern(TDefinition DefinitionPrototype, TArguments... Argume
 	// Read the value type from the definition prototype
 	MIR::FType Type = DefinitionPrototype.GetInfo().Type;
 
-	// Create the concrete extern instruction prototype and pass in the prototype
+	// Create the concrete extern instruction prototype.
 	TTypeCompatibleBytes<TConcreteExtern> ProtoStorage;
 	FMemory::Memzero(ProtoStorage);
 
-	TConcreteExtern& ConcreteProto = * new (&ProtoStorage) TConcreteExtern{ MoveTemp(DefinitionPrototype) };
+	// Copy the definition into the zeroed storage. The approach depends on whether the definition
+	// type has padding bytes: types without padding can be assigned directly, while types with
+	// padding must use CopyTo to write only the real members, keeping padding bytes zeroed for
+	// deterministic hashing.
+	TConcreteExtern& ConcreteProto = *ProtoStorage.GetTypedPtr();
+	if constexpr (std::is_empty_v<TDefinition>)
+	{
+		// Nothing to copy.
+	}
+	else if constexpr (std::has_unique_object_representations_v<TDefinition>)
+	{
+		*static_cast<TDefinition*>(&ConcreteProto) = DefinitionPrototype;
+	}
+	else
+	{
+		static_assert(CExternHasCopyTo<TDefinition>, "This extern definition type has padding bytes and must implement 'void CopyTo(T& Other) const' to ensure deterministic hashing. CopyTo must copy each member individually into Other.");
+		DefinitionPrototype.CopyTo(static_cast<TDefinition&>(ConcreteProto));
+	}
 
 	// Copy over all the provided arguments
 	int32 Index = 0;
diff --git a/Engine/Source/Runtime/Engine/Public/Materials/MaterialIRExtern.cpp b/Engine/Source/Runtime/Engine/Public/Materials/MaterialIRExtern.cpp
--- a/Engine/Source/Runtime/Engine/Public/Materials/MaterialIRExtern.cpp
+++ b/Engine/Source/Runtime/Engine/Public/Materials/MaterialIRExtern.cpp
@@ -142,6 +142,12 @@ void FExternFromMaterialDecl::EmitDebugInfo(FString& Out) const
 	Out.Appendf(TEXT("Name=\"%s\""), *Name.ToString());
 }
 
+void FExternFromMaterialDecl::CopyTo(FExternFromMaterialDecl& Other) const
+{
+	Other.Name = Name;
+	Other.Declaration = Declaration;
+}
+
 //------------------------------------------------------------------------
 //                          FExternFromInlineDecl
 //------------------------------------------------------------------------
@@ -263,6 +269,12 @@ void FExternSimpleHLSL::EmitDebugInfo(FString& Out) const
 	Out.Appendf(TEXT("Code=\"%s\""), CodeHLSL);
 }
 
+void FExternSimpleHLSL::CopyTo(FExternSimpleHLSL& Other) const
+{
+	Other.CodeHLSL = CodeHLSL;
+	Other.Type = Type;
+}
+
 //------------------------------------------------------------------------
 //                        FExternPrinterHLSL
 //------------------------------------------------------------------------
diff --git a/Engine/Source/Runtime/Engine/Public/Materials/MaterialIRExtern.h b/Engine/Source/Runtime/Engine/Public/Materials/MaterialIRExtern.h
--- a/Engine/Source/Runtime/Engine/Public/Materials/MaterialIRExtern.h
+++ b/Engine/Source/Runtime/Engine/Public/Materials/MaterialIRExtern.h
@@ -195,6 +195,18 @@ struct FExternInfo
 //                 ... /* other modes */
 //             }
 //         }
+//
+//         // Optional. If your definition struct does not satisfy std::has_unique_object_representations_v,
+//         // you must implement CopyTo to copy each member individually into Other. The emitter writes
+//         // your definition into a zero-initialized buffer via CopyTo so that padding bytes stay zero,
+//         // ensuring deterministic hashing. Note: the trait is also false for structs containing floats
+//         // (due to +0/-0 and NaN aliasing), even if they have no padding — C++ does not offer a trait
+//         // to detect padding alone, so this is the best we can do. In that case CopyTo is still
+//         // required but is a straightforward member-by-member copy.
+//         void CopyTo(FColorConversionExtern& Other) const
+//         {
+//             Other.Mode = Mode;
+//         }
 //     };
 //
 template <typename T>
@@ -207,6 +219,13 @@ concept CExternDefinition = requires(const T& Self, FExternPrinterHLSL& Printer)
 	{ Self.ToHLSL(Printer) } -> std::same_as<void>;
 };
 
+// Detected by FEmitter::Extern to select the padding-safe copy path.
+template <typename T>
+concept CExternHasCopyTo = requires(const T& Self, T& Other)
+{
+	{ Self.CopyTo(Other) } -> std::same_as<void>;
+};
+
 // See CExternDefinition
 template<typename T>
 concept CExternHasAnalyze = requires(T& Self, FExternAnalysisContext& Context)
@@ -246,6 +265,7 @@ struct FExternFromMaterialDecl
 	void AnalyzeInStage(FExternAnalysisContext& Context, EStage Stage);
 	void ToHLSL(FExternPrinterHLSL& Printer) const;
 	void EmitDebugInfo(FString& Out) const;
+	void CopyTo(FExternFromMaterialDecl& Other) const;
 };
 
 // Helper function to create a material domain mask from a static list.
@@ -395,6 +415,7 @@ struct FExternSimpleHLSL
 	MIR::FExternInfo GetInfo() const;
 	void ToHLSL(MIR::FExternPrinterHLSL& Printer) const;
 	void EmitDebugInfo(FString& Out) const;
+	void CopyTo(FExternSimpleHLSL& Other) const;
 };
 
 // You can << the ExternTag_DD tag to automatically print "DDX" or "DDY" if the printer `Differential`
