New material translator:
- Branch optimization for aggregate types. When both arms of a branch produce the same value for one or more aggregate attributes, the new material translator now extracts those directly instead of branching on them. Only attributes that actually differ are branched on; if all attributes are identical the branch is skipped entirely.
- Some refactoring and cleanup in the HLSL generator to better handle inline instructions.

#rnx

#Synced-CL 51827487

diff --git a/Engine/Source/Runtime/Engine/Private/Materials/MaterialIREmitter.cpp b/Engine/Source/Runtime/Engine/Private/Materials/MaterialIREmitter.cpp
--- a/Engine/Source/Runtime/Engine/Private/Materials/MaterialIREmitter.cpp
+++ b/Engine/Source/Runtime/Engine/Private/Materials/MaterialIREmitter.cpp
@@ -2975,13 +2975,84 @@ FValueRef FEmitter::Branch(FValueRef Condition, FValueRef True, FValueRef False)
 	{
 		return Select(Condition, True, False);
 	}
-
-	// Create the branch instruction.
+	
+	// Create the branch instruction. For aggregate types, if both branches share some identical
+	// attribute values (pointer-equal due to deduplication), we can extract those directly and only
+	// branch on the attributes that actually differ. When the branch result is an aggregate, since branch
+	// instruction emission is deferred, it is skipped entirely when all aggregate attributes are identical.
 	FBranch Proto = MakePrototype<FBranch>(CommonType);
 	Proto.ConditionArg = Condition;
 	Proto.TrueArg = True;
 	Proto.FalseArg = False;
 
+	// When both branches are composite aggregates, build an aggregate result where identical
+	// attributes are extracted directly (bypassing the branch) and differing attributes are
+	// subscripted from the branch result.
+	if (const UMaterialAggregate* Aggregate = CommonType.AsAggregate())
+	{
+		const MIR::FComposite* TrueAsComposite = True->As<MIR::FComposite>();
+		const MIR::FComposite* FalseAsComposite = False->As<MIR::FComposite>();
+		if (TrueAsComposite && FalseAsComposite)
+		{
+			const int32 NumAttributes = Aggregate->Attributes.Num();
+			TConstArrayView<FValue*> TrueComponents = TrueAsComposite->GetComponents();
+			TConstArrayView<FValue*> FalseComponents = FalseAsComposite->GetComponents();
+
+			// First pass: check if any attribute can be collapsed.
+			bool bHasCollapsed = false;
+			for (int32 i = 0; i < NumAttributes; ++i)
+			{
+				if (TrueComponents[i] == FalseComponents[i])
+				{
+					bHasCollapsed = true;
+					break;
+				}
+			}
+
+			if (bHasCollapsed)
+			{
+				// Create stripped composites for TrueArg/FalseArg: null out identical
+				// attributes so the translator skips them, emitting only what differs.
+				FComposite* StrippedTrue = MakeCompositePrototype(*this, CommonType, NumAttributes);
+				FComposite* StrippedFalse = MakeCompositePrototype(*this, CommonType, NumAttributes);
+				TArrayView<FValue*> StrippedTrueComponents = StrippedTrue->GetMutableComponents();
+				TArrayView<FValue*> StrippedFalseComponents = StrippedFalse->GetMutableComponents();
+
+				TTemporaryArray<FValueRef> AttributeValues(NumAttributes);
+
+				for (int32 i = 0; i < NumAttributes; ++i)
+				{
+					if (TrueComponents[i] == FalseComponents[i])
+					{
+						// Both branches produce the same value — use it directly.
+						AttributeValues[i] = TrueComponents[i];
+					}
+					else
+					{
+						StrippedTrueComponents[i] = TrueComponents[i];
+						StrippedFalseComponents[i] = FalseComponents[i];
+					}
+				}
+
+				// Replace the branch args with the stripped composites.
+				Proto.TrueArg = EmitPrototype(*this, *StrippedTrue);
+				Proto.FalseArg = EmitPrototype(*this, *StrippedFalse);
+
+				// Emit the branch, then build the result aggregate.
+				FValue* BranchResult = EmitPrototype(*this, Proto);
+				for (int32 i = 0; i < NumAttributes; ++i)
+				{
+					if (TrueComponents[i] != FalseComponents[i])
+					{
+						AttributeValues[i] = Subscript(BranchResult, i);
+					}
+				}
+
+				return this->Aggregate(Aggregate, FValueRef(), AttributeValues);
+			}
+		}
+	}
+
 	return EmitPrototype(*this, Proto);
 }
 
@@ -4126,7 +4197,7 @@ FValueRef FEmitter::AnalyticalPartialDerivative(FValueRef Value, EDerivativeAxis
 
 			if ((Info.Flags & MIR::EExternFlags::ZeroDifferentials) != MIR::EExternFlags::None)
 			{
-				return ConstantDefault(Info.DifferentialType);
+				return ConstantDefault(DerivativeType);
 			}
 			else if ((Info.Flags & MIR::EExternFlags::NoDifferentials) != MIR::EExternFlags::None)
 			{
diff --git a/Engine/Source/Runtime/Engine/Private/Materials/MaterialIRToHLSLTranslator.cpp b/Engine/Source/Runtime/Engine/Private/Materials/MaterialIRToHLSLTranslator.cpp
--- a/Engine/Source/Runtime/Engine/Private/Materials/MaterialIRToHLSLTranslator.cpp
+++ b/Engine/Source/Runtime/Engine/Private/Materials/MaterialIRToHLSLTranslator.cpp
@@ -113,12 +113,25 @@ struct FPrinter
 
 	FPrinter& operator<<(EEndOfStatement)
 	{
-		if (Buffer.IsEmpty() || Buffer[Buffer.Len() - 1] != TEXT('}'))
+		// Find the last non-whitespace character to determine how to terminate.
+		int32 LastNonWS = Buffer.Len() - 1;
+		while (LastNonWS >= 0 && FChar::IsWhitespace(Buffer[LastNonWS]))
+		{
+			--LastNonWS;
+		}
+
+		if (LastNonWS >= 0 && Buffer[LastNonWS] == TEXT(';'))
+		{
+			return *this; // already terminated, nothing to do
+		}
+
+		if (LastNonWS < 0 || Buffer[LastNonWS] != TEXT('}'))
 		{
 			Buffer.AppendChar(TEXT(';'));
 		}
-        *this << NewLine;
-        return *this;
+
+		*this << NewLine;
+		return *this;
 	}
 
     FPrinter& operator<<(EOpenBlock)
@@ -310,6 +323,7 @@ struct FEntryPointCodeGen
 		CurrentEntryPointIndex = ModuleEntryPointIndex;
 		CurrentStage = EntryPoint.Stage;
 		bCompilingPreviousFrame = bPreviousFrame;
+		InstrToLocalIndex.Empty(InstrToLocalIndex.Num());
 
 		LowerBlock(EntryPoint.RootBlock);
 
@@ -321,15 +335,7 @@ struct FEntryPointCodeGen
 		int32 OldNumLocals = NumLocals;
 		for (MIR::FInstruction* Instr = Block.Instructions; Instr; Instr = Instr->GetNext(CurrentEntryPointIndex))
 		{
-			// If this is a inlinable instruction type and it has only one user, instead of allocating a local
-			// and evaluating the instruction into it, simply skip this instruction for now.
-			// It will be generated inlined by its only user later.
-			if (Instr->GetNumUsers(CurrentEntryPointIndex) == 1 && IsFoldable(Instr, CurrentEntryPointIndex))
-			{
-				continue;
-			}
-
-			// @jason.hoerner This works but it isn't ideal. We should instead omit linking instruction
+			// @massimo.tristano This works but it isn't ideal. We should instead omit linking instruction
 			// in the Step_LinkInstructions (skip them) that do not satisfy this condition. Only instructions that truly need
 			// evaluation in the various stages should be linked so that here we never skip any.
 
@@ -345,6 +351,12 @@ struct FEntryPointCodeGen
 				continue;
 			}
 
+			// If instruction should be inlined, don't lower it to HLSL now, but when it is referenced.
+			if (InstructionShouldBeInlined(Instr))
+			{
+				continue;
+			}
+
 			// If this isn't a set material output instruction but last one was, add a newline for better readability.
 			const MIR::FSetMaterialOutput* SetMaterialOutput = MIR::As<MIR::FSetMaterialOutput>(Instr);
 			if (!SetMaterialOutput && bLastInstructionWasSetMaterialOutput)
@@ -352,11 +364,10 @@ struct FEntryPointCodeGen
 				Printer << NewLine;
 				bLastInstructionWasSetMaterialOutput = false;
 			}
-            
-			// If this instruction has more than one user, we need to declare a local that will contain its
-			// result, e.g. "float _4 = <expression>;"
-			// Note that 'Call' instructions have their own way to generate locals declaration (they may use more than one).
-			if (!Instr->As<MIR::FCall>() && Instr->GetNumUsers(CurrentEntryPointIndex) >= 1)
+
+			// Declare a local to hold this instruction's result, e.g. "float _4 = <expression>;"
+			// Not all instructions need a local (e.g. Call, SetMaterialOutput, void-typed ones are excluded — see InstructionNeedsLocal).
+			if (InstructionNeedsLocal(Instr) && !InstrToLocalIndex.Find(Instr))
 			{
 				// Allocate a new local index for this instruction
 				uint32 LocalIndex = NumLocals++;
@@ -369,7 +380,11 @@ struct FEntryPointCodeGen
 				
 				// If this instruction doesn't use a phi value, we'll immediately assign its local to its result so output the "="
 				// E.g., if Instr is a branch, we will set its value inside the generate "if" {} scopes, so no need for a "=" now.
-				if (!InstructionUsesPhiValue(Instr, CurrentEntryPointIndex))
+				if (InstructionUsesPhiValue(Instr))
+				{
+					Printer << EndOfStatement;
+				}
+				else
 				{
 					Printer << TEXTVIEW(" = ");
 				}
@@ -395,54 +410,8 @@ struct FEntryPointCodeGen
 		return NoOp;
 	}
 	
-	ENoOp LowerValue(const MIR::FValue* Value)
-	{
-		// @jason.hoerner this is okay as this is simply emitting the instruction that reads the preshader result.  
-		if (Value->Analysis_PreshaderOffset && Value->HasSubgraphProperties(MIR::EGraphProperties::PreshaderOut))
-		{
-			EmitPreshaderBufferRead(Value);
-			return NoOp;
-		}
-
-		// Instruction results may be shared among other dependent instructions.
-		if (const MIR::FInstruction* Instr = MIR::AsInstruction(Value))
-		{
-			// If this instruction has only one user (dependant) and is foldable, then generate the full HLSL for the instruction in place.
-			if (Instr->GetNumUsers(CurrentEntryPointIndex) <= 1 && IsFoldable(Instr, CurrentEntryPointIndex))
-			{
-				Printer << LowerInstruction(Instr);
-			}
-			else
-			{
-				// Otherwise, this instruction has already been generated before, we emit a reference to the local that stores its result.
-				check(InstrToLocalIndex.Contains(Instr));
-				Printer << TEXTVIEW("_") << InstrToLocalIndex[Instr];
-			}
-			return NoOp;
-		}
-
-		switch (Value->Kind)
-		{
-			case MIR::VK_Constant: LowerConstant(static_cast<const MIR::FConstant*>(Value)); break;
-			case MIR::VK_Builtin: LowerBuiltin(static_cast<const MIR::FBuiltin*>(Value)); break;
-			case MIR::VK_MaterialParameterCollection: LowerMaterialParameterCollection(static_cast<const MIR::FMaterialParameterCollection*>(Value)); break;
-			case MIR::VK_TextureUniform: LowerTextureUniform(static_cast<const MIR::FTextureUniform*>(Value)); break;
-			case MIR::VK_VirtualTextureUniform: LowerVirtualTextureUniform(static_cast<const MIR::FVirtualTextureUniform*>(Value)); break;
-			case MIR::VK_GetVertexInterpolator: LowerGetVertexInterpolator(static_cast<const MIR::FGetVertexInterpolator*>(Value)); break;
-			default: UE_MIR_UNREACHABLE();
-		}
-
-		return NoOp;
-	}
-
 	ENoOp LowerInstruction(const MIR::FInstruction* Instr)
 	{
-		if (Instr->Analysis_PreshaderOffset && Instr->HasSubgraphProperties(MIR::EGraphProperties::PreshaderOut))
-		{
-			EmitPreshaderBufferRead(Instr);
-			return NoOp;
-		}
-
 		switch (Instr->Kind)
 		{
 			case MIR::VK_Nop: LowerNop(static_cast<const MIR::FNop*>(Instr)); break;
@@ -476,6 +445,44 @@ struct FEntryPointCodeGen
 		return NoOp;
 	}
 
+	ENoOp LowerValue(const MIR::FValue* Value)
+	{
+		// @jason.hoerner this is okay as this is simply emitting the instruction that reads the preshader result.  
+		if (Value->Analysis_PreshaderOffset && Value->HasSubgraphProperties(MIR::EGraphProperties::PreshaderOut))
+		{
+			EmitPreshaderBufferRead(Value);
+			return NoOp;
+		}
+
+		// Instruction results may be shared among other dependent instructions. No need to lower it here, only reference its local.
+		if (const MIR::FInstruction* Instr = MIR::AsInstruction(Value))
+		{
+			if (InstructionShouldBeInlined(Instr))
+			{
+				LowerInstruction(Instr);
+			}
+			else
+			{
+				check(InstrToLocalIndex.Contains(Instr));
+				Printer << TEXTVIEW("_") << InstrToLocalIndex[Instr];
+			}
+			return NoOp;
+		}
+
+		switch (Value->Kind)
+		{
+			case MIR::VK_Constant: LowerConstant(static_cast<const MIR::FConstant*>(Value)); break;
+			case MIR::VK_Builtin: LowerBuiltin(static_cast<const MIR::FBuiltin*>(Value)); break;
+			case MIR::VK_MaterialParameterCollection: LowerMaterialParameterCollection(static_cast<const MIR::FMaterialParameterCollection*>(Value)); break;
+			case MIR::VK_TextureUniform: LowerTextureUniform(static_cast<const MIR::FTextureUniform*>(Value)); break;
+			case MIR::VK_VirtualTextureUniform: LowerVirtualTextureUniform(static_cast<const MIR::FVirtualTextureUniform*>(Value)); break;
+			case MIR::VK_GetVertexInterpolator: LowerGetVertexInterpolator(static_cast<const MIR::FGetVertexInterpolator*>(Value)); break;
+			default: UE_MIR_UNREACHABLE();
+		}
+
+		return NoOp;
+	}
+
 	void LowerConstant(const MIR::FConstant* Constant)
 	{
 		TOptional<MIR::FPrimitive> Primitive = Constant->Type.AsPrimitive();
@@ -876,7 +883,7 @@ struct FEntryPointCodeGen
 		}
 		else if (const UMaterialAggregate* Aggregate = Composite->Type.AsAggregate())
 		{
-			LowerAggregateComposite(Composite, Aggregate);
+			LowerAggregateComposite(Composite, Aggregate, InstrToLocalIndex[Composite]);
 		}
 		else
 		{
@@ -1007,22 +1014,19 @@ struct FEntryPointCodeGen
 		}
 	}
 
-	void LowerAggregateComposite(const MIR::FComposite* Composite, const UMaterialAggregate* Aggregate)
+	void LowerAggregateComposite(const MIR::FComposite* Composite, const UMaterialAggregate* Aggregate, int32 TargetLocal)
 	{
-		Printer << EndOfStatement; // end the declaration
-
-		int Local = InstrToLocalIndex[Composite];
-
 		auto Components = Composite->GetComponents();
 		for (int i = 0; i < Components.Num(); ++i)
 		{
-			Printer << TEXT('_') << Local << TEXT('.');
-			PrintUserIdentifier(Printer, Aggregate->Attributes[i].Name.ToString()); 
-			Printer << TEXT(" = ") << LowerValue(Components[i]);
-			if (i < Components.Num() - 1)
+			if (!Components[i])
 			{
-				Printer << EndOfStatement;
+				continue;
 			}
+
+			Printer << TEXT('_') << TargetLocal << TEXT('.');
+			PrintUserIdentifier(Printer, Aggregate->Attributes[i].Name.ToString());
+			Printer << TEXT(" = ") << LowerValue(Components[i]) << EndOfStatement;
 		}
 	}
 
@@ -1255,33 +1259,41 @@ struct FEntryPointCodeGen
 
 	void LowerBranch(const MIR::FBranch* Branch)
 	{
-		if (IsFoldable(Branch, CurrentEntryPointIndex))
+		const int32 Local = InstrToLocalIndex[Branch];
+
+		// Lowers a branch arm's block.
+		auto LowerBranchArm = [this, Local](MIR::FBlock& Block, const MIR::FValue* Arg)
 		{
-			if (Branch->Type.IsDouble())
+			// If the arg is a composite inside the block, have it write directly to the branch's phi local instead of declaring an intermediate.
+			const MIR::FComposite* ArgComposite = Arg->As<MIR::FComposite>();
+			bool bAggregateComposite = ArgComposite && Arg->Type.AsAggregate();
+			bool bOptimized = bAggregateComposite && ArgComposite->Linkage[CurrentEntryPointIndex].Block == &Block;
+			if (bOptimized)
 			{
-				Printer << TEXTVIEW("WSSelect(") << LowerValue(Branch->ConditionArg)
-					<< TEXTVIEW(", ") << LowerValue(Branch->TrueArg)
-					<< TEXTVIEW(", ") << LowerValue(Branch->FalseArg)
-					<< TEXTVIEW(")");
-				return;
+				InstrToLocalIndex.Add(ArgComposite, Local);
 			}
-			else if (Branch->Type.AsPrimitive())
+
+			Printer << LowerBlock(Block);
+
+			if (!bOptimized)
 			{
-				Printer << TEXT('(') << LowerValue(Branch->ConditionArg)
-					<< TEXTVIEW(" ? ") << LowerValue(Branch->TrueArg)
-					<< TEXTVIEW(" : ") << LowerValue(Branch->FalseArg) << TEXT(')');
-				return;
+				if (bAggregateComposite)
+				{
+					LowerAggregateComposite(ArgComposite, ArgComposite->Type.AsAggregate(), Local);
+				}
+				else
+				{
+					Printer << TEXTVIEW("_") << Local << TEXTVIEW(" = ") << LowerValue(Arg) << EndOfStatement;
+				}
 			}
-		}
+		};
 
-		Printer << EndOfStatement;
 		Printer << TEXTVIEW("if (") << LowerValue(Branch->ConditionArg) << TEXTVIEW(")") << OpenBlock;
-		Printer << LowerBlock(Branch->TrueBlock[CurrentEntryPointIndex]);
-		Printer << TEXTVIEW("_") << InstrToLocalIndex[Branch] << TEXTVIEW(" = ") << LowerValue(Branch->TrueArg) << EndOfStatement;
+		LowerBranchArm(Branch->TrueBlock[CurrentEntryPointIndex], Branch->TrueArg);
 		Printer << CloseBlock << NewLine;
+
 		Printer << TEXTVIEW("else") << OpenBlock;
-		Printer << LowerBlock(Branch->FalseBlock[CurrentEntryPointIndex]);
-		Printer << TEXTVIEW("_") << InstrToLocalIndex[Branch] << TEXTVIEW(" = ") << LowerValue(Branch->FalseArg) << EndOfStatement;
+		LowerBranchArm(Branch->FalseBlock[CurrentEntryPointIndex], Branch->FalseArg);
 		Printer << CloseBlock;
 	}
 			
@@ -2096,51 +2108,49 @@ struct FEntryPointCodeGen
 	}
 
 	/* helpers */
-	
-	static bool InstructionUsesPhiValue(const MIR::FInstruction* Instr, int32 EntryPointIndex)
+
+	static bool InstructionShouldBeInlined(const MIR::FInstruction* Instr)
 	{
-		if (auto Branch = Instr->As<MIR::FBranch>())
+		if (Instr->Type.IsPrimitive() && Instr->IsConstant())
 		{
-			return !Branch->Type.IsPrimitive() || Branch->TrueBlock[EntryPointIndex].Instructions || Branch->FalseBlock[EntryPointIndex].Instructions;
+			return true;
 		}
-		if (Instr->Type.AsAggregate() && Instr->As<MIR::FComposite>())
+		if (Instr->Kind == MIR::VK_Extern)
 		{
-			return true;
+			return (static_cast<const MIR::FExtern*>(Instr)->GetInfo().Flags & MIR::EExternFlags::Inline) != MIR::EExternFlags::None;
 		}
-		return false;
+		return Instr->Kind == MIR::VK_Subscript
+			|| Instr->Kind == MIR::VK_StageSwitch;
 	}
 
-	static bool IsFoldable(const MIR::FInstruction* Instr, int32 EntryPointIndex)
+	static bool InstructionNeedsLocal(const MIR::FInstruction* Instr)
 	{
-		// Instructions that use Phi values cannot be folded/inlined.
-		if (InstructionUsesPhiValue(Instr, EntryPointIndex))
+		if (Instr->Type.IsVoid())
 		{
 			return false;
 		}
-
-		// Don't fold instructions that become very long to improve readability
 		switch (Instr->Kind)
 		{
-			case MIR::VK_SetMaterialOutput:
-			case MIR::VK_VTPageTableRead:
 			case MIR::VK_Call:
-			case MIR::VK_SubstrateHorizontalMixing:
-			case MIR::VK_SubstrateVerticalLayering:
-			case MIR::VK_SubstrateCoverageWeight:
-			case MIR::VK_SubstrateAdd:
-			case MIR::VK_SubstrateSelect:
-			case MIR::VK_SubstratePromoteToOperator:
-			case MIR::VK_SubstrateSlab:
-			case MIR::VK_SubstrateShadingModels:
-			case MIR::VK_SubstrateToon:
+			case MIR::VK_SetMaterialOutput:
 				return false;
-			case MIR::VK_Extern:
-				return (static_cast<const MIR::FExtern*>(Instr)->GetInfo().Flags & MIR::EExternFlags::Inline) != MIR::EExternFlags::None;
-
 			default:
 				return true;
 		}
 	}
+	
+	static bool InstructionUsesPhiValue(const MIR::FInstruction* Instr)
+	{
+		switch (Instr->Kind)
+		{
+			case MIR::VK_Branch:
+				return true;
+			case MIR::VK_Composite:
+				return Instr->Type.AsAggregate() != nullptr;
+			default:
+				return false;
+		}
+	}
 };
 
 /////////////////////////////////////////////////////////////////////////////////////////////////////////////////////////////
