[VisualLogger] fixed autoRTFM issue when serializing entries in a trace

#Synced-CL 51945655

diff --git a/Engine/Source/Runtime/Engine/Private/VisualLogger/VisualLoggerTraceDevice.cpp b/Engine/Source/Runtime/Engine/Private/VisualLogger/VisualLoggerTraceDevice.cpp
--- a/Engine/Source/Runtime/Engine/Private/VisualLogger/VisualLoggerTraceDevice.cpp
+++ b/Engine/Source/Runtime/Engine/Private/VisualLogger/VisualLoggerTraceDevice.cpp
@@ -1,4 +1,4 @@
-﻿// Copyright Epic Games, Inc. All Rights Reserved.
+// Copyright Epic Games, Inc. All Rights Reserved.
 #include "VisualLogger/VisualLoggerTraceDevice.h"
 
 #if UE_DEBUG_RECORDING_ENABLED
@@ -81,7 +81,7 @@ void FVisualLoggerTraceDevice::Serialize(const UObject* InLogOwner, const FName&
 		Archive << const_cast<FVisualLogEntry&>(InLogEntry);
 
 		static std::atomic<uint32> EntryID = 0;
-		const uint32 ID = EntryID.fetch_add(1);
+		const uint32 ID = AutoRTFM::Open([&] { return EntryID.fetch_add(1); });
 
 		const uint32 DataSize = static_cast<uint32>(Archive.Num());
 		constexpr uint32 MaxChunkSize = TNumericLimits<uint16>::Max();
