Control Rig: Fix Default Control Rig cannot change to BP-independent rigs
#rnx

#Synced-CL 51944919

diff --git a/Engine/Source/Runtime/Engine/Classes/Engine/SkeletalMesh.h b/Engine/Source/Runtime/Engine/Classes/Engine/SkeletalMesh.h
--- a/Engine/Source/Runtime/Engine/Classes/Engine/SkeletalMesh.h
+++ b/Engine/Source/Runtime/Engine/Classes/Engine/SkeletalMesh.h
@@ -1283,7 +1283,7 @@ public:
 
 	/** The Default Control Rig To Animate with when used in Sequnecer. */
 	UE_DEPRECATED(4.27, "Please do not access this member directly; use USkeletalMesh::GetDefaultAnimatingRig() or USkeletalMesh::SetDefaultAnimatingRig().")
-	UPROPERTY(EditAnywhere, Category = AnimationRig, BlueprintGetter = GetDefaultAnimatingRig, BlueprintSetter = SetDefaultAnimatingRig, meta = (AllowedClasses = "/Script/ControlRigDeveloper.ControlRigBlueprint"))
+	UPROPERTY(EditAnywhere, Category = AnimationRig, BlueprintGetter = GetDefaultAnimatingRig, BlueprintSetter = SetDefaultAnimatingRig, meta = (AllowedClasses = "/Script/ControlRigDeveloper.ControlRigEditorAssetInterface,/Script/ControlRig.ControlRigRuntimeAssetInterface"))
 	TSoftObjectPtr<UObject> DefaultAnimatingRig;
 	static FName GetDefaultAnimatingRigMemberName()
 	{
