Added Instanced Skinned Mesh material flag to enable inline GPU skinning shaders when r.SkinCache.CompileShaders=2.

#Synced-CL 51979294

diff --git a/Engine/Plugins/Editor/MaterialValidation/Source/Private/MaterialValidationLibrary.cpp b/Engine/Plugins/Editor/MaterialValidation/Source/Private/MaterialValidationLibrary.cpp
--- a/Engine/Plugins/Editor/MaterialValidation/Source/Private/MaterialValidationLibrary.cpp
+++ b/Engine/Plugins/Editor/MaterialValidation/Source/Private/MaterialValidationLibrary.cpp
@@ -1530,7 +1530,7 @@ static const FMaterialInstanceStaticPropertyInfo StaticPropertyInfos[(int32)EMat
 };
 
 /** Static array of column FNames for all material usages. */
-static const FName UsageFlagPropertyIds[MATUSAGE_MAX] = {
+static const FName UsageFlagPropertyIds[] = {
 	"PID_UsageSkeletalMesh",
 	"PID_UsageParticleSprites",
 	"PID_UsageBeamTrails",
@@ -1557,7 +1557,9 @@ static const FName UsageFlagPropertyIds[MATUSAGE_MAX] = {
 	"PID_UsageEditorCompositing",
 	"PID_UsageNeuralNetworks",
 	"PID_UsageMeshDeformer",
+	"PID_UsageInstancedSkinnedMesh",
 };
+static_assert(UE_ARRAY_COUNT(UsageFlagPropertyIds) == MATUSAGE_MAX, "UsageFlagPropertyIds must match EMaterialUsage. Did you add a new MATUSAGE entry?");
 
 static int32 GetPropertyCount(FMaterialDatabaseAssetHierarchyInfo const& InInfo)
 {
diff --git a/Engine/Plugins/Editor/MaterialValidation/Source/Private/MaterialValidationToolkit.cpp b/Engine/Plugins/Editor/MaterialValidation/Source/Private/MaterialValidationToolkit.cpp
--- a/Engine/Plugins/Editor/MaterialValidation/Source/Private/MaterialValidationToolkit.cpp
+++ b/Engine/Plugins/Editor/MaterialValidation/Source/Private/MaterialValidationToolkit.cpp
@@ -769,7 +769,7 @@ static TSharedRef<SWidget> GenerateWidgetForStaticProperty(TSharedPtr<FMaterialI
 }
 
 /** Static array of column FNames for all material usages. */
-static const FName UsageColumnIds[MATUSAGE_MAX] = {
+static const FName UsageColumnIds[] = {
 	"CID_UsageSkeletalMesh",
 	"CID_UsageParticleSprites",
 	"CID_UsageBeamTrails",
@@ -796,7 +796,9 @@ static const FName UsageColumnIds[MATUSAGE_MAX] = {
 	"CID_UsageEditorCompositing",
 	"CID_UsageNeuralNetworks",
 	"CID_UsageMeshDeformer",
+	"CID_UsageInstancedSkinnedMesh",
 };
+static_assert(UE_ARRAY_COUNT(UsageColumnIds) == MATUSAGE_MAX, "UsageColumnIds must match EMaterialUsage. Did you add a new MATUSAGE entry?");
 
 /** Static array of column FNames for column type hashes, outside of Default */
 static const FName CategoryHashIds[(int32)(EMaterialPropertyCategory::Count)] = {
diff --git a/Engine/Source/Runtime/Engine/Private/GPUSkinVertexFactory.cpp b/Engine/Source/Runtime/Engine/Private/GPUSkinVertexFactory.cpp
--- a/Engine/Source/Runtime/Engine/Private/GPUSkinVertexFactory.cpp
+++ b/Engine/Source/Runtime/Engine/Private/GPUSkinVertexFactory.cpp
@@ -803,18 +803,45 @@ TGPUSkinVertexFactory
 template <GPUSkinBoneInfluenceType BoneInfluenceType>
 bool TGPUSkinVertexFactory<BoneInfluenceType>::ShouldCompilePermutation(const FVertexFactoryShaderPermutationParameters& Parameters)
 {
-	static FShaderPlatformCachedIniValue<int32> UBICVar(TEXT("r.GPUSkin.UnlimitedBoneInfluences"));
-	const bool bUseUBI = UBICVar.Get(Parameters.Platform) != 0;
+	// Always compile for the default/special engine materials.
+	if (Parameters.MaterialParameters.bIsSpecialEngineMaterial)
+	{
+		return true;
+	}
 
-	static FShaderPlatformCachedIniValue<bool> UseDeformerForUBICVar(TEXT("r.GPUSkin.AlwaysUseDeformerForUnlimitedBoneInfluences"));
-	const bool bUseDeformerForUBI = UseDeformerForUBICVar.Get(Parameters.Platform);
-		
-	// Compile the shader for UBI if UBI is enabled and we're not forcing the use of a deformer for all UBI meshes
-	const bool bUnlimitedBoneInfluences = BoneInfluenceType == UnlimitedBoneInfluence && bUseUBI && !bUseDeformerForUBI;
+	// Material must be flagged for skeletal mesh or morph target use.
+	if (!Parameters.MaterialParameters.bIsUsedWithSkeletalMesh && !Parameters.MaterialParameters.bIsUsedWithMorphTargets)
+	{
+		return false;
+	}
+
+	// For the unlimited bone influence variant, only compile if UBI is actually enabled
+	// and we're not forcing all UBI meshes through a mesh deformer.
+	if (BoneInfluenceType == UnlimitedBoneInfluence)
+	{
+		static FShaderPlatformCachedIniValue<int32> UBICVar(TEXT("r.GPUSkin.UnlimitedBoneInfluences"));
+		static FShaderPlatformCachedIniValue<bool> UseDeformerForUBICVar(TEXT("r.GPUSkin.AlwaysUseDeformerForUnlimitedBoneInfluences"));
+
+		if (!UBICVar.Get(Parameters.Platform) || UseDeformerForUBICVar.Get(Parameters.Platform))
+		{
+			return false;
+		}
+	}
 
-	return AreInlineSkinningShadersEnabled(Parameters.Platform) &&
-		  (((Parameters.MaterialParameters.bIsUsedWithSkeletalMesh || Parameters.MaterialParameters.bIsUsedWithMorphTargets) && (BoneInfluenceType != UnlimitedBoneInfluence || bUnlimitedBoneInfluences)) 
-			  || Parameters.MaterialParameters.bIsSpecialEngineMaterial);
+	// Inline skinning shaders are compiled when the skin cache doesn't fully replace them.
+	if (AreInlineSkinningShadersEnabled(Parameters.Platform))
+	{
+		return true;
+	}
+
+	// Instanced skinned meshes use inline skinning even when the skin cache is active,
+	// so compile these shaders regardless of the r.SkinCache.CompileShaders setting.
+	if (Parameters.MaterialParameters.bIsUsedWithInstancedSkinnedMesh)
+	{
+		return true;
+	}
+
+	return false;
 }
 
 template <GPUSkinBoneInfluenceType BoneInfluenceType>
@@ -1300,9 +1327,18 @@ void TGPUSkinAPEXClothVertexFactory<BoneInfluenceType>::ModifyCompilationEnviron
 template <GPUSkinBoneInfluenceType BoneInfluenceType>
 bool TGPUSkinAPEXClothVertexFactory<BoneInfluenceType>::ShouldCompilePermutation(const FVertexFactoryShaderPermutationParameters& Parameters)
 {
-	return IsClothEnabled(Parameters.Platform)
-		&& (Parameters.MaterialParameters.bIsUsedWithAPEXCloth || Parameters.MaterialParameters.bIsSpecialEngineMaterial)
-		&& Super::ShouldCompilePermutation(Parameters);
+	if (!IsClothEnabled(Parameters.Platform))
+	{
+		return false;
+	}
+
+	// Material must be flagged for cloth use (or be a special engine material).
+	if (!Parameters.MaterialParameters.bIsUsedWithAPEXCloth && !Parameters.MaterialParameters.bIsSpecialEngineMaterial)
+	{
+		return false;
+	}
+
+	return Super::ShouldCompilePermutation(Parameters);
 }
 
 template <GPUSkinBoneInfluenceType BoneInfluenceType>
diff --git a/Engine/Source/Runtime/Engine/Private/Materials/Material.cpp b/Engine/Source/Runtime/Engine/Private/Materials/Material.cpp
--- a/Engine/Source/Runtime/Engine/Private/Materials/Material.cpp
+++ b/Engine/Source/Runtime/Engine/Private/Materials/Material.cpp
@@ -1534,6 +1534,7 @@ uint32 UMaterial::GetUsageFlags() const
 	UsageFlags |= bUsedWithEditorCompositing ? (1u << MATUSAGE_EditorCompositing) : 0;
 	UsageFlags |= bUsedWithNeuralNetworks ? (1u << MATUSAGE_NeuralNetworks) : 0;
 	UsageFlags |= bUsedWithMeshDeformer ? (1u << MATUSAGE_MeshDeformer) : 0;
+	UsageFlags |= bUsedWithInstancedSkinnedMesh ? (1u << MATUSAGE_InstancedSkinnedMesh) : 0;
 	PRAGMA_ENABLE_DEPRECATION_WARNINGS
 	return UsageFlags;
 }
@@ -1570,6 +1571,7 @@ bool UMaterial::GetUsageByFlag(EMaterialUsage Usage) const
 		case MATUSAGE_EditorCompositing: UsageValue = bUsedWithEditorCompositing; break;
 		case MATUSAGE_NeuralNetworks: UsageValue = bUsedWithNeuralNetworks; break;
 		case MATUSAGE_MeshDeformer: UsageValue = bUsedWithMeshDeformer; break;
+		case MATUSAGE_InstancedSkinnedMesh: UsageValue = bUsedWithInstancedSkinnedMesh; break;
 		default: UE_LOG(LogMaterial, Fatal,TEXT("Unknown material usage: %u"), (int32)Usage);
 	};
 	PRAGMA_ENABLE_DEPRECATION_WARNINGS
@@ -1849,6 +1851,10 @@ void UMaterial::SetUsageByFlag(EMaterialUsage Usage, bool NewValue)
 		{
 			bUsedWithMeshDeformer = NewValue; break;
 		}
+		case MATUSAGE_InstancedSkinnedMesh:
+		{
+			bUsedWithInstancedSkinnedMesh = NewValue; break;
+		}
 		default: UE_LOG(LogMaterial, Fatal, TEXT("Unknown material usage: %u"), (int32)Usage);
 	};
 	PRAGMA_ENABLE_DEPRECATION_WARNINGS
diff --git a/Engine/Source/Runtime/Engine/Private/Materials/MaterialInterface.cpp b/Engine/Source/Runtime/Engine/Private/Materials/MaterialInterface.cpp
--- a/Engine/Source/Runtime/Engine/Private/Materials/MaterialInterface.cpp
+++ b/Engine/Source/Runtime/Engine/Private/Materials/MaterialInterface.cpp
@@ -2786,6 +2786,7 @@ FString UMaterialInterface::GetUsageName(EMaterialUsage Usage)
 	case MATUSAGE_EditorCompositing: UsageName = TEXT("EditorCompositing"); break;
 	case MATUSAGE_NeuralNetworks: UsageName = TEXT("NeuralNetworks"); break;
 	case MATUSAGE_MeshDeformer: UsageName = TEXT("MeshDeformer"); break;
+	case MATUSAGE_InstancedSkinnedMesh: UsageName = TEXT("InstancedSkinnedMesh"); break;
 	default: UE_LOG(LogMaterial, Fatal, TEXT("Unknown material usage: %u"), (int32)Usage);
 	};
 	return UsageName;
diff --git a/Engine/Source/Runtime/Engine/Private/Materials/MaterialShared.cpp b/Engine/Source/Runtime/Engine/Private/Materials/MaterialShared.cpp
--- a/Engine/Source/Runtime/Engine/Private/Materials/MaterialShared.cpp
+++ b/Engine/Source/Runtime/Engine/Private/Materials/MaterialShared.cpp
@@ -1860,6 +1860,11 @@ bool FMaterialResource::IsUsedWithMeshDeformer() const
 	return MaterialInstance ? MaterialInstance->GetUsageByFlag(MATUSAGE_MeshDeformer) : Material->GetUsageByFlag(MATUSAGE_MeshDeformer);
 }
 
+bool FMaterialResource::IsUsedWithInstancedSkinnedMesh() const
+{
+	return MaterialInstance ? MaterialInstance->GetUsageByFlag(MATUSAGE_InstancedSkinnedMesh) : Material->GetUsageByFlag(MATUSAGE_InstancedSkinnedMesh);
+}
+
 bool FMaterialResource::IsUsedWithLandscape() const
 {
 	return false;
@@ -6028,6 +6033,7 @@ FMaterialShaderParametersBase::FMaterialShaderParametersBase(const FMaterial* In
 	bSamplesMaterialCache = InMaterial->SamplesMaterialCache();
 	bHasMaterialCacheOutput = InMaterial->HasMaterialCacheOutput();
 	bIsUsedWithMeshDeformer = InMaterial->IsUsedWithMeshDeformer();
+	bIsUsedWithInstancedSkinnedMesh = InMaterial->IsUsedWithInstancedSkinnedMesh();
 }
 
 FMaterialShaderParameters::FMaterialShaderParameters(const FMaterial* InMaterial)
diff --git a/Engine/Source/Runtime/Engine/Private/SkeletalMeshSceneProxy.cpp b/Engine/Source/Runtime/Engine/Private/SkeletalMeshSceneProxy.cpp
--- a/Engine/Source/Runtime/Engine/Private/SkeletalMeshSceneProxy.cpp
+++ b/Engine/Source/Runtime/Engine/Private/SkeletalMeshSceneProxy.cpp
@@ -225,6 +225,13 @@ FSkeletalMeshSceneProxy::FSkeletalMeshSceneProxy(const FSkinnedMeshSceneProxyDes
 			bAnySectionHasSkelMeshUsage |= RequiredUsage == MATUSAGE_SkeletalMesh;
 			bool bValidUsage = Material && Material->CheckMaterialUsage_Concurrent(RequiredUsage);
 
+			// Instanced skinned meshes use inline GPU skinning and need their own usage flag
+			// to ensure shader permutations are compiled even when r.SkinCache.CompileShaders=2.
+			if (bInstancedSkinnedMesh)
+			{
+				bValidUsage &= Material && Material->CheckMaterialUsage_Concurrent(MATUSAGE_InstancedSkinnedMesh);
+			}
+
 			// if this is a clothing section, then enabled and will be drawn but the corresponding original section should be disabled
 			const bool bClothSection = Section.HasClothingData();
 			if (bClothSection)
diff --git a/Engine/Source/Runtime/Engine/Public/MaterialShared.h b/Engine/Source/Runtime/Engine/Public/MaterialShared.h
--- a/Engine/Source/Runtime/Engine/Public/MaterialShared.h
+++ b/Engine/Source/Runtime/Engine/Public/MaterialShared.h
@@ -1150,6 +1150,7 @@ struct FMaterialShaderParametersBase
 			uint16 bSamplesMaterialCache : 1;
 			uint16 bHasMaterialCacheOutput : 1;
 			uint16 bIsUsedWithMeshDeformer : 1;
+			uint16 bIsUsedWithInstancedSkinnedMesh : 1;
 		};
 	};
 	EBlendMode BlendMode;
@@ -2459,6 +2460,7 @@ public:
 	virtual bool IsUsedWithVirtualHeightfieldMesh() const { return false; }
 	virtual bool IsUsedWithNeuralNetworks() const { return false; }
 	virtual bool IsUsedWithMeshDeformer() const { return false; }
+	virtual bool IsUsedWithInstancedSkinnedMesh() const { return false; }
 	virtual bool IsUsedWithStaticMesh() const { return true; }
 	virtual bool HasAnyUsage() const { return true; }
 	virtual bool IsFullyRough() const { return false; }
@@ -3283,6 +3285,7 @@ public:
 	ENGINE_API virtual bool IsUsedWithVirtualHeightfieldMesh() const override;
 	ENGINE_API virtual bool IsUsedWithNeuralNetworks() const override;
 	ENGINE_API virtual bool IsUsedWithMeshDeformer() const override;
+	ENGINE_API virtual bool IsUsedWithInstancedSkinnedMesh() const override;
 	ENGINE_API virtual bool IsUsedWithNanite() const override;
 	ENGINE_API virtual bool IsUsedWithVoxels() const override;
 	ENGINE_API virtual bool IsUsedWithVolumetricCloud() const override;
diff --git a/Engine/Source/Runtime/Engine/Public/Materials/Material.h b/Engine/Source/Runtime/Engine/Public/Materials/Material.h
--- a/Engine/Source/Runtime/Engine/Public/Materials/Material.h
+++ b/Engine/Source/Runtime/Engine/Public/Materials/Material.h
@@ -886,6 +886,14 @@ public:
 	UPROPERTY(EditAnywhere, BlueprintReadOnly, Category = Usage)
 	uint32 bUsedWithHeterogeneousVolumes : 1;
 
+	/**
+	 * Indicates that the material and its instances can be used with instanced skinned meshes.
+	 * This will result in the shaders required to support instanced skinned meshes being compiled which will increase shader compile time and memory usage.
+	 */
+	UE_DEPRECATED(5.8, "Always use the GetUsageByFlag() accessor.")
+	UPROPERTY(EditAnywhere, BlueprintReadOnly, Category = Usage)
+	uint32 bUsedWithInstancedSkinnedMesh : 1;
+
 	/**
 	 * Indicates that the material and its instances can be used on StaticMeshes, which is true by default on all Materials
 	 * This can be used to remove shaders which are not required, reducing shader compile time and memory usage
diff --git a/Engine/Source/Runtime/Engine/Public/Materials/MaterialInterface.h b/Engine/Source/Runtime/Engine/Public/Materials/MaterialInterface.h
--- a/Engine/Source/Runtime/Engine/Public/Materials/MaterialInterface.h
+++ b/Engine/Source/Runtime/Engine/Public/Materials/MaterialInterface.h
@@ -125,6 +125,8 @@ enum EMaterialUsage : int
 	MATUSAGE_NeuralNetworks,
 	/** Indicates that the material and its instances can be used with mesh deformers. */
 	MATUSAGE_MeshDeformer,
+	/** Indicates that the material and its instances can be used with instanced skinned meshes. */
+	MATUSAGE_InstancedSkinnedMesh,
 
 	MATUSAGE_MAX,
 };
