Another tentative fix for UE-295545

#rnx
#tests Could not reproduce.
#Synced-CL 52005599

diff --git a/Engine/Source/Runtime/Engine/Private/Materials/HLSLMaterialTranslator.cpp b/Engine/Source/Runtime/Engine/Private/Materials/HLSLMaterialTranslator.cpp
--- a/Engine/Source/Runtime/Engine/Private/Materials/HLSLMaterialTranslator.cpp
+++ b/Engine/Source/Runtime/Engine/Private/Materials/HLSLMaterialTranslator.cpp
@@ -13124,6 +13124,16 @@ int32 FHLSLMaterialTranslator::SubstrateConversionFromLegacy(
 	int32 CustomTangent_Tangent,
 	FSubstrateOperator* PromoteToOperator)
 {
+	if (Normal == INDEX_NONE)
+	{
+		Errorf(TEXT("SubstrateConversionFromLegacy requires a valid normal input in Material %s (asset: %s).\r\n"), *Material->GetDebugName(), *Material->GetAssetPath().ToString());
+		return INDEX_NONE;
+	}
+	if (ClearCoat_Normal == INDEX_NONE)
+	{
+		Errorf(TEXT("SubstrateConversionFromLegacy requires a valid clear coat normal input in Material %s (asset: %s).\r\n"), *Material->GetDebugName(), *Material->GetAssetPath().ToString());
+		return INDEX_NONE;
+	}
 	const FString NormalCode = GetParameterCode(Normal);
 	const FString TangentCode = Tangent != INDEX_NONE ? *GetParameterCode(Tangent) : TEXT("NONE");
 	const int32 RawTangent = Tangent != INDEX_NONE ? Tangent : Normal;
@@ -13149,6 +13159,12 @@ int32 FHLSLMaterialTranslator::SubstrateConversionFromLegacy(
 		}
 	}
 
+	if (InShadingModel == INDEX_NONE)
+	{
+		Errorf(TEXT("SubstrateConversionFromLegacy requires a valid shading model input in Material %s (asset: %s).\r\n"), *Material->GetDebugName(), *Material->GetAssetPath().ToString());
+		return INDEX_NONE;
+	}
+
 	if (PromoteToOperator)
 	{
 		if (PromoteToOperator->Index == INDEX_NONE || PromoteToOperator->BSDFIndex == INDEX_NONE)
diff --git a/Engine/Source/Runtime/Engine/Private/Materials/MaterialExpressionSubstrate.cpp b/Engine/Source/Runtime/Engine/Private/Materials/MaterialExpressionSubstrate.cpp
--- a/Engine/Source/Runtime/Engine/Private/Materials/MaterialExpressionSubstrate.cpp
+++ b/Engine/Source/Runtime/Engine/Private/Materials/MaterialExpressionSubstrate.cpp
@@ -3962,19 +3962,27 @@ int32 UMaterialExpressionSubstrateConvertMaterialAttributes::CompileCommon(class
 	// Regular normal basis
 	int32 NormalCodeChunk = MaterialAttributes.CompileWithDefault(Compiler, FMaterialAttributeDefinitionMap::GetID(MP_Normal));
 	NormalCodeChunk = Compiler->ForceCast(NormalCodeChunk, MCT_Float3, MFCF_ExactMatch | MFCF_ReplicateValue);
-	NormalCodeChunk = Compiler->TransformNormalFromRequestedBasisToWorld(NormalCodeChunk);
 
 	// When computing NormalCodeChunk, we invoke TransformNormalFromRequestedBasisToWorld which requires input to be float or float3.
 	// Certain material do not respect this requirement. We handle here a simple recovery when source material doesn't have a valid 
 	// normal (e.g., vec2 normal), and avoid crashing the material compilation. The error will still be reported by the compiler up 
 	// to the user, but the compilation will succeed.
-	if (NormalCodeChunk == INDEX_NONE) { NormalCodeChunk = Compiler->VertexNormal(); }
+	// This is done before TransformNormalFromRequestedBasisToWorld because TransformNormalFromRequestedBasisToWorld(INDEX_NONE) would crash inside GetParameterCode.
+	if (NormalCodeChunk == INDEX_NONE) 
+	{
+		NormalCodeChunk = Compiler->VertexNormal();
+	}
+	NormalCodeChunk = Compiler->TransformNormalFromRequestedBasisToWorld(NormalCodeChunk);
 
 	int32 TangentCodeChunk = INDEX_NONE;
 	if (bHasAnisotropy)
 	{
 		TangentCodeChunk = MaterialAttributes.CompileWithDefault(Compiler, FMaterialAttributeDefinitionMap::GetID(MP_Tangent));
 		TangentCodeChunk = Compiler->ForceCast(TangentCodeChunk, MCT_Float3, MFCF_ExactMatch | MFCF_ReplicateValue);
+		if (TangentCodeChunk == INDEX_NONE)
+		{
+			TangentCodeChunk = Compiler->VertexTangent(); // Similar fix as for NormalCodeChunk above.
+		}
 		TangentCodeChunk = Compiler->TransformNormalFromRequestedBasisToWorld(TangentCodeChunk);
 	}
 	const FSubstrateRegisteredSharedLocalBasis NewRegisteredSharedLocalBasis = SubstrateCompilationInfoCreateSharedLocalBasis(Compiler, NormalCodeChunk, TangentCodeChunk);
