[WorldPartition] Fix race conditions when using wp.Runtime.UpdateStreaming.EnableAsyncUpdate
- Fix 1: GetStreamingPerformance called GetStreamingStatus on a worker thread, which reads ULevelStreaming::LoadedLevel without synchronization.
Added UWorldPartitionRuntimeCell::IsVisible() that uses GetLevelStreamingState() instead, which is protected by CurrentStateLock.
- Fix 2: GetCellEffectiveWantedState caches EffectiveWantedState and EffectiveWantedStateEpoch as mutable fields without synchronization.
These are written from the async worker thread and read from GameThread via IsStreamingCompleted (which can be called from any gameplay context)
Uses FTransactionallySafeRWLock to protect both fields (rather than std::atomic to remain compatible with AutoRTFM transactional contexts)

#Synced-CL 51624647

diff --git a/Engine/Source/Runtime/Engine/Private/WorldPartition/WorldPartitionRuntimeCell.cpp b/Engine/Source/Runtime/Engine/Private/WorldPartition/WorldPartitionRuntimeCell.cpp
--- a/Engine/Source/Runtime/Engine/Private/WorldPartition/WorldPartitionRuntimeCell.cpp
+++ b/Engine/Source/Runtime/Engine/Private/WorldPartition/WorldPartitionRuntimeCell.cpp
@@ -148,18 +148,24 @@ EDataLayerRuntimeState UWorldPartitionRuntimeCell::GetCellEffectiveWantedState(c
 {
 	if (!HasDataLayers())
 	{
-		EffectiveWantedState = EDataLayerRuntimeState::Activated;
+		UE::TWriteScopeLock WriteLock(EffectiveWantedStateLock);
+		return EffectiveWantedState = EDataLayerRuntimeState::Activated;
 	}
-	else
+
+	const int32 ResolvingDataLayersRuntimeStateEpoch = Context.GetResolvingDataLayersRuntimeStateEpoch();
 	{
-		const int32 ResolvingDataLayersRuntimeStateEpoch = Context.GetResolvingDataLayersRuntimeStateEpoch();
-		if (EffectiveWantedStateEpoch != ResolvingDataLayersRuntimeStateEpoch)
+		UE::TReadScopeLock ReadLock(EffectiveWantedStateLock);
+		if (EffectiveWantedStateEpoch == ResolvingDataLayersRuntimeStateEpoch)
 		{
-			EffectiveWantedState = Context.ResolveDataLayerRuntimeState(DataLayers);
-			EffectiveWantedStateEpoch = ResolvingDataLayersRuntimeStateEpoch;
+			return EffectiveWantedState;
 		}
 	}
 
+	// Epoch mismatch - resolve outside lock, then cache under write lock
+	const EDataLayerRuntimeState ResolvedState = Context.ResolveDataLayerRuntimeState(DataLayers);
+	UE::TWriteScopeLock WriteLock(EffectiveWantedStateLock);
+	EffectiveWantedState = ResolvedState;
+	EffectiveWantedStateEpoch = ResolvingDataLayersRuntimeStateEpoch;
 	return EffectiveWantedState;
 }
 
diff --git a/Engine/Source/Runtime/Engine/Private/WorldPartition/WorldPartitionRuntimeHash.cpp b/Engine/Source/Runtime/Engine/Private/WorldPartition/WorldPartitionRuntimeHash.cpp
--- a/Engine/Source/Runtime/Engine/Private/WorldPartition/WorldPartitionRuntimeHash.cpp
+++ b/Engine/Source/Runtime/Engine/Private/WorldPartition/WorldPartitionRuntimeHash.cpp
@@ -631,7 +631,7 @@ EWorldPartitionStreamingPerformance UWorldPartitionRuntimeHash::GetStreamingPerf
 	{
 		for (const UWorldPartitionRuntimeCell* Cell : CellsToActivate)
 		{
-			if (!Cell->IsAlwaysLoaded() && Cell->GetStreamingStatus() != LEVEL_Visible)
+			if (!Cell->IsAlwaysLoaded() && !Cell->IsVisible())
 			{
 				bool bOutShouldBlockCell;
 				EWorldPartitionStreamingPerformance CellPerformance = GetStreamingPerformanceForCell(Cell, bOutShouldBlockCell);
diff --git a/Engine/Source/Runtime/Engine/Private/WorldPartition/WorldPartitionRuntimeLevelStreamingCell.cpp b/Engine/Source/Runtime/Engine/Private/WorldPartition/WorldPartitionRuntimeLevelStreamingCell.cpp
--- a/Engine/Source/Runtime/Engine/Private/WorldPartition/WorldPartitionRuntimeLevelStreamingCell.cpp
+++ b/Engine/Source/Runtime/Engine/Private/WorldPartition/WorldPartitionRuntimeLevelStreamingCell.cpp
@@ -250,6 +250,17 @@ EStreamingStatus UWorldPartitionRuntimeLevelStreamingCell::GetStreamingStatus()
 	return Super::GetStreamingStatus();
 }
 
+bool UWorldPartitionRuntimeLevelStreamingCell::IsVisible() const
+{
+	if (LevelStreaming)
+	{
+		// Use GetLevelStreamingState() which is protected by CurrentStateLock,
+		// unlike GetLevelStreamingStatus() which reads LoadedLevel without synchronization.
+		return LevelStreaming->GetLevelStreamingState() == ELevelStreamingState::LoadedVisible;
+	}
+	return false;
+}
+
 FLinearColor UWorldPartitionRuntimeLevelStreamingCell::GetDebugColor(EWorldPartitionRuntimeCellVisualizeMode VisualizeMode) const
 {
 #if !UE_BUILD_SHIPPING
diff --git a/Engine/Source/Runtime/Engine/Public/WorldPartition/WorldPartitionRuntimeCell.h b/Engine/Source/Runtime/Engine/Public/WorldPartition/WorldPartitionRuntimeCell.h
--- a/Engine/Source/Runtime/Engine/Public/WorldPartition/WorldPartitionRuntimeCell.h
+++ b/Engine/Source/Runtime/Engine/Public/WorldPartition/WorldPartitionRuntimeCell.h
@@ -12,6 +12,7 @@
 #include "WorldPartition/DataLayer/DataLayerInstanceNames.h"
 #include "ProfilingDebugging/ProfilingHelpers.h"
 #include "Misc/HierarchicalLogArchive.h"
+#include "Misc/TransactionallySafeRWLock.h"
 #include "Algo/AnyOf.h"
 #include "UObject/ObjectMacros.h"
 #include "WorldPartitionRuntimeCell.generated.h"
@@ -262,6 +263,8 @@ class UWorldPartitionRuntimeCell : public UObject, public IWorldPartitionCell, p
 	virtual void SetIsSpatiallyLoaded(bool bInIsSpatiallyLoaded) { bIsSpatiallyLoaded = bInIsSpatiallyLoaded; }
 	virtual void SetStreamingPriority(int32 InStreamingPriority) const PURE_VIRTUAL(UWorldPartitionRuntimeCell::SetStreamingPriority,);
 	virtual EStreamingStatus GetStreamingStatus() const { return LEVEL_Unloaded; }
+	/** Thread-safe. Returns whether this cell's level is visible (added to world and rendering). */
+	virtual bool IsVisible() const { return false; }
 	void SetClientOnlyVisible(bool bInClientOnlyVisible) { bClientOnlyVisible = bInClientOnlyVisible; }
 	bool GetClientOnlyVisible() const { return bClientOnlyVisible; }
 	virtual FGuid const& GetContentBundleID() const { return ContentBundleID; }
@@ -286,6 +289,7 @@ class UWorldPartitionRuntimeCell : public UObject, public IWorldPartitionCell, p
 	//~End IWorldPartitionCell Interface
 
 	ENGINE_API UDataLayerManager* GetDataLayerManager() const;
+	/** Thread-safe resolve and cache of the effective wanted state for this cell's data layers. */
 	ENGINE_API EDataLayerRuntimeState GetCellEffectiveWantedState(const FWorldPartitionStreamingContext& Context) const;
 	inline bool HasDataLayers() const { return !DataLayers.IsEmpty(); }
 
@@ -358,7 +362,13 @@ protected:
 	//          The same should be done for injected external streaming objects.
 	virtual bool ShouldServerWaitForClientLevelVisibility() const { return true; }
 
-	EDataLayerRuntimeState GetCellEffectiveWantedStateRaw() const { return !HasDataLayers() ? EDataLayerRuntimeState::Activated : EffectiveWantedState; }
+	/** Thread-safe read of the cached effective wanted state (without resolving). */
+	EDataLayerRuntimeState GetCellEffectiveWantedStateRaw() const
+	{
+		if (!HasDataLayers()) { return EDataLayerRuntimeState::Activated; }
+		UE::TReadScopeLock ReadLock(EffectiveWantedStateLock);
+		return EffectiveWantedState;
+	}
 
 	UPROPERTY()
 	bool bIsAlwaysLoaded;
@@ -410,6 +420,7 @@ protected:
 	TObjectPtr<const UExternalDataLayerAsset> ExternalDataLayerAsset;
 #endif
 
+	mutable FTransactionallySafeRWLock EffectiveWantedStateLock;
 	mutable EDataLayerRuntimeState EffectiveWantedState;
 	mutable int32 EffectiveWantedStateEpoch;
 
diff --git a/Engine/Source/Runtime/Engine/Public/WorldPartition/WorldPartitionRuntimeLevelStreamingCell.h b/Engine/Source/Runtime/Engine/Public/WorldPartition/WorldPartitionRuntimeLevelStreamingCell.h
--- a/Engine/Source/Runtime/Engine/Public/WorldPartition/WorldPartitionRuntimeLevelStreamingCell.h
+++ b/Engine/Source/Runtime/Engine/Public/WorldPartition/WorldPartitionRuntimeLevelStreamingCell.h
@@ -28,6 +28,7 @@ class UWorldPartitionRuntimeLevelStreamingCell : public UWorldPartitionRuntimeCe
 	ENGINE_API virtual FLinearColor GetDebugColor(EWorldPartitionRuntimeCellVisualizeMode VisualizeMode) const override;
 	ENGINE_API virtual void SetIsAlwaysLoaded(bool bInIsAlwaysLoaded) override;
 	ENGINE_API virtual EStreamingStatus GetStreamingStatus() const override;
+	ENGINE_API virtual bool IsVisible() const override;
 	//~End UWorldPartitionRuntimeCell interface
 
 	//~Begin IWorldPartitionCell interface
