[WorldPartition] Fix ~FDeferred not handling actor removal during FlushAsyncLoading, which can happen with the ALT enabled

#Synced-CL 51980806

diff --git a/Engine/Source/Runtime/Engine/Private/WorldPartition/ActorDescContainerInstance.cpp b/Engine/Source/Runtime/Engine/Private/WorldPartition/ActorDescContainerInstance.cpp
--- a/Engine/Source/Runtime/Engine/Private/WorldPartition/ActorDescContainerInstance.cpp
+++ b/Engine/Source/Runtime/Engine/Private/WorldPartition/ActorDescContainerInstance.cpp
@@ -202,6 +202,10 @@ void UActorDescContainerInstance::OnUnregisterChildContainerInstance(const FGuid
 
 void UActorDescContainerInstance::Uninitialize()
 {
+	// Ensure that a non-deferred context is currently being used, because otherwise the actor desc instance and container deletion resulting from this call will risk the context
+	// use deleted memory. If this check fails, a fix would be to flush the current FWorldPartitionLoadingContext::IContext through a new static method, after having forced unregistrations below.
+	check(!FWorldPartitionLoadingContext::IsDeferringRawObjects());
+
 	const UWorld* OuterWorld = GetTypedOuter<UWorld>();
 
 	bIsInitialized = false;
@@ -223,7 +227,7 @@ void UActorDescContainerInstance::Uninitialize()
 					// Force actor unregistration when destroying the container instance without the world being destroyed to avoid dangling actors
 					FWorldPartitionReference LastReference(ActorDescInstancePtr.Get());
 				}
-			}			
+			}
 
 			RemoveActorDescInstance(&ActorDescInstancePtr);
 		}
@@ -477,6 +481,10 @@ void UActorDescContainerInstance::RemoveActorDescInstance(TUniquePtr<FWorldParti
 		check(!ChildContainerInstances.Contains(ActorDescInstance->GetGuid()));
 	}
 
+	// This actor desc might have bIsRegisteringOrUnregistering set to true, meaning it is currently stored in the Registrations array in a FDeferred loading context.
+	// RemoveActorDescriptor would delete it without removing it from the array, leading to dangling pointer access.
+	FWorldPartitionLoadingContext::CancelOperationIfAny(ActorDescInstance);
+
 	RemoveActorDescriptor(ActorDescInstance);
 
 	InActorDescInstance->Get()->Invalidate();
diff --git a/Engine/Source/Runtime/Engine/Private/WorldPartition/WorldPartitionActorDescInstance.cpp b/Engine/Source/Runtime/Engine/Private/WorldPartition/WorldPartitionActorDescInstance.cpp
--- a/Engine/Source/Runtime/Engine/Private/WorldPartition/WorldPartitionActorDescInstance.cpp
+++ b/Engine/Source/Runtime/Engine/Private/WorldPartition/WorldPartitionActorDescInstance.cpp
@@ -305,8 +305,6 @@ void FWorldPartitionActorDescInstance::AsyncLoadFinished(const FName& ActorPacka
 
 void FWorldPartitionActorDescInstance::MarkUnload()
 {
-	FlushAsyncLoad();
-
 	if (AActor* Actor = GetActor())
 	{
 		// At this point, it can happen that an actor isn't in an external package:
diff --git a/Engine/Source/Runtime/Engine/Private/WorldPartition/WorldPartitionHandle.cpp b/Engine/Source/Runtime/Engine/Private/WorldPartition/WorldPartitionHandle.cpp
--- a/Engine/Source/Runtime/Engine/Private/WorldPartition/WorldPartitionHandle.cpp
+++ b/Engine/Source/Runtime/Engine/Private/WorldPartition/WorldPartitionHandle.cpp
@@ -56,6 +56,21 @@ AActor* FWorldPartitionImplBase::GetActor(FWorldPartitionActorDescInstance* InAc
 FWorldPartitionLoadingContext::FImmediate FWorldPartitionLoadingContext::DefaultContext;
 FWorldPartitionLoadingContext::IContext* FWorldPartitionLoadingContext::ActiveContext = &DefaultContext;
 
+bool FWorldPartitionLoadingContext::IsDeferringRawObjects()
+{
+	return ActiveContext->IsDeferringRawObjects();
+}
+
+void FWorldPartitionLoadingContext::CancelOperationIfAny(FWorldPartitionActorDescInstance* InActorDescInstance)
+{
+	if (!InActorDescInstance->bIsRegisteringOrUnregistering)
+	{
+		return;
+	}
+
+	ActiveContext->CancelOperation(InActorDescInstance);
+}
+
 void FWorldPartitionLoadingContext::LoadAndRegisterActor(FWorldPartitionActorDescInstance* InActorDescInstance)
 {
 	check(!InActorDescInstance->bIsRegisteringOrUnregistering);
@@ -145,6 +160,10 @@ void FWorldPartitionLoadingContext::FImmediate::UnregisterActor(FWorldPartitionA
 */
 FWorldPartitionLoadingContext::FDeferred::~FDeferred()
 {
+	// There shouldn't be any cancellations when starting, as they should only get generated when ContainerInstanceOps does not contain the Op
+	// because the container is moved on the stack here. In other words, ContainerInstanceCancellations can only grow from within this destructor.
+	check(ContainerInstanceCancellations.IsEmpty()); 									 
+
 	TMap<UActorDescContainerInstance*, FContainerInstanceOps> LocalContainerInstanceOps = MoveTemp(ContainerInstanceOps);
 
 	while (LocalContainerInstanceOps.Num())
@@ -154,7 +173,7 @@ FWorldPartitionLoadingContext::FDeferred::~FDeferred()
 			const FTransform& ContainerTransform = Container->GetTransform();
 			const FTransform* ContainerTransformPtr = ContainerTransform.Equals(FTransform::Identity) ? nullptr : &ContainerTransform;
 
-			auto CreateActorList = [](TArray<AActor*>& ActorList, const TSet<FWorldPartitionActorDescInstance*>& SourceList) -> ULevel*
+			auto CreateActorList = [this](TArray<AActor*>& ActorList, TSet<FWorldPartitionActorDescInstance*>& SourceList) -> ULevel*
 			{
 				ActorList.Empty(SourceList.Num());
 
@@ -168,10 +187,17 @@ FWorldPartitionLoadingContext::FDeferred::~FDeferred()
 					}
 				}
 				
-				//Only flush if there are items in the array. An empty array will result in a complete flush which is not the intended result.
-				if(AllAsyncLoadIds.Num() > 0)
+				// Only flush if there are items in the array. An empty array will result in a complete flush which is not the intended result.
+				if (AllAsyncLoadIds.Num() > 0)
 				{
 					FlushAsyncLoading(AllAsyncLoadIds);
+
+					// After a flush, Actor Instances in AllAsyncLoadIds may have been reinstanced, which means their FWorldPartitionActorDescInstance pointer is now deleted and unsafe to dereference. Apply those.
+					for (FWorldPartitionActorDescInstance* CancelledActorDescInst : ContainerInstanceCancellations)
+					{
+						SourceList.Remove(CancelledActorDescInst);
+					}
+					ContainerInstanceCancellations.Reset();
 				}
 
 				ULevel* Level = nullptr;
@@ -189,6 +215,8 @@ FWorldPartitionLoadingContext::FDeferred::~FDeferred()
 					}
 				}
 
+				check(ContainerInstanceCancellations.IsEmpty()); // Verify that cancellations haven't been generated outside the FlushAsyncLoading block
+
 				return Level;
 			};
 
@@ -278,6 +306,45 @@ void FWorldPartitionLoadingContext::FDeferred::UnregisterActor(FWorldPartitionAc
 	}
 }
 
+void FWorldPartitionLoadingContext::FDeferred::CancelOperation(FWorldPartitionActorDescInstance* InActorDescInstance)
+{
+	check(InActorDescInstance);
+	check(InActorDescInstance->bIsRegisteringOrUnregistering); // Callers should avoid CancelOperation when unnecessary
+
+	// If we aren't inside a ~FDeferred, we can find and cancel the op directly
+	bool bRemoved = false;
+	UActorDescContainerInstance* Container = InActorDescInstance->GetContainerInstance();
+	check(Container);
+	FContainerInstanceOps* Ops = ContainerInstanceOps.Find(Container);
+	if (Ops != nullptr)
+	{
+		bRemoved = Ops->Registrations.Remove(InActorDescInstance) > 0;
+		if (bRemoved)
+		{
+			--NumRegistrations;
+		}
+		else
+		{
+			bRemoved = Ops->Unregistrations.Remove(InActorDescInstance) > 0;
+			if (bRemoved)
+			{
+				--NumUnregistrations;
+			}
+		}
+	}
+
+	// If we have not found the registration or unregistration, we are inside ~FDeferred, where the ContainerInstanceOps are held locally. In this case push to a list of cancellations.
+	// No need to update NumRegistrations/NumUnregistrations as we are destroying
+	if (!bRemoved)
+	{
+		bool bIsAlreadyInSetPtr;
+		ContainerInstanceCancellations.Add(InActorDescInstance, &bIsAlreadyInSetPtr);
+		check(!bIsAlreadyInSetPtr);
+	}
+
+	InActorDescInstance->bIsRegisteringOrUnregistering = false;
+}
+
 /**
 * FWorldPartitionHandleImpl
 */
diff --git a/Engine/Source/Runtime/Engine/Public/WorldPartition/WorldPartitionHandle.h b/Engine/Source/Runtime/Engine/Public/WorldPartition/WorldPartitionHandle.h
--- a/Engine/Source/Runtime/Engine/Public/WorldPartition/WorldPartitionHandle.h
+++ b/Engine/Source/Runtime/Engine/Public/WorldPartition/WorldPartitionHandle.h
@@ -39,8 +39,11 @@ public:
 		ENGINE_API virtual ~IContext();
 
 	private:
+		// Should return true if the context caches objects that can be released outside its control
+		virtual bool IsDeferringRawObjects() { return false; };
 		virtual void RegisterActor(FWorldPartitionActorDescInstance* InActorDescInstance) = 0;
 		virtual void UnregisterActor(FWorldPartitionActorDescInstance* InActorDescInstance) = 0;
+		virtual void CancelOperation(FWorldPartitionActorDescInstance* InActorDescInstance) = 0;
 	};
 
 	/**
@@ -51,6 +54,7 @@ public:
 	private:
 		ENGINE_API virtual void RegisterActor(FWorldPartitionActorDescInstance* InActorDescInstance) override;
 		ENGINE_API virtual void UnregisterActor(FWorldPartitionActorDescInstance* InActorDescInstance) override;
+		virtual void CancelOperation(FWorldPartitionActorDescInstance* InActorDescInstance) override {} // Immediate context, nothing to cancel
 	};
 
 	/**
@@ -68,17 +72,20 @@ public:
 		bool GetNeedsClearTransactions() const { return bNeedsClearTransactions; }
 
 	private:
+		virtual bool IsDeferringRawObjects() override { return true; };
 		ENGINE_API virtual void RegisterActor(FWorldPartitionActorDescInstance* InActorDescInstance) override;
 		ENGINE_API virtual void UnregisterActor(FWorldPartitionActorDescInstance* InActorDescInstance) override;
+		ENGINE_API virtual void CancelOperation(FWorldPartitionActorDescInstance* InActorDescInstance) override;
 
 		struct FContainerInstanceOps
 		{
 			TSet<FWorldPartitionActorDescInstance*> Registrations;
 			TSet<FWorldPartitionActorDescInstance*> Unregistrations;
 		};
-
 		TMap<UActorDescContainerInstance*, FContainerInstanceOps> ContainerInstanceOps;
 
+		TSet<FWorldPartitionActorDescInstance*> ContainerInstanceCancellations;
+
 		int32 NumRegistrations = 0;
 		int32 NumUnregistrations = 0;
 		bool bNeedsClearTransactions = false;
@@ -92,8 +99,14 @@ public:
 	private:
 		virtual void RegisterActor(FWorldPartitionActorDescInstance* InActorDescInstance) override {}
 		virtual void UnregisterActor(FWorldPartitionActorDescInstance* InActorDescInstance) override {}
+		virtual void CancelOperation(FWorldPartitionActorDescInstance* InActorDescInstance) override {}
 	};
 
+	// Returns true if the current context is vulnerable to deleted objects
+	static bool IsDeferringRawObjects();
+	// If the current context is operating on the actor desc, cancel it
+	static void CancelOperationIfAny(FWorldPartitionActorDescInstance* InActorDescInstance);
+
 private:
 	static void LoadAndRegisterActor(FWorldPartitionActorDescInstance* ActorDescInstance);
 	static void UnloadAndUnregisterActor(FWorldPartitionActorDescInstance* ActorDescInstance);
