Deprecate anything related to extracting collision data from cloth simulations directly (e.g., FClothingSimulationInstance::AppendSimulationData, USkeletalMeshComponent::FindClothCollisions). This has never been supported by [REDACTED_PERSON] and only existed for Apex cloth support.

If users want to add collisions from a different SKM+physics asset to their simulation, they should use AddClothCollisionSource instead.

#Synced-CL 51899829

diff --git a/Engine/Plugins/ChaosCloth/Source/ChaosCloth/Private/ChaosCloth/ChaosClothingSimulation.cpp b/Engine/Plugins/ChaosCloth/Source/ChaosCloth/Private/ChaosCloth/ChaosClothingSimulation.cpp
--- a/Engine/Plugins/ChaosCloth/Source/ChaosCloth/Private/ChaosCloth/ChaosClothingSimulation.cpp
+++ b/Engine/Plugins/ChaosCloth/Source/ChaosCloth/Private/ChaosCloth/ChaosClothingSimulation.cpp
@@ -834,33 +834,6 @@ void FClothingSimulation::ClearExternalCollisions()
 	ExternalCollisionData.Reset();
 }
 
-void FClothingSimulation::GetCollisions(FClothCollisionData& OutCollisions, bool bIncludeExternal) const
-{
-	// This code only gathers old apex collisions that don't appear in the physics mesh
-	// It is also never called with bIncludeExternal = true 
-	// but the collisions are then added untransformed and added as external
-	// This function is bound to be deprecated at some point
-
-	OutCollisions.Reset();
-
-	// Add internal asset collisions
-	for (const TUniquePtr<FClothingSimulationCloth>& Cloth : Cloths)
-	{
-		for (const FClothingSimulationCollider* const Collider : Cloth->GetColliders())
-		{
-			OutCollisions.Append(Collider->GetCollisionData(Solver.Get(), Cloth.Get()));
-		}
-	}
-
-	// Add external asset collisions
-	if (bIncludeExternal)
-	{
-		OutCollisions.Append(ExternalCollisionData);
-	}
-
-	UE_LOG(LogChaosCloth, VeryVerbose, TEXT("GetCollisions returned collisions: %d spheres, %d capsules, %d convexes, %d boxes."), OutCollisions.Spheres.Num() - 2 * OutCollisions.SphereConnections.Num(), OutCollisions.SphereConnections.Num(), OutCollisions.Convexes.Num(), OutCollisions.Boxes.Num());
-}
-
 void FClothingSimulation::RefreshClothConfig(const IClothingSimulationContext* InContext)
 {
 	UpdateSimulationFromSharedSimConfig();
diff --git a/Engine/Plugins/ChaosCloth/Source/ChaosCloth/Public/ChaosCloth/ChaosClothingSimulation.h b/Engine/Plugins/ChaosCloth/Source/ChaosCloth/Public/ChaosCloth/ChaosClothingSimulation.h
--- a/Engine/Plugins/ChaosCloth/Source/ChaosCloth/Public/ChaosCloth/ChaosClothingSimulation.h
+++ b/Engine/Plugins/ChaosCloth/Source/ChaosCloth/Public/ChaosCloth/ChaosClothingSimulation.h
@@ -62,7 +62,6 @@ PRAGMA_ENABLE_DEPRECATION_WARNINGS
 
 		virtual void AddExternalCollisions(const FClothCollisionData& InData) override;
 		virtual void ClearExternalCollisions() override;
-		virtual void GetCollisions(FClothCollisionData& OutCollisions, bool bIncludeExternal = true) const override;
 		//~ End IClothingSimulationInterface Interface
 
 	public:
diff --git a/Engine/Plugins/ChaosClothAsset/Source/ChaosClothAssetEngine/Private/ChaosClothAsset/ClothAssetSKMClothingSimulation.cpp b/Engine/Plugins/ChaosClothAsset/Source/ChaosClothAssetEngine/Private/ChaosClothAsset/ClothAssetSKMClothingSimulation.cpp
--- a/Engine/Plugins/ChaosClothAsset/Source/ChaosClothAssetEngine/Private/ChaosClothAsset/ClothAssetSKMClothingSimulation.cpp
+++ b/Engine/Plugins/ChaosClothAsset/Source/ChaosClothAssetEngine/Private/ChaosClothAsset/ClothAssetSKMClothingSimulation.cpp
@@ -245,11 +245,6 @@ namespace UE::Chaos::ClothAsset
 		}
 	}
 
-	void FSKMClothingSimulation::GetCollisions(FClothCollisionData& OutCollisions, bool bIncludeExternal) const
-	{
-		OutCollisions.Reset();
-	}
-
 	void FSKMClothingSimulation::SetNumIterations(int32 NumIterations)
 	{
 		for (const TSharedPtr<::Chaos::Softs::FCollectionPropertyFacade>& SolverPropertyFacade : GetSolverPropertyFacades())
diff --git a/Engine/Plugins/ChaosClothAsset/Source/ChaosClothAssetEngine/Public/ChaosClothAsset/ClothAssetSKMClothingSimulation.h b/Engine/Plugins/ChaosClothAsset/Source/ChaosClothAssetEngine/Public/ChaosClothAsset/ClothAssetSKMClothingSimulation.h
--- a/Engine/Plugins/ChaosClothAsset/Source/ChaosClothAssetEngine/Public/ChaosClothAsset/ClothAssetSKMClothingSimulation.h
+++ b/Engine/Plugins/ChaosClothAsset/Source/ChaosClothAssetEngine/Public/ChaosClothAsset/ClothAssetSKMClothingSimulation.h
@@ -100,7 +100,6 @@ namespace UE::Chaos::ClothAsset
 
 		virtual void AddExternalCollisions(const FClothCollisionData& Data) override;
 		virtual void ClearExternalCollisions() override;
-		virtual void GetCollisions(FClothCollisionData& OutCollisions, bool bIncludeExternal = true) const override;
 		virtual int32 GetNumCloths() const override { return ClothSimulationProxy.IsValid() ? ClothSimulationProxy->GetNumCloths() : 0; }
 		virtual int32 GetNumKinematicParticles() const override { return ClothSimulationProxy.IsValid() ? ClothSimulationProxy->GetNumKinematicParticles() : 0; }
 		virtual int32 GetNumDynamicParticles() const override { return ClothSimulationProxy.IsValid() ? ClothSimulationProxy->GetNumDynamicParticles() : 0; }
diff --git a/Engine/Source/Editor/Persona/Private/AnimViewportShowCommands.cpp b/Engine/Source/Editor/Persona/Private/AnimViewportShowCommands.cpp
--- a/Engine/Source/Editor/Persona/Private/AnimViewportShowCommands.cpp
+++ b/Engine/Source/Editor/Persona/Private/AnimViewportShowCommands.cpp
@@ -69,7 +69,6 @@ void FAnimViewportShowCommands::RegisterCommands()
 
 	UI_COMMAND( ResetClothSimulation, "Reset Cloth Simulation", "Reset simulated cloth to its initial state", EUserInterfaceActionType::Button, FInputChord(EModifierKey::Control | EModifierKey::Alt, EKeys::C) );
 
-	UI_COMMAND( EnableCollisionWithAttachedClothChildren, "Collide with Cloth Children", "Enables collision detection between collision primitives in the base mesh and clothing on any attachments in the preview scene", EUserInterfaceActionType::ToggleButton, FInputChord() );	
 	UI_COMMAND(PauseClothWithAnim, "Pause with Animation", "If enabled, the clothing simulation will pause when the animation is paused using the scrub panel", EUserInterfaceActionType::ToggleButton, FInputChord());
 
 	// below 3 menus are radio button styles
diff --git a/Engine/Source/Editor/Persona/Private/AnimViewportShowCommands.h b/Engine/Source/Editor/Persona/Private/AnimViewportShowCommands.h
--- a/Engine/Source/Editor/Persona/Private/AnimViewportShowCommands.h
+++ b/Engine/Source/Editor/Persona/Private/AnimViewportShowCommands.h
@@ -150,6 +150,7 @@ public:
 	/** Enables collision detection between collision primitives in the base mesh 
 	  * and clothing on any attachments in the preview scene. 
 	*/
+	UE_DEPRECATED(5.8, "This functionality has been deprecated")
 	TSharedPtr< FUICommandInfo > EnableCollisionWithAttachedClothChildren;
 
 	/** Show all sections which means the original state */
diff --git a/Engine/Source/Editor/Persona/Private/SAnimationEditorViewport.cpp b/Engine/Source/Editor/Persona/Private/SAnimationEditorViewport.cpp
--- a/Engine/Source/Editor/Persona/Private/SAnimationEditorViewport.cpp
+++ b/Engine/Source/Editor/Persona/Private/SAnimationEditorViewport.cpp
@@ -1158,12 +1158,6 @@ void SAnimationEditorViewportTabBody::BindCommands()
 		FExecuteAction::CreateSP(this, &SAnimationEditorViewportTabBody::OnResetClothSimulation),
 		FCanExecuteAction::CreateSP(this, &SAnimationEditorViewportTabBody::IsClothSimulationEnabled));
 
-	CommandList.MapAction( 
-		ViewportShowMenuCommands.EnableCollisionWithAttachedClothChildren,
-		FExecuteAction::CreateSP(this, &SAnimationEditorViewportTabBody::OnEnableCollisionWithAttachedClothChildren),
-		FCanExecuteAction(),
-		FIsActionChecked::CreateSP(this, &SAnimationEditorViewportTabBody::IsEnablingCollisionWithAttachedClothChildren));
-
 	CommandList.MapAction(
 		ViewportShowMenuCommands.PauseClothWithAnim,
 		FExecuteAction::CreateSP(this, &SAnimationEditorViewportTabBody::OnPauseClothingSimWithAnim),
@@ -2758,29 +2752,6 @@ float SAnimationEditorViewportTabBody::GetGravityScaleSliderValue() const
 	return GetPreviewScene()->GetGravityScale();
 }
 
-void SAnimationEditorViewportTabBody::OnEnableCollisionWithAttachedClothChildren()
-{
-	UDebugSkelMeshComponent* PreviewComponent = GetPreviewScene()->GetPreviewMeshComponent();
-
-	if( PreviewComponent )
-	{
-		PreviewComponent->bCollideWithAttachedChildren = !PreviewComponent->bCollideWithAttachedChildren;
-		RefreshViewport();
-	}
-}
-
-bool SAnimationEditorViewportTabBody::IsEnablingCollisionWithAttachedClothChildren() const
-{
-	UDebugSkelMeshComponent* PreviewComponent = GetPreviewScene()->GetPreviewMeshComponent();
-
-	if( PreviewComponent )
-	{
-		return PreviewComponent->bCollideWithAttachedChildren;
-	}
-
-	return false;
-}
-
 void SAnimationEditorViewportTabBody::OnSetSectionsDisplayMode(ESectionDisplayMode DisplayMode)
 {
 	UDebugSkelMeshComponent* PreviewComponent = GetPreviewScene()->GetPreviewMeshComponent();
diff --git a/Engine/Source/Editor/Persona/Private/SAnimationEditorViewport.h b/Engine/Source/Editor/Persona/Private/SAnimationEditorViewport.h
--- a/Engine/Source/Editor/Persona/Private/SAnimationEditorViewport.h
+++ b/Engine/Source/Editor/Persona/Private/SAnimationEditorViewport.h
@@ -590,10 +590,6 @@ private:
 	void OnPauseClothingSimWithAnim();
 	bool IsPausingClothingSimWithAnim();
 
-	/** Enable collision with clothes on attached children */
-	void OnEnableCollisionWithAttachedClothChildren();
-	bool IsEnablingCollisionWithAttachedClothChildren() const;
-
 	/** Show all sections which means the original state */
 	void OnSetSectionsDisplayMode(ESectionDisplayMode DisplayMode);
 	bool IsSectionsDisplayMode(ESectionDisplayMode DisplayMode) const;
diff --git a/Engine/Source/Editor/Persona/Private/ViewportToolbar/AnimationEditorMenus.cpp b/Engine/Source/Editor/Persona/Private/ViewportToolbar/AnimationEditorMenus.cpp
--- a/Engine/Source/Editor/Persona/Private/ViewportToolbar/AnimationEditorMenus.cpp
+++ b/Engine/Source/Editor/Persona/Private/ViewportToolbar/AnimationEditorMenus.cpp
@@ -757,7 +757,6 @@ void FillCharacterClothingMenu(UToolMenu* InMenu)
 			"", GravityWidget.ToSharedRef(), LOCTEXT("ClothPreview_GravityScale", "Gravity Scale:")
 		));
 
-		ClothPreviewSection.AddMenuEntry(Actions.EnableCollisionWithAttachedClothChildren);
 		ClothPreviewSection.AddMenuEntry(Actions.PauseClothWithAnim);
 	}
 
diff --git a/Engine/Source/Runtime/ClothingSystemRuntimeInterface/Private/ClothingSimulationInstance.cpp b/Engine/Source/Runtime/ClothingSystemRuntimeInterface/Private/ClothingSimulationInstance.cpp
--- a/Engine/Source/Runtime/ClothingSystemRuntimeInterface/Private/ClothingSimulationInstance.cpp
+++ b/Engine/Source/Runtime/ClothingSystemRuntimeInterface/Private/ClothingSimulationInstance.cpp
@@ -81,11 +81,3 @@ void FClothingSimulationInstance::FillContextAndPrepareTick(const USkeletalMeshC
 	const bool bForceTeleportResetOnly = !bIsInitialization && !ClothingSimulation->ShouldSimulateLOD(CurrentOwnerLODIndex);
 	ClothingSimulation->FillContextAndPrepareTick(OwnerComponent, DeltaTime, ClothingSimulationContext, bIsInitialization, bForceTeleportResetOnly);
 }
-
-void FClothingSimulationInstance::AppendUniqueCollisions(FClothCollisionData& OutCollisions, bool bIncludeExternal) const
-{
-	FClothCollisionData CollisionsToAppend;
-	ClothingSimulation->GetCollisions(CollisionsToAppend, bIncludeExternal);
-
-	OutCollisions.AppendUnique(CollisionsToAppend);
-}
diff --git a/Engine/Source/Runtime/ClothingSystemRuntimeInterface/Public/ClothingSimulationInstance.h b/Engine/Source/Runtime/ClothingSystemRuntimeInterface/Public/ClothingSimulationInstance.h
--- a/Engine/Source/Runtime/ClothingSystemRuntimeInterface/Public/ClothingSimulationInstance.h
+++ b/Engine/Source/Runtime/ClothingSystemRuntimeInterface/Public/ClothingSimulationInstance.h
@@ -119,13 +119,12 @@ public:
 	}
 
 	/** Return the collision data from this simulation instance. */
-	void GetCollisions(FClothCollisionData& Collisions, bool bIncludeExternal) const
-	{
-		ClothingSimulation->GetCollisions(Collisions, bIncludeExternal);
-	}
+	UE_DEPRECATED(5.8, "GetCollisions has been deprecated. Use CollisionSources to add external collisions instead.")
+	void GetCollisions(FClothCollisionData& Collisions, bool bIncludeExternal) const {}
 
 	/** Append unique collision data from this simulation instance to the specified collision data. */
-	UE_API void AppendUniqueCollisions(FClothCollisionData& OutCollisions, bool bIncludeExternal) const;
+	UE_DEPRECATED(5.8, "AppendUniqueCollisions has been deprecated. Use CollisionSources to add external collisions instead.")
+	void AppendUniqueCollisions(FClothCollisionData& OutCollisions, bool bIncludeExternal) const {}
 
 	/** Add external dynamic collisions. */
 	void AddExternalCollisions(const FClothCollisionData& Collisions)
diff --git a/Engine/Source/Runtime/ClothingSystemRuntimeInterface/Public/ClothingSimulationInterface.h b/Engine/Source/Runtime/ClothingSystemRuntimeInterface/Public/ClothingSimulationInterface.h
--- a/Engine/Source/Runtime/ClothingSystemRuntimeInterface/Public/ClothingSimulationInterface.h
+++ b/Engine/Source/Runtime/ClothingSystemRuntimeInterface/Public/ClothingSimulationInterface.h
@@ -135,14 +135,8 @@ protected:
 	 */
 	virtual void ClearExternalCollisions() = 0;
 
-	/**
-	 * Called by the engine to request data on all active collisions in a simulation. if bIncludeExternal is
-	 * true, derived implementations should add the asset collisions and any collisions added at runtime, if
-	 * false only the collisions from the asset should be considered
-	 * @param OutCollisions - Array to write collisions to
-	 * @param bIncludeExternal - Whether or not external collisions should be retrieved, or just asset collisions
-	 */
-	virtual void GetCollisions(FClothCollisionData& OutCollisions, bool bIncludeExternal = true) const = 0;
+	UE_DEPRECATED(5.8, "GetCollisions has been deprecated. Use CollisionSources to add external collisions instead.")
+	virtual void GetCollisions(FClothCollisionData& OutCollisions, bool bIncludeExternal = true) const {};
 
 public:
 PRAGMA_DISABLE_DEPRECATION_WARNINGS
diff --git a/Engine/Source/Runtime/Engine/Classes/Components/SkeletalMeshComponent.h b/Engine/Source/Runtime/Engine/Classes/Components/SkeletalMeshComponent.h
--- a/Engine/Source/Runtime/Engine/Classes/Components/SkeletalMeshComponent.h
+++ b/Engine/Source/Runtime/Engine/Classes/Components/SkeletalMeshComponent.h
@@ -729,9 +729,12 @@ public:
 	/** can't collide with part of environment if total collision volumes exceed 16 capsules or 32 planes per convex */
 	UPROPERTY(EditAnywhere, BlueprintReadWrite, Category=Clothing)
 	uint8 bCollideWithEnvironment:1;
+
 	/** can't collide with part of attached children if total collision volumes exceed 16 capsules or 32 planes per convex */
-	UPROPERTY(EditAnywhere, BlueprintReadWrite, Category=Clothing)
+	UE_DEPRECATED(5.8, "bCollideWithAttachedChildren is deprecated. Use CollisionSources to add external collisions instead.")
+	UPROPERTY(BlueprintReadWrite, Category = Clothing, meta=(DeprecatedProperty, DeprecationMessage="CollideWithAttachedChildren is deprecated. Use CollisionSources to add external collisions instead."))
 	uint8 bCollideWithAttachedChildren:1;
+
 	/**
 	 * Forces the cloth simulation to constantly update its external collisions, at the expense of performance.
 	 * This will help to prevent missed collisions if the cloth's skeletal mesh component isn't moving,
@@ -2542,13 +2545,16 @@ protected:
 	ENGINE_API void ProcessClothCollisionWithEnvironment();
 
 	/** copy parent's cloth collisions to attached children, where parent means this component */
-	ENGINE_API void CopyClothCollisionsToChildren();
+	UE_DEPRECATED(5.8, "CopyClothCollisionsToChildren has been deprecated. Use CollisionSources to add external collisions instead.")
+	void CopyClothCollisionsToChildren() {}
 
 	/** copy children's cloth collisions to parent, where parent means this component */
-	ENGINE_API void CopyChildrenClothCollisionsToParent();
+	UE_DEPRECATED(5.8, "CopyChildrenClothCollisionsToParent has been deprecated. Use CollisionSources to add external collisions instead.")
+	void CopyChildrenClothCollisionsToParent(){}
 
 	/** find if this component has collisions for clothing and return the results calculated by bone transforms */
-	ENGINE_API void FindClothCollisions(FClothCollisionData& OutCollisions) const;
+	UE_DEPRECATED(5.8, "FindClothCollisions has been deprecated. Use CollisionSources to add external collisions instead.")
+	void FindClothCollisions(FClothCollisionData& OutCollisions) const {}
 
 #endif
 
diff --git a/Engine/Source/Runtime/Engine/Private/SkeletalMeshComponentPhysics.cpp b/Engine/Source/Runtime/Engine/Private/SkeletalMeshComponentPhysics.cpp
--- a/Engine/Source/Runtime/Engine/Private/SkeletalMeshComponentPhysics.cpp
+++ b/Engine/Source/Runtime/Engine/Private/SkeletalMeshComponentPhysics.cpp
@@ -3229,23 +3229,6 @@ void USkeletalMeshComponent::GetWindForCloth_GameThread(FVector& WindDirection,
 
 #if WITH_CLOTH_COLLISION_DETECTION
 
-void USkeletalMeshComponent::FindClothCollisions(FClothCollisionData& OutCollisions) const
-{
-	if (ClothingSimulationInstances.Num())
-	{
-		// append skeletal component collisions
-		constexpr bool bIncludeExternal = false;
-
-		ClothingSimulationInstances[0].GetCollisions(OutCollisions, bIncludeExternal);
-
-		// AppendUnique is quite expensive, so let's only call it for subsequent clothing simulation instances
-		for (int32 Index = 1; Index < ClothingSimulationInstances.Num(); ++Index)
-		{
-			ClothingSimulationInstances[Index].AppendUniqueCollisions(OutCollisions, bIncludeExternal);
-		}
-	}
-}
-
 // UE_DEPRECATED(5.7, "Use ClothCollisionSource::ExtractCollisions() instead.")
 void USkeletalMeshComponent::ExtractCollisionsForCloth(
 	USkeletalMeshComponent* SourceComponent, 
@@ -3261,38 +3244,6 @@ void USkeletalMeshComponent::ExtractCollisionsForCloth(
 	}
 }
 
-void USkeletalMeshComponent::CopyClothCollisionsToChildren()
-{
-	// Find all suitable children
-	const TArray<TObjectPtr<USceneComponent>>& AttachedChildren = USceneComponent::GetAttachChildren();
-
-	TArray<USkeletalMeshComponent*> ClothChildren;
-	ClothChildren.Reserve(AttachedChildren.Num());
-
-	for (const TObjectPtr<USceneComponent>& AttachedChild : AttachedChildren)
-	{
-		if (USkeletalMeshComponent* const Child = Cast<USkeletalMeshComponent>(AttachedChild);
-			Child && Child->ClothingSimulationInstances.Num() > 0)
-		{
-			ClothChildren.Add(Child);
-		}
-	}
-
-	// Add new collision to children
-	if (ClothChildren.Num())
-	{
-		FClothCollisionData NewCollisions;
-		FindClothCollisions(NewCollisions);
-
-		for (USkeletalMeshComponent* const Child : ClothChildren)
-		{
-			for (FClothingSimulationInstance& ClothingSimulationInstance : Child->ClothingSimulationInstances)
-			{
-				ClothingSimulationInstance.AddExternalCollisions(NewCollisions);
-			}
-		}
-	}
-}
 
 void USkeletalMeshComponent::CopyClothCollisionSources()
 {
@@ -3309,26 +3260,6 @@ void USkeletalMeshComponent::CopyClothCollisionSources()
 	}
 }
 
-void USkeletalMeshComponent::CopyChildrenClothCollisionsToParent()
-{
-	// Find new collisions from children
-	FClothCollisionData NewCollisions;
-
-	for (const TObjectPtr<USceneComponent>& AttachedChild : GetAttachChildren())
-	{
-		if (const USkeletalMeshComponent* const Child = Cast<USkeletalMeshComponent>(AttachedChild))
-		{
-			Child->FindClothCollisions(NewCollisions);
-		}
-	}
-
-	// Add new collisions to parent (this component)
-	for (FClothingSimulationInstance& ClothingSimulationInstance : ClothingSimulationInstances)
-	{
-		ClothingSimulationInstance.AddExternalCollisions(NewCollisions);
-	}
-}
-
 void USkeletalMeshComponent::ProcessClothCollisionWithEnvironment()
 {
 	// don't handle collision detection if this component is in editor
@@ -3462,11 +3393,6 @@ void USkeletalMeshComponent::UpdateClothTransformImp()
 			ClothingSimulationInstance.ClearExternalCollisions();
 		}
 
-		if (bCollideWithAttachedChildren)
-		{
-			CopyClothCollisionsToChildren();
-		}
-
 		if (ClothCollisionSources.Num() > 0)
 		{
 			CopyClothCollisionSources();
@@ -3729,19 +3655,6 @@ void USkeletalMeshComponent::UpdateClothStateAndSimulate(float DeltaTime, FTickF
 	// Make sure we aren't already in flight from previous frame
 	HandleExistingParallelClothSimulation();
 
-#if WITH_CLOTH_COLLISION_DETECTION
-	if (bCollideWithAttachedChildren)
-	{
-		for (FClothingSimulationInstance& ClothingSimulationInstance : ClothingSimulationInstances)
-		{
-			ClothingSimulationInstance.ClearExternalCollisions();
-		}
-
-		CopyClothCollisionsToChildren();
-		CopyChildrenClothCollisionsToParent();
-	}
-#endif // WITH_CLOTH_COLLISION_DETECTION
-
 	UpdateClothSimulationContext(DeltaTime);
 
 	ParallelClothTask = TGraphTask<FParallelClothTask>::CreateTask(nullptr, ENamedThreads::GameThread).ConstructAndDispatchWhenReady(*this, DeltaTime);
