Fix in the new material translator causing links connected to the final output node to not be semantically colorized.

#rnx
#Synced-CL 51827487

diff --git a/Engine/Source/Runtime/Engine/Private/Materials/MaterialIRModuleBuilder.cpp b/Engine/Source/Runtime/Engine/Private/Materials/MaterialIRModuleBuilder.cpp
--- a/Engine/Source/Runtime/Engine/Private/Materials/MaterialIRModuleBuilder.cpp
+++ b/Engine/Source/Runtime/Engine/Private/Materials/MaterialIRModuleBuilder.cpp
@@ -702,6 +702,9 @@ struct FMaterialIRModuleBuilderImpl
 			{
 				Module->AddError(nullptr, FString::Printf(TEXT("Expression connected to the MaterialAttributes material output does is not a MaterialAttributes value (it is a '%s' instead)."), *MaterialAttributesValue->Type.GetSpelling()));
 			}
+
+			// Push insight for the MaterialAttributes connection itself so the edge colorizes.
+			PushConnectionInsight(Builder->MaterialInterface->GetMaterial(), (int32)MP_MaterialAttributes, InputDesc.Input->Expression, InputDesc.Input->OutputIndex, MaterialAttributesValue->Type);
 		}
 
 		for (EMaterialProperty Property : MaterialAttributesAggregate::GetMaterialProperties())
