[PSOPrecache] Skip PSO precache when shadermap resource for IAD assets is not resolved.

PSO precache is another point where shadermap resources must verify they are fully resolved when they belong to IAD containers.
Also fix crash in `FOnDemandInstallCache::CompleteRequest()` when buffer was reported as Ok but has invalid encoded size.

#test Ran client and added debug code to ensure PSOPrecache can be skipped to IAD shadermaps without visual issues.
#rnx

diff --git a/Engine/Source/Runtime/Engine/Private/Materials/MaterialShader.cpp b/Engine/Source/Runtime/Engine/Private/Materials/MaterialShader.cpp
--- a/Engine/Source/Runtime/Engine/Private/Materials/MaterialShader.cpp
+++ b/Engine/Source/Runtime/Engine/Private/Materials/MaterialShader.cpp
@@ -3469,6 +3469,15 @@ FPSOPrecacheDataArray FMaterialShaderMap::CollectPSOPrecacheData(const FMaterial
 	{
 		return FPSOPrecacheDataArray();
 	}
+
+	// Check if all shader chunks for this shader map are resolved (i.e. IoStore chunks are installed).
+	// If not, bail out early to avoid a fatal error in GetResultOrDie when creating shaders from unresolved chunks.
+	// The higher-level rendering system will use the shader map once it becomes resolved, so we only risk a hitch.
+	if (!ensureMsgf(GetResource() && GetResource()->IsResolved(), TEXT("Shadermap resource '%s' is not resolved during PSO precaching"), PrecacheParams.Material ? *PrecacheParams.Material->GetAssetName() : TEXT("<Unknown>")))
+	{
+		return FPSOPrecacheDataArray();
+	}
+
 	bool bHasShaderMapForVFType = PrecacheParams.Material->IsUIMaterial() || (PrecacheParams.VertexFactoryData.VertexFactoryType && LocalContent->GetMeshShaderMap(PrecacheParams.VertexFactoryData.VertexFactoryType->GetHashedName()));
 	if (!bHasShaderMapForVFType)
 	{
diff --git a/Engine/Source/Runtime/Experimental/IoStore/OnDemand/Private/OnDemandInstallCache.cpp b/Engine/Source/Runtime/Experimental/IoStore/OnDemand/Private/OnDemandInstallCache.cpp
--- a/Engine/Source/Runtime/Experimental/IoStore/OnDemand/Private/OnDemandInstallCache.cpp
+++ b/Engine/Source/Runtime/Experimental/IoStore/OnDemand/Private/OnDemandInstallCache.cpp
@@ -4533,6 +4533,8 @@ void FOnDemandInstallCache::CompleteRequest(FIoRequestImpl* Request, EIoErrorCod
 	}
 	else
 	{
+		// Ensure buffer is valid, otherwise Request->GetBuffer() will crash without information about the ChunkId
+		checkf(Request->HasBuffer(), TEXT("Missing buffer when IoRequest status was reported as Ok, ChunkId=%s"), *LexToString(Request->ChunkId));
 		TRACE_IOSTORE_BACKEND_REQUEST_COMPLETED(Request, Request->GetBuffer().GetSize());
 	}
 
diff --git a/Engine/Source/Runtime/RenderCore/Private/ShaderLibrary/ShaderLibraryInstance.h b/Engine/Source/Runtime/RenderCore/Private/ShaderLibrary/ShaderLibraryInstance.h
--- a/Engine/Source/Runtime/RenderCore/Private/ShaderLibrary/ShaderLibraryInstance.h
+++ b/Engine/Source/Runtime/RenderCore/Private/ShaderLibrary/ShaderLibraryInstance.h
@@ -65,6 +65,9 @@ public:
 	/** Returns chunk id for this instance, if it is a part of a chunked named library. INDEX_NONE is returned for a monolithic lib. */
 	int32 GetChunkId() const { return ChunkId; }
 
+	/** Returns whether the specified shadermap in this library instance is considered resolved. Always true for shader libraries that represent non-IAD containers. */
+	bool IsShaderMapResourceResolved(int32 ShaderMapIndex) const;
+
 public:
 	FRHIShaderLibraryRef Library;
 
@@ -75,8 +78,6 @@ private:
 
 	bool OnExternalReadCallback(const FGraphEventArray& Events, double RemainingTime);
 
-	bool IsShaderMapResourceResolved(int32 ShaderMapIndex) const;
-
 	enum class EPreloadingState : uint8
 	{
 		NotPreloaded = 0, // Not preloaded.
diff --git a/Engine/Source/Runtime/RenderCore/Private/ShaderLibrary/ShaderMapResourceSharedCode.cpp b/Engine/Source/Runtime/RenderCore/Private/ShaderLibrary/ShaderMapResourceSharedCode.cpp
--- a/Engine/Source/Runtime/RenderCore/Private/ShaderLibrary/ShaderMapResourceSharedCode.cpp
+++ b/Engine/Source/Runtime/RenderCore/Private/ShaderLibrary/ShaderMapResourceSharedCode.cpp
@@ -79,6 +79,11 @@ FSHAHash FShaderMapResource_SharedCode::GetShaderMapHash() const
 	return LibraryInstance->Library->GetShaderMapHash(ShaderMapIndex);
 }
 
+bool FShaderMapResource_SharedCode::IsResolved() const
+{
+	return LibraryInstance && LibraryInstance->IsShaderMapResourceResolved(ShaderMapIndex);
+}
+
 void FShaderMapResource_SharedCode::AssignToNewShaderLibrary(FShaderLibraryInstance* NewLibraryInstance, int32 NewShaderMapIndex)
 {
 	const int32 OldShaderMapIndex = this->ShaderMapIndex;
diff --git a/Engine/Source/Runtime/RenderCore/Private/ShaderLibrary/ShaderMapResourceSharedCode.h b/Engine/Source/Runtime/RenderCore/Private/ShaderLibrary/ShaderMapResourceSharedCode.h
--- a/Engine/Source/Runtime/RenderCore/Private/ShaderLibrary/ShaderMapResourceSharedCode.h
+++ b/Engine/Source/Runtime/RenderCore/Private/ShaderLibrary/ShaderMapResourceSharedCode.h
@@ -38,6 +38,7 @@ public:
 	virtual int32 GetLibraryShaderIndex(int32 ShaderIndex) const override;
 	virtual uint32 GetShaderSizeBytes(int32 ShaderIndex) const override;
 	virtual FSHAHash GetShaderMapHash() const override;
+	virtual bool IsResolved() const override;
 
 	// Assigns this resource to the new shader library with its new shadermap index. See FShaderLibrariesCollection::MoveShaderMapResourceOwnership().
 	void AssignToNewShaderLibrary(FShaderLibraryInstance* NewLibraryInstance, int32 NewShaderMapIndex);
diff --git a/Engine/Source/Runtime/RenderCore/Public/Shader.h b/Engine/Source/Runtime/RenderCore/Public/Shader.h
--- a/Engine/Source/Runtime/RenderCore/Public/Shader.h
+++ b/Engine/Source/Runtime/RenderCore/Public/Shader.h
@@ -391,6 +391,12 @@ public:
 		return INDEX_NONE;
 	}
 
+	/** Returns whether all shader chunks for this shader map are resolved and available. Used to guard against PSO precaching on shaders whose IoStore chunks are not yet installed. */
+	virtual bool IsResolved() const
+	{
+		return true;
+	}
+
 	inline bool ContainsAtLeastOneRHIShaderCreated() const
 	{
 		return bAtLeastOneRHIShaderCreated;
