[Engine][Animation] Fix linked anim layer crash during BP reinstancing

Bug: When compiling a pawn Blueprint whose AnimBP uses linked anim layers and depends on the pawn BP class, the editor crashes with 'Assertion failed: InBoneContainer && InBoneContainer->IsValid()' in FBonePose::SetBoneContainer.

Cause: The pawn BP and AnimBP are recompiled together in the same compilation batch. During reinstancing, linked anim layer AnimInstances (e.g. gesture layers) are created on the old pawn's skeletal mesh component. The old pawn is then destroyed and its component is trashed (TRASH_ prefix, unregistered, no owning actor), but the linked AnimInstance survives and remains subscribed to FCoreUObjectDelegates::OnObjectsReinstanced. When the broadcast fires, FAnimNode_LinkedAnimLayer::HandleObjectsReinstanced_Impl runs on this stale AnimInstance. The existing guard only checks GetSkeletalMeshAsset(), which passes because USkeletalMeshComponent::OnUnregister clears RequiredBones and SharedRequiredBones but does not clear the mesh reference. RecalcRequiredBones then attempts to initialize a bone container from the empty RequiredBones data and crashes in SetBoneContainer.

Fix: Added an IsRegistered() check at the top of HandleObjectsReinstanced_Impl that gates the entire function. If the owning mesh component is not registered, the function returns immediately. This is safe because an unregistered component has no valid bone data and no need for anim graph reinitialization.

#Synced-CL 51968400

#Synced-CL 51968400

diff --git a/Engine/Source/Runtime/Engine/Private/Animation/AnimNode_LinkedAnimLayer.cpp b/Engine/Source/Runtime/Engine/Private/Animation/AnimNode_LinkedAnimLayer.cpp
--- a/Engine/Source/Runtime/Engine/Private/Animation/AnimNode_LinkedAnimLayer.cpp
+++ b/Engine/Source/Runtime/Engine/Private/Animation/AnimNode_LinkedAnimLayer.cpp
@@ -218,6 +218,18 @@ void FAnimNode_LinkedAnimLayer::HandleObjectsReinstanced_Impl(UObject* InSourceO
 	if(UseLegacyAnimInstanceReinstancingBehavior == nullptr || !UseLegacyAnimInstanceReinstancingBehavior->GetBool())
 	{
 		UAnimInstance* SourceAnimInstance = CastChecked<UAnimInstance>(InSourceObject);
+
+		// Skip if the owning mesh component is not registered. During BP reinstancing,
+		// linked anim layer instances can survive on trashed/orphaned components (the old
+		// pawn was destroyed but the linked AnimInstance is still subscribed to
+		// OnObjectsReinstanced). The component's bone data was cleared in OnUnregister
+		// but the mesh reference remains, so the existing mesh-only check is insufficient.
+		USkeletalMeshComponent* MeshComponent = SourceAnimInstance->GetSkelMeshComponent();
+		if (!MeshComponent || !MeshComponent->IsRegistered())
+		{
+			return;
+		}
+
 		FAnimInstanceProxy& SourceProxy = SourceAnimInstance->GetProxyOnAnyThread<FAnimInstanceProxy>();
 
 		// Call Initialize here to ensure any custom proxies are initialized (as they may have been re-created during
@@ -230,10 +242,7 @@ void FAnimNode_LinkedAnimLayer::HandleObjectsReinstanced_Impl(UObject* InSourceO
 
 		SourceProxy.InitializeCachedClassData();
 
-		// Ensure we have a valid mesh at this point, as calling into the graph without one can result in crashes
-		// as we assume a valid bone container/reference skeleton is present
-		USkeletalMeshComponent* MeshComponent = SourceAnimInstance->GetSkelMeshComponent();
-		if(MeshComponent && MeshComponent->GetSkeletalMeshAsset())
+		if(MeshComponent->GetSkeletalMeshAsset())
 		{
 			SourceAnimInstance->RecalcRequiredBones();
 
