Fix Skeletal Mesh Editor not updating active skin weight profile

#rnx
#Synced-CL 52037470

diff --git a/Engine/Source/Runtime/Engine/Private/Components/SkinnedMeshComponent.cpp b/Engine/Source/Runtime/Engine/Private/Components/SkinnedMeshComponent.cpp
--- a/Engine/Source/Runtime/Engine/Private/Components/SkinnedMeshComponent.cpp
+++ b/Engine/Source/Runtime/Engine/Private/Components/SkinnedMeshComponent.cpp
@@ -1904,6 +1904,7 @@ void USkinnedMeshComponent::TickComponent(float DeltaTime, enum ELevelTick TickT
 			// here. 
 			UpdateSkinWeightOverrideBuffer();
 			bSkinWeightProfileBufferOverridePending = false;
+			bSkinWeightProfileSet = true;
 		}
 	}
 }
