[Insights] Fix snapshot hitches feature paths

#Synced-CL 51621904

diff --git a/Engine/Source/Runtime/Engine/Private/ProfilingDebugging/TraceHitches.cpp b/Engine/Source/Runtime/Engine/Private/ProfilingDebugging/TraceHitches.cpp
--- a/Engine/Source/Runtime/Engine/Private/ProfilingDebugging/TraceHitches.cpp
+++ b/Engine/Source/Runtime/Engine/Private/ProfilingDebugging/TraceHitches.cpp
@@ -5,6 +5,7 @@
 #include "Misc/CoreMisc.h"
 #include "Misc/OutputDevice.h"
 #include "Misc/Paths.h"
+#include "HAL/FileManager.h"
 #include "ProfilingDebugging/TraceAuxiliary.h"
 #include "ProfilingDebugging/MiscTrace.h"
 #include "Stats/StatsData.h"
@@ -88,7 +89,8 @@ private:
 		if ((GameThreadTime > HitchThresholdSecs) || (RenderThreadTime > HitchThresholdSecs))
 		{
 			TRACE_BOOKMARK(TEXT("Hitch detected"));
-			FString OutputPath = FPaths::Combine(FPaths::ProfilingDir(), TEXT("Hitches"), FApp::GetInstanceId().ToString());
+			FString OutputPathRel = FPaths::Combine(FPaths::ProfilingDir(), TEXT("Hitches"), FApp::GetInstanceId().ToString());
+			FString OutputPath = IFileManager::Get().ConvertToAbsolutePathForExternalAppForWrite(*OutputPathRel);
 			FString SnapshotFile = FPaths::Combine(OutputPath, FString::Printf(TEXT("Frame-%lld.utrace"), Frame));
 			FString ScreenshotFile = FPaths::Combine(OutputPath, FString::Printf(TEXT("Frame-%lld.png"), Frame));
 			//Write snapshot to target directory. This will also create directory in case it doesn't exist
