Substrate - Toon
- UI renaming and tool tip update from artist feedback.
- Detail panel element greyed out when not usable.
- Clarified shadow hatching application. Now replace shadow/transmittance

#rnx
#tests PC editor toon

#Synced-CL 51842837

diff --git a/Engine/Shaders/Private/Substrate/SubstrateEvaluation.ush b/Engine/Shaders/Private/Substrate/SubstrateEvaluation.ush
--- a/Engine/Shaders/Private/Substrate/SubstrateEvaluation.ush
+++ b/Engine/Shaders/Private/Substrate/SubstrateEvaluation.ush
@@ -1312,7 +1312,7 @@ FSubstrateEvaluateResult SubstrateEvaluateBSDFCommon(FSubstrateBSDFContext BSDFC
 			{
 				// Apply shadow/transmittance hatching pattern
 				FToonProfileShadowSettings ToonProfileShadowSettings = GetToonProfileShadowSettings(ToonProfileId);
-				TransmittanceShadow *= EvaluateToonProfileShadowHatching(ToonProfileId, ToonPatternUV, ToonProfileShadowSettings, TransmittanceShadow);
+				TransmittanceShadow = EvaluateToonProfileShadowHatching(ToonProfileId, ToonPatternUV, ToonProfileShadowSettings, TransmittanceShadow);
 			}
 
 			float3 Specular = 0.0f;
diff --git a/Engine/Shaders/Private/ToonProfile.usf b/Engine/Shaders/Private/ToonProfile.usf
--- a/Engine/Shaders/Private/ToonProfile.usf
+++ b/Engine/Shaders/Private/ToonProfile.usf
@@ -11,11 +11,11 @@ uint ToonProfileIndex;
 
 float ToonProfileShadowExtinctionCoefficientInInvCentimeter;
 
-float DiffuseDitherPatternStrength;
-float DiffuseDitherPatternSize;
+float DiffuseRampOffsetStrength;
+float DiffuseRampOffsetSize;
 
-float SpecularDitherPatternStrength;
-float SpecularDitherPatternSize;
+float SpecularRampOffsetStrength;
+float SpecularRampOffsetSize;
 
 float ShadowHatchingPatternSize;
 float ShadowHatchingPatternStrength;
@@ -27,14 +27,14 @@ float SpecularIndirectRampRepetition;
 
 RWTexture2DArray<float4>	ToonProfileTexture;
 
-Texture2D<float4>			DiffuseDitherPattern;
-SamplerState				DiffuseDitherPatternSampler;
+Texture2D<float4>			DiffuseRampOffsetTexture;
+SamplerState				DiffuseRampOffsetTextureSampler;
 
-Texture2D<float4>			SpecularDitherPattern;
-SamplerState				SpecularDitherPatternSampler;
+Texture2D<float4>			SpecularRampOffsetTexture;
+SamplerState				SpecularRampOffsetTextureSampler;
 
-Texture2D<float4>			ShadowHatchingPattern;
-SamplerState				ShadowHatchingPatternSampler;
+Texture2D<float4>			ShadowHatchingPatternTexture;
+SamplerState				ShadowHatchingPatternTextureSampler;
 
 Buffer<float4>				DiffuseRampBufferSRV;
 Buffer<float4>				SpecularRampBufferSRV;
@@ -72,7 +72,7 @@ void ToonProfileUpdateAtlasTileCS(uint3 DispatchThreadId : SV_DispatchThreadID)
 		{
 			if (PixelPos.x == TOON_PROFILE_EVAL_PARAMS_X)
 			{
-				ToonProfileTexture[uint3(PixelPos, ToonProfileIndex)] = float4(DiffuseDitherPatternStrength, DiffuseDitherPatternSize, SpecularDitherPatternStrength, SpecularDitherPatternSize);
+				ToonProfileTexture[uint3(PixelPos, ToonProfileIndex)] = float4(DiffuseRampOffsetStrength, DiffuseRampOffsetSize, SpecularRampOffsetStrength, SpecularRampOffsetSize);
 			}
 			else if (PixelPos.x == TOON_PROFILE_EVAL_PARAMS2_X)
 			{
@@ -91,7 +91,7 @@ void ToonProfileUpdateAtlasTileCS(uint3 DispatchThreadId : SV_DispatchThreadID)
 				ToonProfileTexture[uint3(PixelPos, ToonProfileIndex)] = DefaultColor;
 			}
 		}
-		else if (PixelPos.y >= uint(TOON_PROFILE_TEXTURES_Y + TOON_PROFILE_TEXTURE_DITHERPATTERN_IDX * ProfileResolution) && PixelPos.y < uint(TOON_PROFILE_TEXTURES_Y + (TOON_PROFILE_TEXTURE_DITHERPATTERN_IDX+1) * ProfileResolution))
+		else if (PixelPos.y >= uint(TOON_PROFILE_TEXTURES_Y + TOON_PROFILE_TEXTURE_RAMPOFFSET_IDX * ProfileResolution) && PixelPos.y < uint(TOON_PROFILE_TEXTURES_Y + (TOON_PROFILE_TEXTURE_RAMPOFFSET_IDX +1) * ProfileResolution))
 		{
 		/*
 			For a Toon profile of full resolution N, we render an edge of 1 texel all around to make sure the texture is repeatable, marked `e` below.
@@ -116,7 +116,7 @@ void ToonProfileUpdateAtlasTileCS(uint3 DispatchThreadId : SV_DispatchThreadID)
 			| e | e | e | e | e | e | e | e |
 		0,1 --------------------------------- 1,1
 		*/
-			const float2 FullOrigin = float2(0.0f, TOON_PROFILE_TEXTURES_Y + TOON_PROFILE_TEXTURE_DITHERPATTERN_IDX * ProfileResolution);
+			const float2 FullOrigin = float2(0.0f, TOON_PROFILE_TEXTURES_Y + TOON_PROFILE_TEXTURE_RAMPOFFSET_IDX * ProfileResolution);
 			const float2 BakeOrigin = FullOrigin + 1.0f;
 
 			const float  BakeResolution = ProfileResolution - 2.0f;
@@ -125,13 +125,13 @@ void ToonProfileUpdateAtlasTileCS(uint3 DispatchThreadId : SV_DispatchThreadID)
 			const float2 FullCoord  = float2(PixelPos.xy) + 0.5f.xx;
 			const float2 FullUVs    = (FullCoord - BakeOrigin) * BakeResolutionInv;
 
-			const float DiffuseDitherPatternValue  = DiffuseDitherPattern.SampleLevel (DiffuseDitherPatternSampler,  FullUVs, 0.0f).x;
-			const float SpecularDitherPatternValue = SpecularDitherPattern.SampleLevel(SpecularDitherPatternSampler, FullUVs, 0.0f).y;
-			ToonProfileTexture[uint3(PixelPos, ToonProfileIndex)] = float4(DiffuseDitherPatternValue, SpecularDitherPatternValue, 0.0f, 0.0f);
+			const float DiffuseRampOffsetTextureValue  = DiffuseRampOffsetTexture.SampleLevel (DiffuseRampOffsetTextureSampler,  FullUVs, 0.0f).x;
+			const float SpecularRampOffsetTextureValue = SpecularRampOffsetTexture.SampleLevel(SpecularRampOffsetTextureSampler, FullUVs, 0.0f).y;
+			ToonProfileTexture[uint3(PixelPos, ToonProfileIndex)] = float4(DiffuseRampOffsetTextureValue, SpecularRampOffsetTextureValue, 0.0f, 0.0f);
 		}
 		else if (PixelPos.y >= uint(TOON_PROFILE_TEXTURES_Y + TOON_PROFILE_TEXTURE_SHADOWHATCHING_IDX * ProfileResolution) && PixelPos.y < uint(TOON_PROFILE_TEXTURES_Y + (TOON_PROFILE_TEXTURE_SHADOWHATCHING_IDX + 1) * ProfileResolution))
 		{
-			// This texture is stored with an edge for filtering similarly to TOON_PROFILE_TEXTURE_DITHERPATTERN_IDX
+			// This texture is stored with an edge for filtering similarly to TOON_PROFILE_TEXTURE_RAMPOFFSET_IDX
 			const float2 FullOrigin = float2(0.0f, TOON_PROFILE_TEXTURES_Y + TOON_PROFILE_TEXTURE_SHADOWHATCHING_IDX * ProfileResolution);
 			const float2 BakeOrigin = FullOrigin + 1.0f;
 
@@ -141,7 +141,7 @@ void ToonProfileUpdateAtlasTileCS(uint3 DispatchThreadId : SV_DispatchThreadID)
 			const float2 FullCoord  = float2(PixelPos.xy) + 0.5f.xx;
 			const float2 FullUVs    = (FullCoord - BakeOrigin) * BakeResolutionInv;
 
-			const float4 ShadowHatchingPatternValues = ShadowHatchingPattern.SampleLevel(ShadowHatchingPatternSampler, FullUVs, 0.0f).rgba;
+			const float4 ShadowHatchingPatternValues = ShadowHatchingPatternTexture.SampleLevel(ShadowHatchingPatternTextureSampler, FullUVs, 0.0f).rgba;
 			ToonProfileTexture[uint3(PixelPos, ToonProfileIndex)] = ShadowHatchingPatternValues;
 		}
 		else
diff --git a/Engine/Shaders/Private/ToonProfileCommon.ush b/Engine/Shaders/Private/ToonProfileCommon.ush
--- a/Engine/Shaders/Private/ToonProfileCommon.ush
+++ b/Engine/Shaders/Private/ToonProfileCommon.ush
@@ -31,7 +31,7 @@ float EvaluateToonProfileSpecularIndirectRamp(in float ToonProfileId, in float L
 	return EvaluateToonProfileRamp(ToonProfileId, Level, TOON_PROFILE_SPEC_IND_RAMP_Y).x;
 }
 
-float EvaluateToonProfileShadowRamp(in float ToonProfileId, in float Transmittance)
+float EvaluateToonProfileShadowHatchingPatternDistributionRamp(in float ToonProfileId, in float Transmittance)
 {
 	return EvaluateToonProfileRamp(ToonProfileId, Transmittance, TOON_PROFILE_SHADOW_RAMP_Y).r;
 }
@@ -40,10 +40,10 @@ float EvaluateToonProfileShadowRamp(in float ToonProfileId, in float Transmittan
 
 struct FToonProfileEvaluateSettings
 {
-	float DiffuseDitherPatternStrength;
-	float DiffuseDitherPatternSize;
-	float SpecularDitherPatternStrength;
-	float SpecularDitherPatternSize;
+	float DiffuseRampOffsetStrength;
+	float DiffuseRampOffsetSize;
+	float SpecularRampOffsetStrength;
+	float SpecularRampOffsetSize;
 	float bUnifiedDiffuseIncludeShadow;
 };
 
@@ -57,10 +57,10 @@ FToonProfileEvaluateSettings GetToonProfileEvaluateSettings(in float ToonProfile
 	float4 Data = View.ToonProfileTexture.SampleLevel(View.ToonProfileSampler, ToonProfileCoord, 0);
 
 	FToonProfileEvaluateSettings Ret;
-	Ret.DiffuseDitherPatternStrength	= Data.x;
-	Ret.DiffuseDitherPatternSize		= Data.y;
-	Ret.SpecularDitherPatternStrength	= Data.z;
-	Ret.SpecularDitherPatternSize		= Data.w;
+	Ret.DiffuseRampOffsetStrength	= Data.x;
+	Ret.DiffuseRampOffsetSize		= Data.y;
+	Ret.SpecularRampOffsetStrength	= Data.z;
+	Ret.SpecularRampOffsetSize		= Data.w;
 
 	const float3 ToonProfileCoord2 = float3(
 		(TOON_PROFILE_EVAL_PARAMS2_X + 0.5f) * View.ToonProfileTextureSizeAndInvSize.z,
@@ -74,13 +74,13 @@ FToonProfileEvaluateSettings GetToonProfileEvaluateSettings(in float ToonProfile
 
 float EvaluateToonProfileDiffuseDither(in float ToonProfileId, in float2 PatternUVs, in FToonProfileEvaluateSettings ToonProfileEvaluateSettings)
 {
-	if (ToonProfileEvaluateSettings.DiffuseDitherPatternStrength > 0.0f)
+	if (ToonProfileEvaluateSettings.DiffuseRampOffsetStrength > 0.0f)
 	{
 		// See ToonProfile.usf for the baking/sampling description.
 
-		PatternUVs = frac(PatternUVs * ToonProfileEvaluateSettings.DiffuseDitherPatternSize);
+		PatternUVs = frac(PatternUVs * ToonProfileEvaluateSettings.DiffuseRampOffsetSize);
 
-		const float2 FullOrigin = float2(0.0f, TOON_PROFILE_TEXTURES_Y + TOON_PROFILE_TEXTURE_DITHERPATTERN_IDX * View.ToonProfileDitherPatternSizeAndInvSize.x);
+		const float2 FullOrigin = float2(0.0f, TOON_PROFILE_TEXTURES_Y + TOON_PROFILE_TEXTURE_RAMPOFFSET_IDX * View.ToonProfileDitherPatternSizeAndInvSize.x);
 		const float2 BakeOrigin = FullOrigin + 1.0f;
 		
 		const float2 BakeCoords        = BakeOrigin + PatternUVs * View.ToonProfileDitherPatternSizeAndInvSize.x;
@@ -90,20 +90,20 @@ float EvaluateToonProfileDiffuseDither(in float ToonProfileId, in float2 Pattern
 			ToonProfileId
 		);
 
-		return ToonProfileEvaluateSettings.DiffuseDitherPatternStrength * (View.ToonProfileTexture.SampleLevel(View.ToonProfileSampler, ToonProfileCoord, 0).x * 2.0f - 1.0f);
+		return ToonProfileEvaluateSettings.DiffuseRampOffsetStrength * (View.ToonProfileTexture.SampleLevel(View.ToonProfileSampler, ToonProfileCoord, 0).x * 2.0f - 1.0f);
 	}
 	return 0.0f;
 }
 
 float EvaluateToonProfileSpecularDither(in float ToonProfileId, in float2 PatternUVs, in FToonProfileEvaluateSettings ToonProfileEvaluateSettings)
 {
-	if (ToonProfileEvaluateSettings.SpecularDitherPatternStrength > 0.0f)
+	if (ToonProfileEvaluateSettings.SpecularRampOffsetStrength > 0.0f)
 	{
 		// See ToonProfile.usf for the baking/sampling description.
 
-		PatternUVs = frac(PatternUVs * ToonProfileEvaluateSettings.SpecularDitherPatternSize);
+		PatternUVs = frac(PatternUVs * ToonProfileEvaluateSettings.SpecularRampOffsetSize);
 
-		const float2 FullOrigin = float2(0.0f, TOON_PROFILE_TEXTURES_Y + TOON_PROFILE_TEXTURE_DITHERPATTERN_IDX * View.ToonProfileTextureSizeAndInvSize.x);
+		const float2 FullOrigin = float2(0.0f, TOON_PROFILE_TEXTURES_Y + TOON_PROFILE_TEXTURE_RAMPOFFSET_IDX * View.ToonProfileTextureSizeAndInvSize.x);
 		const float2 BakeOrigin = FullOrigin + 1.0f;
 		
 		const float2 BakeCoords = BakeOrigin + PatternUVs * View.ToonProfileDitherPatternSizeAndInvSize.x;
@@ -113,7 +113,7 @@ float EvaluateToonProfileSpecularDither(in float ToonProfileId, in float2 Patter
 			ToonProfileId
 		);
 
-		return ToonProfileEvaluateSettings.SpecularDitherPatternStrength * (View.ToonProfileTexture.SampleLevel(View.ToonProfileSampler, ToonProfileCoord, 0).y * 2.0f - 1.0f);
+		return ToonProfileEvaluateSettings.SpecularRampOffsetStrength * (View.ToonProfileTexture.SampleLevel(View.ToonProfileSampler, ToonProfileCoord, 0).y * 2.0f - 1.0f);
 	}
 	return 0.0f;
 }
@@ -171,20 +171,32 @@ float EvaluateToonProfileShadowHatching(in float ToonProfileId, in float2 Patter
 
 		const float4 ShadowHatching = View.ToonProfileTexture.SampleLevel(View.ToonProfileSampler, ToonProfileCoord, 0).rgba;
 
-		const float TransmittanceCurve = EvaluateToonProfileShadowRamp(ToonProfileId, Transmittance).x;
+		const float Distribution = EvaluateToonProfileShadowHatchingPatternDistributionRamp(ToonProfileId, Transmittance).x;
+
+		// SUBSTRATE_TODO optimise
 		const float Step = 0.25f;
-		const float StepInv = 1.0f / Step;
-		const float4 Step4 = float4(1.0f, 0.75f, 0.50f, 0.25f);
-	
-		const float curveInputX = TransmittanceCurve;
-		float4 LerpFactors = 1.0f-saturate(abs(curveInputX - Step4) * StepInv);
-	
-		// Fade to white from 0.95 to 1.0. SUBSTRATE_TODOO make that better parameter.
-		const float FadeToWhite = saturate(1.0f - (curveInputX - 0.95f) / 0.05);
-	
-		return lerp(1.0f, dot(ShadowHatching, LerpFactors), ToonProfileShadowSettings.ShadowHatchingPatternStrength * FadeToWhite);
+		float FinalShadow = 0.0;
+		if (Distribution <= 0.25)
+		{
+			FinalShadow = lerp(ShadowHatching.a, ShadowHatching.b, saturate((Distribution - 0.00) / 0.25));
+		}
+		else if (Distribution <= 0.5)
+		{
+			FinalShadow = lerp(ShadowHatching.b, ShadowHatching.g, saturate((Distribution - 0.25) / 0.25));
+		}
+		else if (Distribution <= 0.75)
+		{
+			FinalShadow = lerp(ShadowHatching.g, ShadowHatching.r, saturate((Distribution - 0.50) / 0.25));
+		}
+		else if (Distribution <= 1.0)
+		{
+			FinalShadow = lerp(ShadowHatching.r, 1.0f            , saturate((Distribution - 0.75) / 0.25));
+		}
+
+		const float FadeToWhite = ToonProfileShadowSettings.ShadowHatchingPatternStrength;
+		return lerp(1.0f, FinalShadow, FadeToWhite);
 	}
-	return 1.0f;
+	return Transmittance;
 }
 
 ////////////////////////////////////////////////////////////////////////////////////////////////////
diff --git a/Engine/Shaders/Shared/ToonProfileDefinitions.h b/Engine/Shaders/Shared/ToonProfileDefinitions.h
--- a/Engine/Shaders/Shared/ToonProfileDefinitions.h
+++ b/Engine/Shaders/Shared/ToonProfileDefinitions.h
@@ -14,6 +14,6 @@
 #define TOON_PROFILE_SHADOW_PARAMS_X	2
 #define TOON_PROFILE_GI_PARAMS_X		3
 
-#define TOON_PROFILE_TEXTURE_DITHERPATTERN_IDX		0
+#define TOON_PROFILE_TEXTURE_RAMPOFFSET_IDX			0
 #define TOON_PROFILE_TEXTURE_SHADOWHATCHING_IDX		1
 #define TOON_PROFILE_TEXTURE_COUNT					2
diff --git a/Engine/Source/Runtime/Engine/Classes/Engine/ToonProfile.h b/Engine/Source/Runtime/Engine/Classes/Engine/ToonProfile.h
--- a/Engine/Source/Runtime/Engine/Classes/Engine/ToonProfile.h
+++ b/Engine/Source/Runtime/Engine/Classes/Engine/ToonProfile.h
@@ -7,6 +7,7 @@
 #include "Misc/Guid.h"
 #include "RenderResource.h"
 #include "Curves/CurveLinearColor.h"
+#include "Curves/CurveFloat.h"
 #include "ToonProfile.generated.h"
 
 #define MAX_TOON_PROFILE_COUNT 256
@@ -31,22 +32,22 @@ struct FToonProfileStruct
 	FRuntimeCurveLinearColor DiffuseRamp;
 
 	/**
-	 * The diffuse texture dithering pattern. It is recommended for it to be grey scale level in [0, 1], internally remapped to [-1, 1] for positive and negative offset.
+	 * The texture containing ramp offset for diffuse. It is recommended for it to be grey scale level in [0, 1], internally remapped to [-1, 1] for positive and negative offset. A distance field is recommended.
 	 */
 	UPROPERTY(Category = "Diffuse", EditAnywhere)
-	TObjectPtr<UTexture> DiffuseDitherPattern;
+	TObjectPtr<UTexture> DiffuseRampOffsetTexture;
 
 	/**
-	 * The strength of the offset applied by the dither pattern to the diffuse ramp.
+	 * The strength of the offset applied by the ramp offset texture to the diffuse ramp.
 	 */
 	UPROPERTY(Category = "Diffuse", EditAnywhere, meta = (UIMin = 0.0, UIMax = 1.0, ClampMin = 0.0))
-	float DiffuseDitherPatternStrength;
+	float DiffuseRampOffsetStrength;
 
 	/**
-	 * The size on screen of the dithering pattern.
+	 * The multiplier applied on UVs before sampling.
 	 */
-	UPROPERTY(Category = "Diffuse", EditAnywhere, meta = (UIMin = 1.0, UIMax = 10.0, ClampMin = 1.0))
-	float DiffuseDitherPatternSize;
+	UPROPERTY(Category = "Diffuse", EditAnywhere, meta = (UIMin = 1.0, UIMax = 100.0, ClampMin = 1.0, SliderExponent = 2.0))
+	float DiffuseRampOffsetSize;
 
 	/**
 	* Defines the specular ramp.
@@ -55,22 +56,22 @@ struct FToonProfileStruct
 	FRuntimeCurveLinearColor SpecularRamp;
 
 	/**
-	 * The specular texture dithering pattern. It is recommended for it to be grey scale level in [0, 1], internally remapped to [-1, 1] for positive and negative offset.
+	 * The texture containing ramp offset for specular. It is recommended for it to be grey scale level in [0, 1], internally remapped to [-1, 1] for positive and negative offset. A distance field is recommended.
 	 */
 	UPROPERTY(Category = "Specular", EditAnywhere)
-	TObjectPtr<UTexture> SpecularDitherPattern;
+	TObjectPtr<UTexture> SpecularRampOffsetTexture;
 
 	/**
-	 * The strength of the offset applied by the dither pattern to the specular ramp.
+	 * The strength of the offset applied by the ramp offset texture to the specular ramp.
 	 */
 	UPROPERTY(Category = "Specular", EditAnywhere, meta = (UIMin = 0.0, UIMax = 1.0, ClampMin = 0.0))
-	float SpecularDitherPatternStrength;
+	float SpecularRampOffsetStrength;
 
 	/**
-	 * The size on screen of the dithering pattern.
+	 * The multiplier applied on UVs before sampling.
 	 */
-	UPROPERTY(Category = "Specular", EditAnywhere, meta = (UIMin = 1.0, UIMax = 10.0, ClampMin = 1.0))
-	float SpecularDitherPatternSize;
+	UPROPERTY(Category = "Specular", EditAnywhere, meta = (UIMin = 1.0, UIMax = 100.0, ClampMin = 1.0, SliderExponent = 2.0))
+	float SpecularRampOffsetSize;
 
 	/**
 	* Defines how deep light will penetrate inside the model for the shadow map occluder distance. 
@@ -80,22 +81,24 @@ struct FToonProfileStruct
 	float ShadowExtinctionCoefficient;
 
 	/**
-	* Defines the shadow ramp used to control how the four channels for the hatching pattern are distributed along the transmittance curve.
+	* Defines the curve distribution used to control how the four channels for the shadow hatching pattern texture are distributed along the evaluated shadow/transmittance.
+	* Shadow/Transmittance => Shadow hatching ramp evaluation => Value is used to interpolate the shadow hatching pattern RGBA channels (1.0 maps to white, 0.75 to R, 0.5 to G, 0.25 to B, 0.0 to A)
 	*/
 	UPROPERTY(Category = "Shadow", EditAnywhere, meta = (AllowZoomOutput = "false", ShowZoomButtons = "false", ViewMinInput = "0", ViewMaxInput = "1", ViewMinOutput = "0", ViewMaxOutput = "1", TimelineLength = "1", ShowInputGridNumbers = "false", ShowOutputGridNumbers = "false"))
-	FRuntimeCurveLinearColor ShadowRamp;
+	FRuntimeFloatCurve ShadowHatchingPatternDistributionRamp;
 
 	/**
-	 * The shadow texture defining the hatching pattern applied to the grey scale shadow. R applied first close to the lit front surface and progressively G, B and A channels lerped when away from the lit front surface.
-	 * Typically R, G, B and A would have light to bold hatching pattern.
+	 * The shadow texture defining the hatching pattern applied to the grey scale shadow. 
+	 * R applied first close onto the lit front surface and progressively G, B and A channels lerped when away from the lit front surface as the shadow/transmittance gets darker.
+	 * Typically R, G, B and A would have light to bold hatching bHasDiffuseDitherPatternTexturepattern.
 	 */
 	UPROPERTY(Category = "Shadow", EditAnywhere)
-	TObjectPtr<UTexture> ShadowHatchingPattern;
+	TObjectPtr<UTexture> ShadowHatchingPatternTexture;
 
 	/**
-	 * The size on screen of the shadow hatching pattern.
+	 * The multiplier applied on UVs before sampling.
 	 */
-	UPROPERTY(Category = "Shadow", EditAnywhere, meta = (UIMin = 1.0, UIMax = 10.0, ClampMin = 1.0))
+	UPROPERTY(Category = "Shadow", EditAnywhere, meta = (UIMin = 1.0, UIMax = 100.0, ClampMin = 1.0, SliderExponent = 2.0))
 	float ShadowHatchingPatternSize;
 
 	/**
@@ -120,7 +123,7 @@ struct FToonProfileStruct
 	* Defines the specular global illumination ramp that can be used to quantise the specular indirect reflections. Requires r.Substrate.Experimental.ToonReflectionQuantizationEnabled=1.
 	*/
 	UPROPERTY(Category = "Global Illumination", EditAnywhere, meta = (DisplayName = "Specular Indirect Ramp", AllowZoomOutput = "false", ShowZoomButtons = "false", ViewMinInput = "0", ViewMaxInput = "1", ViewMinOutput = "0", ViewMaxOutput = "1", TimelineLength = "1", ShowInputGridNumbers = "false", ShowOutputGridNumbers = "false"))
-	FRuntimeCurveLinearColor SpecularIndirectRamp;
+	FRuntimeFloatCurve SpecularIndirectRamp;
 
 	/**
 	 * The scale of the the specular indirect ramp application over the range of luminance. When 0, the ramp is not applied to specular indirect reflections. Requires r.Substrate.Experimental.ToonReflectionQuantizationEnabled=1.
@@ -129,7 +132,7 @@ struct FToonProfileStruct
 	float SpecularIndirectRampRepetition;
 
 	/**
-	 * Experimental. Requires r.Substrate.Experimental.ToonUnifiedDiffuse=1.
+	 * Experimental. Include shadows as part of the unified diffuse ramp evaluation. Requires r.Substrate.Experimental.ToonUnifiedDiffuse=1.
 	 */
 	UPROPERTY(Category = "Experimental", EditAnywhere, meta = (DisplayName = "Unified Diffuse Include Shadow"))
 	uint8 bUnifiedDiffuseIncludeShadow : 1;
@@ -168,6 +171,7 @@ class UToonProfile : public UObject
 	virtual void PostDuplicate(EDuplicateMode::Type DuplicateMode);
 #if WITH_EDITOR
 	virtual void PostEditChangeProperty(struct FPropertyChangedEvent& PropertyChangedEvent);
+	virtual bool CanEditChange(const FProperty* InProperty) const override;
 #endif
 	//~ End UObject Interface
 };
diff --git a/Engine/Source/Runtime/Engine/Private/Rendering/ToonProfile.cpp b/Engine/Source/Runtime/Engine/Private/Rendering/ToonProfile.cpp
--- a/Engine/Source/Runtime/Engine/Private/Rendering/ToonProfile.cpp
+++ b/Engine/Source/Runtime/Engine/Private/Rendering/ToonProfile.cpp
@@ -74,10 +74,10 @@ public:
 	BEGIN_SHADER_PARAMETER_STRUCT(FParameters, )
 		SHADER_PARAMETER(FIntPoint, ToonProfileTextureResolution)
 		SHADER_PARAMETER(uint32, ToonProfileIndex)
-		SHADER_PARAMETER(float, DiffuseDitherPatternStrength)
-		SHADER_PARAMETER(float, DiffuseDitherPatternSize)
-		SHADER_PARAMETER(float, SpecularDitherPatternStrength)
-		SHADER_PARAMETER(float, SpecularDitherPatternSize)
+		SHADER_PARAMETER(float, DiffuseRampOffsetStrength)
+		SHADER_PARAMETER(float, DiffuseRampOffsetSize)
+		SHADER_PARAMETER(float, SpecularRampOffsetStrength)
+		SHADER_PARAMETER(float, SpecularRampOffsetSize)
 		SHADER_PARAMETER(float, ShadowHatchingPatternSize)
 		SHADER_PARAMETER(float, ShadowHatchingPatternStrength)
 		SHADER_PARAMETER(float, DiffuseIndirectScale)
@@ -89,14 +89,14 @@ public:
 
 		SHADER_PARAMETER_RDG_TEXTURE_UAV(RWTexture2DArray, ToonProfileTexture)
 
-		SHADER_PARAMETER_RDG_TEXTURE(Texture2D, DiffuseDitherPattern)
-		SHADER_PARAMETER_SAMPLER(SamplerState, DiffuseDitherPatternSampler)
+		SHADER_PARAMETER_RDG_TEXTURE(Texture2D, DiffuseRampOffsetTexture)
+		SHADER_PARAMETER_SAMPLER(SamplerState, DiffuseRampOffsetTextureSampler)
 
-		SHADER_PARAMETER_RDG_TEXTURE(Texture2D, SpecularDitherPattern)
-		SHADER_PARAMETER_SAMPLER(SamplerState, SpecularDitherPatternSampler)
+		SHADER_PARAMETER_RDG_TEXTURE(Texture2D, SpecularRampOffsetTexture)
+		SHADER_PARAMETER_SAMPLER(SamplerState, SpecularRampOffsetTextureSampler)
 
-		SHADER_PARAMETER_RDG_TEXTURE(Texture2D, ShadowHatchingPattern)
-		SHADER_PARAMETER_SAMPLER(SamplerState, ShadowHatchingPatternSampler)
+		SHADER_PARAMETER_RDG_TEXTURE(Texture2D, ShadowHatchingPatternTexture)
+		SHADER_PARAMETER_SAMPLER(SamplerState, ShadowHatchingPatternTextureSampler)
 
 		SHADER_PARAMETER_RDG_BUFFER_SRV(Buffer, DiffuseRampBufferSRV)
 		SHADER_PARAMETER_RDG_BUFFER_SRV(Buffer, SpecularRampBufferSRV)
@@ -184,9 +184,9 @@ private:
 		FToonProfileStruct Settings;
 		const UToonProfile* Profile = nullptr; // Game thread pointer! Do not dereference, only for comparison.
 		FName ParameterName;
-		FIntPoint DiffuseDitherPatternCachedResolution = FIntPoint::ZeroValue;
-		FIntPoint SpecularDitherPatternCachedResolution = FIntPoint::ZeroValue;
-		FIntPoint ShadowHatchingPatternCachedResolution = FIntPoint::ZeroValue;
+		FIntPoint DiffuseRampOffsetTextureCachedResolution = FIntPoint::ZeroValue;
+		FIntPoint SpecularRampOffsetTextureCachedResolution = FIntPoint::ZeroValue;
+		FIntPoint ShadowHatchingPatternTextureCachedResolution = FIntPoint::ZeroValue;
 	};
 	TArray<FToonProfileEntry> ToonProfileEntries;
 	const FTextureReference* Texture = nullptr;
@@ -344,9 +344,9 @@ IPooledRenderTarget* FToonProfileTextureManager::GetAtlasTexture(FRDGBuilder& Gr
 			}
 		};
 
-		CheckIfTextureHasStreamedIn(Data.DiffuseDitherPattern  ? Data.DiffuseDitherPattern->GetResource()  : nullptr, Entry.DiffuseDitherPatternCachedResolution);
-		CheckIfTextureHasStreamedIn(Data.SpecularDitherPattern ? Data.SpecularDitherPattern->GetResource() : nullptr, Entry.SpecularDitherPatternCachedResolution);
-		CheckIfTextureHasStreamedIn(Data.ShadowHatchingPattern ? Data.ShadowHatchingPattern->GetResource() : nullptr, Entry.ShadowHatchingPatternCachedResolution);
+		CheckIfTextureHasStreamedIn(Data.DiffuseRampOffsetTexture		? Data.DiffuseRampOffsetTexture->GetResource()		: nullptr, Entry.DiffuseRampOffsetTextureCachedResolution);
+		CheckIfTextureHasStreamedIn(Data.SpecularRampOffsetTexture		? Data.SpecularRampOffsetTexture->GetResource()		: nullptr, Entry.SpecularRampOffsetTextureCachedResolution);
+		CheckIfTextureHasStreamedIn(Data.ShadowHatchingPatternTexture	? Data.ShadowHatchingPatternTexture->GetResource()	: nullptr, Entry.ShadowHatchingPatternTextureCachedResolution);
 	}
 
 	if (bForceUpdate)
@@ -365,7 +365,7 @@ IPooledRenderTarget* FToonProfileTextureManager::GetAtlasTexture(FRDGBuilder& Gr
 
 		// See ToonProfileCommon.ush
 		const uint32 ResolutionX = ToonProfileResolution;
-		const uint32 ResolutionY = TOON_PROFILE_TEXTURES_Y + TOON_PROFILE_TEXTURE_COUNT * ToonProfileResolution; // row for ramps, parameters and 2d dithering patterns.
+		const uint32 ResolutionY = TOON_PROFILE_TEXTURES_Y + TOON_PROFILE_TEXTURE_COUNT * ToonProfileResolution; // row for ramps, parameters and textures.
 
 		// 1. Create atlas texture
 		const FIntPoint ToonProfileResolution2D = FIntPoint(ToonProfileResolution, ResolutionY);
@@ -387,17 +387,17 @@ IPooledRenderTarget* FToonProfileTextureManager::GetAtlasTexture(FRDGBuilder& Gr
 		{
 			const FToonProfileStruct& Data = ToonProfileEntries[LayerIt].Settings;
 
-			FTextureResource* DiffuseDitherPatternTextureResource = Data.DiffuseDitherPattern ? Data.DiffuseDitherPattern->GetResource() : nullptr;
-			FRHITexture* DiffuseDitherPatternTexture2DRHI = DiffuseDitherPatternTextureResource ? DiffuseDitherPatternTextureResource->GetTexture2DRHI() : nullptr;
-			bool bHasDiffuseDitherPatternTexture = DiffuseDitherPatternTexture2DRHI != nullptr;
-			DiffuseDitherPatternTexture2DRHI = DiffuseDitherPatternTexture2DRHI ? DiffuseDitherPatternTexture2DRHI : GBlackTexture->TextureRHI->GetTexture2D();
+			FTextureResource* DiffuseRampOffsetTextureResource = Data.DiffuseRampOffsetTexture ? Data.DiffuseRampOffsetTexture->GetResource() : nullptr;
+			FRHITexture* DiffuseRampOffsetTexture2DRHI = DiffuseRampOffsetTextureResource ? DiffuseRampOffsetTextureResource->GetTexture2DRHI() : nullptr;
+			bool bHasDiffuseRampOffsetTexture = DiffuseRampOffsetTexture2DRHI != nullptr;
+			DiffuseRampOffsetTexture2DRHI = DiffuseRampOffsetTexture2DRHI ? DiffuseRampOffsetTexture2DRHI : GBlackTexture->TextureRHI->GetTexture2D();
 
-			FTextureResource* SpecularDitherPatternTextureResource = Data.SpecularDitherPattern ? Data.SpecularDitherPattern->GetResource() : nullptr;
-			FRHITexture* SpecularDitherPatternTexture2DRHI = SpecularDitherPatternTextureResource ? SpecularDitherPatternTextureResource->GetTexture2DRHI() : nullptr;
-			bool bHasSpecularDitherPatternTexture = SpecularDitherPatternTexture2DRHI != nullptr;
-			SpecularDitherPatternTexture2DRHI = SpecularDitherPatternTexture2DRHI ? SpecularDitherPatternTexture2DRHI : GBlackTexture->TextureRHI->GetTexture2D();
+			FTextureResource* SpecularRampOffsetTextureResource = Data.SpecularRampOffsetTexture ? Data.SpecularRampOffsetTexture->GetResource() : nullptr;
+			FRHITexture* SpecularRampOffsetTexture2DRHI = SpecularRampOffsetTextureResource ? SpecularRampOffsetTextureResource->GetTexture2DRHI() : nullptr;
+			bool bHasSpecularRampOffsetTexture = SpecularRampOffsetTexture2DRHI != nullptr;
+			SpecularRampOffsetTexture2DRHI = SpecularRampOffsetTexture2DRHI ? SpecularRampOffsetTexture2DRHI : GBlackTexture->TextureRHI->GetTexture2D();
 
-			FTextureResource* ShadowHatchingPatternTextureResource = Data.ShadowHatchingPattern ? Data.ShadowHatchingPattern->GetResource() : nullptr;
+			FTextureResource* ShadowHatchingPatternTextureResource = Data.ShadowHatchingPatternTexture ? Data.ShadowHatchingPatternTexture->GetResource() : nullptr;
 			FRHITexture* ShadowHatchingPatternTexture2DRHI = ShadowHatchingPatternTextureResource ? ShadowHatchingPatternTextureResource->GetTexture2DRHI() : nullptr;
 			bool bHasShadowHatchingPatternTexture = ShadowHatchingPatternTexture2DRHI != nullptr;
 			ShadowHatchingPatternTexture2DRHI = ShadowHatchingPatternTexture2DRHI ? ShadowHatchingPatternTexture2DRHI : GBlackTexture->TextureRHI->GetTexture2D();
@@ -406,10 +406,10 @@ IPooledRenderTarget* FToonProfileTextureManager::GetAtlasTexture(FRDGBuilder& Gr
 			PassParameters->ToonProfileTexture = ToonProfileUAVSkipBarrier;
 			PassParameters->ToonProfileTextureResolution = ToonProfileTexture->Desc.Extent;
 			PassParameters->ToonProfileIndex = LayerIt;
-			PassParameters->DiffuseDitherPatternStrength = bHasDiffuseDitherPatternTexture ? Data.DiffuseDitherPatternStrength : 0.0f;
-			PassParameters->DiffuseDitherPatternSize = Data.DiffuseDitherPatternSize;
-			PassParameters->SpecularDitherPatternStrength = bHasSpecularDitherPatternTexture ? Data.SpecularDitherPatternStrength : 0.0f;
-			PassParameters->SpecularDitherPatternSize = Data.SpecularDitherPatternSize;
+			PassParameters->DiffuseRampOffsetStrength = bHasDiffuseRampOffsetTexture ? Data.DiffuseRampOffsetStrength : 0.0f;
+			PassParameters->DiffuseRampOffsetSize = Data.DiffuseRampOffsetSize;
+			PassParameters->SpecularRampOffsetStrength = bHasSpecularRampOffsetTexture ? Data.SpecularRampOffsetStrength : 0.0f;
+			PassParameters->SpecularRampOffsetSize = Data.SpecularRampOffsetSize;
 			PassParameters->ToonProfileShadowExtinctionCoefficientInInvCentimeter = Data.GetShadowExtinctionCoefficientInInvCentimeter();
 			PassParameters->ShadowHatchingPatternSize = Data.ShadowHatchingPatternSize;
 			PassParameters->ShadowHatchingPatternStrength = bHasShadowHatchingPatternTexture ? Data.ShadowHatchingPatternStrength : 0.0f;
@@ -419,7 +419,7 @@ IPooledRenderTarget* FToonProfileTextureManager::GetAtlasTexture(FRDGBuilder& Gr
 			PassParameters->bUnifiedDiffuseIncludeShadow = Data.bUnifiedDiffuseIncludeShadow ? 1.0f : 0.0f;
 			PassParameters->SpecularIndirectRampRepetition = Data.SpecularIndirectRampRepetition;
 
-			auto CreateBufferFRomCurve = [&](FRuntimeCurveLinearColor ColorCurve)
+			auto CreateBufferFromColorCurve = [&](const FRuntimeCurveLinearColor& ColorCurve)
 			{
 					const uint32 ComponentsPerElement = 4;
 				const uint32 BytesPerComponent = sizeof(FFloat16);
@@ -441,27 +441,49 @@ IPooledRenderTarget* FToonProfileTextureManager::GetAtlasTexture(FRDGBuilder& Gr
 
 				return ColorCurveBuffer;
 			};
-
-			FRDGBufferRef DiffuseRampBuffer = CreateBufferFRomCurve(Data.DiffuseRamp);
+			auto CreateBufferFromFloatCurve = [&](const FRuntimeFloatCurve& FloatCurve)
+				{
+					const uint32 ComponentsPerElement = 4;
+					const uint32 BytesPerComponent = sizeof(FFloat16);
+					const uint32 BytesPerElement = ComponentsPerElement * BytesPerComponent;
+
+					TArray<FFloat16> ColorCurveData;
+					ColorCurveData.SetNum(ToonProfileResolution * ComponentsPerElement);
+					for (uint32 It = 0; It < ToonProfileResolution; ++It)
+					{
+						float CurveSample = FloatCurve.GetRichCurveConst()->Eval(float(It + 0.5f) / ToonProfileResolution);
+						ColorCurveData[It * 4 + 0] = CurveSample;
+						ColorCurveData[It * 4 + 1] = CurveSample;
+						ColorCurveData[It * 4 + 2] = CurveSample;
+						ColorCurveData[It * 4 + 3] = CurveSample;
+					}
+
+					FRDGBufferRef ColorCurveBuffer = GraphBuilder.CreateBuffer(FRDGBufferDesc::CreateBufferDesc(BytesPerElement, ToonProfileResolution), TEXT("ColorCurveBuffer"), ERDGBufferFlags::MultiFrame);
+					GraphBuilder.QueueBufferUpload(ColorCurveBuffer, ColorCurveData.GetData(), BytesPerComponent * ColorCurveData.Num(), ERDGInitialDataFlags::None);
+
+					return ColorCurveBuffer;
+				};
+
+			FRDGBufferRef DiffuseRampBuffer = CreateBufferFromColorCurve(Data.DiffuseRamp);
 			PassParameters->DiffuseRampBufferSRV = GraphBuilder.CreateSRV(DiffuseRampBuffer, Format);
 
-			FRDGBufferRef SpecularRampBuffer = CreateBufferFRomCurve(Data.SpecularRamp);
+			FRDGBufferRef SpecularRampBuffer = CreateBufferFromColorCurve(Data.SpecularRamp);
 			PassParameters->SpecularRampBufferSRV = GraphBuilder.CreateSRV(SpecularRampBuffer, Format);
 
-			FRDGBufferRef SpecularIndirectRampBuffer = CreateBufferFRomCurve(Data.SpecularIndirectRamp);
+			FRDGBufferRef SpecularIndirectRampBuffer = CreateBufferFromFloatCurve(Data.SpecularIndirectRamp);
 			PassParameters->SpecularIndirectRampBufferSRV = GraphBuilder.CreateSRV(SpecularIndirectRampBuffer, Format);
 
-			FRDGBufferRef ShadowRampBuffer = CreateBufferFRomCurve(Data.ShadowRamp);
+			FRDGBufferRef ShadowRampBuffer = CreateBufferFromFloatCurve(Data.ShadowHatchingPatternDistributionRamp);
 			PassParameters->ShadowRampBufferSRV = GraphBuilder.CreateSRV(ShadowRampBuffer, Format);
 
-			PassParameters->DiffuseDitherPattern = GraphBuilder.RegisterExternalTexture(CreateRenderTarget(DiffuseDitherPatternTexture2DRHI, TEXT("DiffuseDitherPattern")));
-			PassParameters->DiffuseDitherPatternSampler = TStaticSamplerState<SF_Trilinear, AM_Wrap, AM_Wrap>::GetRHI();
+			PassParameters->DiffuseRampOffsetTexture = GraphBuilder.RegisterExternalTexture(CreateRenderTarget(DiffuseRampOffsetTexture2DRHI, TEXT("DiffuseRampOffsetTexture")));
+			PassParameters->DiffuseRampOffsetTextureSampler = TStaticSamplerState<SF_Trilinear, AM_Wrap, AM_Wrap>::GetRHI();
 
-			PassParameters->SpecularDitherPattern = GraphBuilder.RegisterExternalTexture(CreateRenderTarget(SpecularDitherPatternTexture2DRHI, TEXT("SpecularDitherPattern")));
-			PassParameters->SpecularDitherPatternSampler = TStaticSamplerState<SF_Trilinear, AM_Wrap, AM_Wrap>::GetRHI();
+			PassParameters->SpecularRampOffsetTexture = GraphBuilder.RegisterExternalTexture(CreateRenderTarget(SpecularRampOffsetTexture2DRHI, TEXT("SpecularRampOffsetTexture")));
+			PassParameters->SpecularRampOffsetTextureSampler = TStaticSamplerState<SF_Trilinear, AM_Wrap, AM_Wrap>::GetRHI();
 
-			PassParameters->ShadowHatchingPattern = GraphBuilder.RegisterExternalTexture(CreateRenderTarget(ShadowHatchingPatternTexture2DRHI, TEXT("ShadowHatchingPattern")));
-			PassParameters->ShadowHatchingPatternSampler = TStaticSamplerState<SF_Trilinear, AM_Wrap, AM_Wrap>::GetRHI();
+			PassParameters->ShadowHatchingPatternTexture = GraphBuilder.RegisterExternalTexture(CreateRenderTarget(ShadowHatchingPatternTexture2DRHI, TEXT("ShadowHatchingPattern")));
+			PassParameters->ShadowHatchingPatternTextureSampler = TStaticSamplerState<SF_Trilinear, AM_Wrap, AM_Wrap>::GetRHI();
 
 			FToonProfileUpdateAtlasTileCS::FPermutationDomain PermutationVector;
 			TShaderMapRef<FToonProfileUpdateAtlasTileCS> Shader(GlobalShaderMap, PermutationVector);
@@ -511,42 +533,46 @@ FToonProfileIdQuery FToonProfileTextureManager::FindAllocationId(const UToonProf
 
 FToonProfileStruct::FToonProfileStruct()
 {
-	auto AddPointToCurve = [](FRuntimeCurveLinearColor& LinearColorCurve, float Time, float Value)
+	auto AddPointToColorCurve = [](FRuntimeCurveLinearColor& LinearColorCurve, float Time, float Value)
 	{
 		LinearColorCurve.ColorCurves[0].AddKey(Time, Value);
 		LinearColorCurve.ColorCurves[1].AddKey(Time, Value);
 		LinearColorCurve.ColorCurves[2].AddKey(Time, Value);
 	};
+	auto AddPointToFloatCurve = [](FRuntimeFloatCurve& FloatCurve, float Time, float Value)
+		{
+			FloatCurve.EditorCurveData.AddKey(Time, Value);
+		};
 
-	AddPointToCurve(DiffuseRamp, 0.00, 0.00);
-	AddPointToCurve(DiffuseRamp, 0.25, 0.00);
-	AddPointToCurve(DiffuseRamp, 0.26, 0.25);
-	AddPointToCurve(DiffuseRamp, 0.50, 0.25);
-	AddPointToCurve(DiffuseRamp, 0.51, 0.50);
-	AddPointToCurve(DiffuseRamp, 0.75, 0.50);
-	AddPointToCurve(DiffuseRamp, 0.76, 0.75);
-	AddPointToCurve(DiffuseRamp, 0.99, 0.75);
-	AddPointToCurve(DiffuseRamp, 1.00, 1.00);
+	AddPointToColorCurve(DiffuseRamp, 0.00, 0.00);
+	AddPointToColorCurve(DiffuseRamp, 0.25, 0.00);
+	AddPointToColorCurve(DiffuseRamp, 0.26, 0.25);
+	AddPointToColorCurve(DiffuseRamp, 0.50, 0.25);
+	AddPointToColorCurve(DiffuseRamp, 0.51, 0.50);
+	AddPointToColorCurve(DiffuseRamp, 0.75, 0.50);
+	AddPointToColorCurve(DiffuseRamp, 0.76, 0.75);
+	AddPointToColorCurve(DiffuseRamp, 0.99, 0.75);
+	AddPointToColorCurve(DiffuseRamp, 1.00, 1.00);
 
-	AddPointToCurve(SpecularRamp, 0.00, 0.0);
-	AddPointToCurve(SpecularRamp, 0.50, 0.0);
-	AddPointToCurve(SpecularRamp, 0.51, 1.0);
-	AddPointToCurve(SpecularRamp, 1.00, 1.0);
+	AddPointToColorCurve(SpecularRamp, 0.00, 0.0);
+	AddPointToColorCurve(SpecularRamp, 0.50, 0.0);
+	AddPointToColorCurve(SpecularRamp, 0.51, 1.0);
+	AddPointToColorCurve(SpecularRamp, 1.00, 1.0);
 
-	AddPointToCurve(SpecularIndirectRamp, 0.0, 0.0);
-	AddPointToCurve(SpecularIndirectRamp, 1.0, 1.0);
+	AddPointToFloatCurve(SpecularIndirectRamp, 0.0, 0.0);
+	AddPointToFloatCurve(SpecularIndirectRamp, 1.0, 1.0);
 
-	DiffuseDitherPattern = nullptr;
-	DiffuseDitherPatternStrength = 0.0f;	// By default, skip sampling dither pattern.
-	DiffuseDitherPatternSize = 1.0f;
+	DiffuseRampOffsetTexture = nullptr;
+	DiffuseRampOffsetStrength = 0.0f;	// By default, skip.
+	DiffuseRampOffsetSize = 1.0f;
 
-	SpecularDitherPattern = nullptr;
-	SpecularDitherPatternStrength = 0.0f;	// By default, skip sampling dither pattern.
-	SpecularDitherPatternSize = 1.0f;
+	SpecularRampOffsetTexture = nullptr;
+	SpecularRampOffsetStrength = 0.0f;	// By default, skip.
+	SpecularRampOffsetSize = 1.0f;
 
 	ShadowExtinctionCoefficient = 10.0f;	// Default extinction coefficients that work well.
 
-	ShadowHatchingPattern = nullptr;
+	ShadowHatchingPatternTexture = nullptr;
 	ShadowHatchingPatternSize = 1.0f;
 	ShadowHatchingPatternStrength = 1.0f;
 
@@ -556,9 +582,8 @@ FToonProfileStruct::FToonProfileStruct()
 
 	bUnifiedDiffuseIncludeShadow = 0;
 
-	AddPointToCurve(ShadowRamp, 0.00, 0.0);
-	AddPointToCurve(ShadowRamp, 1.00, 1.0);
-	
+	AddPointToFloatCurve(ShadowHatchingPatternDistributionRamp, 0.00, 0.0);
+	AddPointToFloatCurve(ShadowHatchingPatternDistributionRamp, 1.00, 1.0);
 }
 
 ///////////////////////////////////////////////////////////////////////////////////////////////////
@@ -597,6 +622,44 @@ void UToonProfile::PostEditChangeProperty(struct FPropertyChangedEvent& Property
 		GToonProfileTextureManager.UpdateProfile(LocalProfile, LocalSettings);
 	});
 }
+
+bool UToonProfile::CanEditChange(const FProperty* InProperty) const
+{
+	FString PropertyName = InProperty->GetName();
+
+	if (PropertyName == GET_MEMBER_NAME_STRING_CHECKED(FToonProfileStruct, DiffuseRampOffsetStrength)
+		|| PropertyName == GET_MEMBER_NAME_STRING_CHECKED(FToonProfileStruct, DiffuseRampOffsetSize))
+	{
+		return Settings.DiffuseRampOffsetTexture != nullptr;
+	}
+	
+	if (PropertyName == GET_MEMBER_NAME_STRING_CHECKED(FToonProfileStruct, SpecularRampOffsetStrength)
+		|| PropertyName == GET_MEMBER_NAME_STRING_CHECKED(FToonProfileStruct, SpecularRampOffsetSize))
+	{
+		return Settings.SpecularRampOffsetTexture != nullptr;
+	}
+	
+	if (PropertyName == GET_MEMBER_NAME_STRING_CHECKED(FToonProfileStruct, ShadowHatchingPatternSize)
+		|| PropertyName == GET_MEMBER_NAME_STRING_CHECKED(FToonProfileStruct, ShadowHatchingPatternStrength))
+	{
+		return Settings.ShadowHatchingPatternTexture != nullptr;
+	}
+	
+	if (PropertyName == GET_MEMBER_NAME_STRING_CHECKED(FToonProfileStruct, SpecularIndirectRamp)
+		|| PropertyName == GET_MEMBER_NAME_STRING_CHECKED(FToonProfileStruct, SpecularIndirectRampRepetition))
+	{
+		static IConsoleVariable* CVar = IConsoleManager::Get().FindConsoleVariable(TEXT("r.Substrate.Experimental.ToonReflectionQuantizationEnabled"));
+		return CVar && CVar->GetInt() > 0;
+	}
+	
+	if (PropertyName == GET_MEMBER_NAME_STRING_CHECKED(FToonProfileStruct, bUnifiedDiffuseIncludeShadow))
+	{
+		static IConsoleVariable* CVar = IConsoleManager::Get().FindConsoleVariable(TEXT("r.Substrate.Experimental.ToonUnifiedDiffuse"));
+		return CVar && CVar->GetInt() > 0;
+	}
+
+	return true;
+}
 #endif // WITH_EDITOR
 
 void UToonProfile::PostDuplicate(EDuplicateMode::Type DuplicateMode)
