Fixed crash due to dynamic wind failing to unregister a proxy.

#Synced-CL 52010655

diff --git a/Engine/Source/Runtime/Engine/Private/SkinningSceneExtensionProxy.cpp b/Engine/Source/Runtime/Engine/Private/SkinningSceneExtensionProxy.cpp
--- a/Engine/Source/Runtime/Engine/Private/SkinningSceneExtensionProxy.cpp
+++ b/Engine/Source/Runtime/Engine/Private/SkinningSceneExtensionProxy.cpp
@@ -201,11 +201,5 @@ void FInstancedSkinningSceneExtensionProxy::DestroyRenderThreadResources()
 
 const FGuid& FInstancedSkinningSceneExtensionProxy::GetTransformProviderId() const
 {
-	if (!TransformProviderProxy)
-	{
-		static FGuid RefPoseProviderId(REF_POSE_TRANSFORM_PROVIDER_GUID);
-		return RefPoseProviderId;
-	}
-
 	return TransformProviderId;
 }
