Fixed bad section offset in vertex factory with unified bone indices are used.

#Synced-CL 52010655

diff --git a/Engine/Source/Runtime/Engine/Private/SkeletalRenderGPUSkin.cpp b/Engine/Source/Runtime/Engine/Private/SkeletalRenderGPUSkin.cpp
--- a/Engine/Source/Runtime/Engine/Private/SkeletalRenderGPUSkin.cpp
+++ b/Engine/Source/Runtime/Engine/Private/SkeletalRenderGPUSkin.cpp
@@ -2591,13 +2591,14 @@ void FSkeletalMeshObjectGPUSkin::FSkeletalMeshObjectLOD::InitResources(
 {
 	check(SkelMeshRenderData);
 	check(SkelMeshRenderData->LODRenderData.IsValidIndex(LODIndex));
+	const bool bHasUnifiedBoneMap = SkelMeshRenderData->HasUnifiedBoneMap();
 
 	// vertex buffer for each lod has already been created when skelmesh was loaded
 	FSkeletalMeshLODRenderData& LODData = SkelMeshRenderData->LODRenderData[LODIndex];
 	MeshObjectWeightBuffer = FSkeletalMeshObject::GetSkinWeightVertexBuffer(LODData, CompLODInfo);
 	MeshObjectColorBuffer = FSkeletalMeshObject::GetColorVertexBuffer(LODData, CompLODInfo);
 
-	if (InBoneTransformAllocationMode == EBoneTransformAllocationMode::VertexFactory && SkelMeshRenderData->HasUnifiedBoneMap())
+	if (InBoneTransformAllocationMode == EBoneTransformAllocationMode::VertexFactory && bHasUnifiedBoneMap)
 	{
 		UnifiedBoneBuffer = MakeUnique<FUnifiedBoneBuffer>(SkelMeshRenderData->GetUnifiedBoneMap());
 	}
@@ -2607,10 +2608,10 @@ void FSkeletalMeshObjectGPUSkin::FSkeletalMeshObjectLOD::InitResources(
 	GetVertexBuffers(VertexBuffers, LODData, InBoneTransformStorageMode);
 
 	ENQUEUE_RENDER_COMMAND(FSkeletalMeshObjectLOD_InitResources)(UE::RenderCommandPipe::SkeletalMesh,
-		[this, InGPUSkinCache, VertexBuffers, &LODData, InFeatureLevel, VertexAttributeMask](FRHICommandList& RHICmdList)
+		[this, InGPUSkinCache, VertexBuffers, &LODData, InFeatureLevel, VertexAttributeMask, bHasUnifiedBoneMap](FRHICommandList& RHICmdList)
 	{
 		MeshObjectWeightBuffer_RenderThread = GetSkinWeightVertexBufferInfo(VertexBuffers.SkinWeightVertexBuffer);
-		GPUSkinVertexFactories.InitVertexFactories(RHICmdList, VertexBuffers, LODData.RenderSections, InFeatureLevel, VertexAttributeMask, GPUSkinTechnique, UnifiedBoneBuffer.IsValid());
+		GPUSkinVertexFactories.InitVertexFactories(RHICmdList, VertexBuffers, LODData.RenderSections, InFeatureLevel, VertexAttributeMask, GPUSkinTechnique, bHasUnifiedBoneMap);
 
 		if (DynamicBoundsStartOffset == INDEX_NONE && GPUSkinTechnique == ESkeletalMeshGPUSkinTechnique::GPUSkinCache)
 		{
