Fix crash in SendAllEndOfFrameUpdates when OnPreEndOfFrameSync pumps game-thread tasks.

OnPreEndOfFrameSync (e.g. cloth simulation) can call WaitUntilTaskCompletes, which
retracts game-thread tasks. Those tasks may call MarkActorComponentForNeededEndOfFrameUpdate
on already-marked components, which validates array indices. After the CL 42464944 refactor,
the component arrays were moved into FEndOfFrameUpdateContext before OnPreEndOfFrameSync ran,
so stored array indices pointed into moved-away data, causing an assertion failure.

Fix: Process OnPreEndOfFrameSync before moving the arrays into the context, using a draining
loop that handles components added during iteration. This also fixes a pre-existing correctness
issue where components marked during OnPreEndOfFrameSync were silently deferred to the next frame.

#Synced-CL 52010655

diff --git a/Engine/Source/Runtime/Engine/Private/LevelTick.cpp b/Engine/Source/Runtime/Engine/Private/LevelTick.cpp
--- a/Engine/Source/Runtime/Engine/Private/LevelTick.cpp
+++ b/Engine/Source/Runtime/Engine/Private/LevelTick.cpp
@@ -1038,7 +1038,6 @@ struct UWorld::FEndOfFrameUpdateContext
 {
 	TArray<TObjectPtr<UActorComponent>> Components;
 	TArray<TObjectPtr<UActorComponent>> Components_OnGameThread;
-	TArray<TObjectPtr<UActorComponent>> ComponentsToSync;
 	FSceneInterface::IPrimitiveTransformUpdater* PrimitiveTransformUpdater = nullptr;
 
 	void SetUpdating(bool bUpdating)
@@ -1155,41 +1154,57 @@ void UWorld::SendAllEndOfFrameUpdatesInternal(EEndOfFrameUpdateLocation Location
 		}
 	};
 
-	const auto MergeSet = [](TArray<TObjectPtr<UActorComponent>>& Dst, TSet<TObjectPtr<UActorComponent>>& Src)
-	{
-		Dst.Reserve(Dst.Num() + Src.Num());
-		for (TObjectPtr<UActorComponent> Ptr : Src)
+	// Process OnPreEndOfFrameSync BEFORE moving the component arrays into the context.
+	// OnPreEndOfFrameSync can trigger WaitUntilTaskCompletes (e.g. cloth simulation), which pumps
+	// game-thread tasks. Those tasks can call MarkActorComponentForNeededEndOfFrameUpdate, which
+	// validates that already-marked components still exist at their stored array index. If we moved
+	// the arrays first, those indices would point into the moved-away (context-owned) data, causing
+	// a check failure. Processing sync first keeps the arrays in place and indices valid.
+	//
+	// We use a draining loop because OnPreEndOfFrameSync (or tasks pumped during it) may mark
+	// additional components for sync. The TSet is copied to a local TArray for safe iteration,
+	// then drained until no new sync components remain.
+	{
+		auto DrainSyncSet = [](TSet<TObjectPtr<UActorComponent>>& SyncSet)
 		{
-			Dst.Emplace(Ptr);
+			while (!SyncSet.IsEmpty())
+			{
+				TSet<TObjectPtr<UActorComponent>> LocalSync = MoveTemp(SyncSet);
+				SyncSet.Reset();
+
+				for (UActorComponent* Component : LocalSync)
+				{
+					if (Component)
+					{
+						check(!IsValid(Component) || Component->GetMarkedForPreEndOfFrameSync());
+						if (IsValid(Component))
+						{
+							Component->OnPreEndOfFrameSync();
+						}
+						FMarkComponentEndOfFrameUpdateState::ClearMarkedForPreEndOfFrameSync(Component);
+					}
+				}
+			}
+		};
+
+		DrainSyncSet(ComponentsThatNeedPreEndOfFrameSync_AsyncTick);
+
+		if (!bAsyncInTick)
+		{
+			DrainSyncSet(ComponentsThatNeedPreEndOfFrameSync);
 		}
-		Src.Reset();
-	};
+	}
 
 	FEndOfFrameUpdateContext* Context = new FEndOfFrameUpdateContext();
 	MergeArray(Context->Components, ComponentsThatNeedEndOfFrameUpdate_AsyncTick);
-	MergeSet(Context->ComponentsToSync, ComponentsThatNeedPreEndOfFrameSync_AsyncTick);
 
 	if (!bAsyncInTick)
 	{
 		MergeArray(Context->Components, ComponentsThatNeedEndOfFrameUpdate);
-		MergeSet(Context->ComponentsToSync, ComponentsThatNeedPreEndOfFrameSync);
 		Context->Components_OnGameThread = MoveTemp(ComponentsThatNeedEndOfFrameUpdate_OnGameThread);
 		ComponentsThatNeedEndOfFrameUpdate_OnGameThread.Empty(Context->Components_OnGameThread.Num());
 	}
 
-	for (UActorComponent* Component : Context->ComponentsToSync)
-	{
-		if (Component)
-		{
-			check(!IsValid(Component) || Component->GetMarkedForPreEndOfFrameSync());
-			if (IsValid(Component))
-			{
-				Component->OnPreEndOfFrameSync();
-			}
-			FMarkComponentEndOfFrameUpdateState::ClearMarkedForPreEndOfFrameSync(Component);
-		}
-	}
-
 	// If we have a previous context active, then force this one to run asynchronously as well. They are not allowed to touch
 	// the same components (validated via IsInAsyncEndOfFrameUpdate) so it is safe to run the tasks concurrently.
 	const bool bLaunchAsyncTask = Scene && (!EndOfFrameUpdateContexts.IsEmpty() || Location == EEndOfFrameUpdateLocation::Async);
