Two related fixes:
* It makes Physics Control Component wait on the input pose being ready (this was a potential problem, but not causing visible artefacts)
* It ensures skeletal meshes associated with "floor" actors don't push CharacterMovementComponent into the POstPhysics tick - this was causing frozen animation.

#Synced-CL 52003484

diff --git a/Engine/Plugins/Experimental/PhysicsControl/Source/PhysicsControl/Private/PhysicsControlComponentImpl.cpp b/Engine/Plugins/Experimental/PhysicsControl/Source/PhysicsControl/Private/PhysicsControlComponentImpl.cpp
--- a/Engine/Plugins/Experimental/PhysicsControl/Source/PhysicsControl/Private/PhysicsControlComponentImpl.cpp
+++ b/Engine/Plugins/Experimental/PhysicsControl/Source/PhysicsControl/Private/PhysicsControlComponentImpl.cpp
@@ -182,6 +182,14 @@ void UPhysicsControlComponent::UpdateCachedSkeletalBoneData(float DeltaTime)
 		TWeakObjectPtr<USkeletalMeshComponent> SkeletalMeshComponent = CachedSkeletalMeshDataPair.Key;
 		if (USkeletalMeshComponent* SkeletalMesh = SkeletalMeshComponent.Get())
 		{
+			// Ensure any in-flight parallel animation evaluation or physics blend task has
+			// completed before we read the editable component space transforms. Without this,
+			// SwapEvaluationContextBuffers() (called during parallel
+			// work) can temporarily replace the buffer that GetEditableComponentSpaceTransforms()
+			// returns, causing us to read stale or invalid data. This can occur when tick
+			// prerequisites push this component into a later tick group (e.g. PostPhysics),
+			// overlapping with the skeletal mesh's EndPhysicsTickFunction.
+			SkeletalMesh->HandleExistingParallelEvaluationTask(true, true);
 			CachedSkeletalMeshData.Update(
 				SkeletalMesh, DeltaTime, TeleportDistanceThreshold, TeleportRotationThreshold);
 		}
diff --git a/Engine/Source/Runtime/Engine/Private/Character.cpp b/Engine/Source/Runtime/Engine/Private/Character.cpp
--- a/Engine/Source/Runtime/Engine/Private/Character.cpp
+++ b/Engine/Source/Runtime/Engine/Private/Character.cpp
@@ -629,7 +629,10 @@ namespace MovementBaseUtility
 				for (UActorComponent* Component : NewBaseOwner->GetComponents())
 				{
 					// Dont allow a based component (e.g. a particle system) to push us into a different tick group
-					if (Component && Component->PrimaryComponentTick.bCanEverTick && Component->PrimaryComponentTick.TickGroup <= BasedObjectTick.TickGroup)
+					if (Component && 
+						Component->PrimaryComponentTick.bCanEverTick && 
+						Component->PrimaryComponentTick.TickGroup <= BasedObjectTick.TickGroup && 
+						Component->PrimaryComponentTick.EndTickGroup <= BasedObjectTick.TickGroup)
 					{
 						BasedObjectTick.AddPrerequisite(Component, Component->PrimaryComponentTick);
 					}
@@ -844,7 +847,10 @@ namespace MovementBaseUtility
 					for (UActorComponent* Component : NewBaseOwner->GetComponents())
 					{
 						// Dont allow a based component (e.g. a particle system) to push us into a different tick group
-						if (Component && Component->PrimaryComponentTick.bCanEverTick && Component->PrimaryComponentTick.TickGroup <= BasedObjectTick.TickGroup)
+						if (Component && 
+							Component->PrimaryComponentTick.bCanEverTick && 
+							Component->PrimaryComponentTick.TickGroup <= BasedObjectTick.TickGroup && 
+							Component->PrimaryComponentTick.EndTickGroup <= BasedObjectTick.TickGroup)
 						{
 							BasedObjectTick.AddPrerequisite(Component, Component->PrimaryComponentTick);
 						}
