[StandaloneHLOD] Fix HLOD0 actors from HLOD-only Level Instances never hiding

* HLOD0 actors injected via AWorldPartitionHLODOnlyLevelInstance never hide when the containing map's content cells load. Their baked SourceCellGuid points to the original LI's content cells, which don't exist at runtime. HLOD1+ actors work correctly via ExternalChildHLODActors -> StandaloneHLODActorToSourceCellsMap.
* Fix by creating source cell placeholders for HLOD0 actors during streaming generation. The placeholder lands in the correct MainGrid cell through the partitioning system, and the HLOD0 -> MainGrid cell mapping is recorded in StandaloneHLODActorToSourceCellsMap for runtime use.
* Rename AWorldPartitionCustomHLODPlaceholder to AWorldPartitionHLODSourceCellPlaceholder to reflect its generalized role.

#Synced-CL 51879020

diff --git a/Engine/Source/Runtime/Engine/Private/WorldPartition/HLOD/CustomHLODPlaceholderActor.cpp b/Engine/Source/Runtime/Engine/Private/WorldPartition/HLOD/CustomHLODPlaceholderActor.cpp
deleted file mode 100644
diff --git a/Engine/Source/Runtime/Engine/Private/WorldPartition/HLOD/CustomHLODPlaceholderActorDesc.cpp b/Engine/Source/Runtime/Engine/Private/WorldPartition/HLOD/CustomHLODPlaceholderActorDesc.cpp
deleted file mode 100644
diff --git a/Engine/Source/Runtime/Engine/Private/WorldPartition/HLOD/HLODSourceCellPlaceholderActor.cpp b/Engine/Source/Runtime/Engine/Private/WorldPartition/HLOD/HLODSourceCellPlaceholderActor.cpp
new file mode 100644
--- /dev/null
+++ b/Engine/Source/Runtime/Engine/Private/WorldPartition/HLOD/HLODSourceCellPlaceholderActor.cpp
@@ -0,0 +1,74 @@
+// Copyright Epic Games, Inc. All Rights Reserved.
+
+#include "WorldPartition/HLOD/HLODSourceCellPlaceholderActor.h"
+
+#include "WorldPartition/ActorDescContainerInstance.h"
+#include "WorldPartition/WorldPartitionActorDescInstance.h"
+#include "WorldPartition/HLOD/HLODSourceCellPlaceholderActorDesc.h"
+
+#include UE_INLINE_GENERATED_CPP_BY_NAME(HLODSourceCellPlaceholderActor)
+
+AWorldPartitionHLODSourceCellPlaceholder::AWorldPartitionHLODSourceCellPlaceholder(const FObjectInitializer& ObjectInitializer)
+	: Super(ObjectInitializer)
+#if WITH_EDITOR
+	, HLODActorDescInstance(nullptr)
+#endif
+{
+	RootComponent = CreateDefaultSubobject<USceneComponent>(TEXT("Root"));
+}
+
+#if WITH_EDITOR
+void AWorldPartitionHLODSourceCellPlaceholder::InitFrom(const FWorldPartitionActorDescInstance* InHLODActorDescInstance)
+{
+	HLODActorDescInstance = InHLODActorDescInstance;
+	HLODActorGuid = HLODActorDescInstance->GetContainerInstance()->GetContainerID().GetActorGuid(HLODActorDescInstance->GetGuid());
+
+	SetActorTransform(HLODActorDescInstance->GetActorTransform());
+	SetRuntimeGrid(HLODActorDescInstance->GetRuntimeGrid());
+	SetIsSpatiallyLoaded(HLODActorDescInstance->GetActorDesc()->GetIsSpatiallyLoadedRaw());
+	bIsEditorOnlyActor = HLODActorDescInstance->GetActorIsEditorOnly();
+}
+
+const FGuid& AWorldPartitionHLODSourceCellPlaceholder::GetHLODActorGuid() const
+{
+	return HLODActorGuid;
+}
+
+void AWorldPartitionHLODSourceCellPlaceholder::SetPlaceholderType(EHLODSourceCellPlaceholderType InPlaceholderType)
+{
+	PlaceholderType = InPlaceholderType;
+}
+
+EHLODSourceCellPlaceholderType AWorldPartitionHLODSourceCellPlaceholder::GetPlaceholderType() const
+{
+	return PlaceholderType;
+}
+
+void AWorldPartitionHLODSourceCellPlaceholder::GetStreamingBounds(FBox& OutRuntimeBounds, FBox& OutEditorBounds) const
+{
+	if (HLODActorDescInstance)
+	{
+		if (PlaceholderType == EHLODSourceCellPlaceholderType::StandaloneHLOD)
+		{
+			// Use instance bounds including container transform, because the Standalone HLOD placeholder is in the main container while the source is in a child container.
+			OutRuntimeBounds = HLODActorDescInstance->GetRuntimeBounds();
+			OutEditorBounds = HLODActorDescInstance->GetEditorBounds();
+		}
+		else
+		{
+			// Use desc bounds without container transform, because the Custom HLOD placeholder is registered in the same container as the source.
+			OutRuntimeBounds = HLODActorDescInstance->GetActorDesc()->GetRuntimeBounds();
+			OutEditorBounds = HLODActorDescInstance->GetActorDesc()->GetEditorBounds();
+		}
+	}
+	else
+	{
+		OutRuntimeBounds = OutEditorBounds = FBox(EForceInit::ForceInit);
+	}
+}
+
+TUniquePtr<FWorldPartitionActorDesc> AWorldPartitionHLODSourceCellPlaceholder::CreateClassActorDesc() const
+{
+	return TUniquePtr<FWorldPartitionActorDesc>(new FHLODSourceCellPlaceholderActorDesc());
+}
+#endif
diff --git a/Engine/Source/Runtime/Engine/Private/WorldPartition/HLOD/HLODSourceCellPlaceholderActorDesc.cpp b/Engine/Source/Runtime/Engine/Private/WorldPartition/HLOD/HLODSourceCellPlaceholderActorDesc.cpp
new file mode 100644
--- /dev/null
+++ b/Engine/Source/Runtime/Engine/Private/WorldPartition/HLOD/HLODSourceCellPlaceholderActorDesc.cpp
@@ -0,0 +1,50 @@
+// Copyright Epic Games, Inc. All Rights Reserved.
+
+#include "WorldPartition/HLOD/HLODSourceCellPlaceholderActorDesc.h"
+
+#include "WorldPartition/HLOD/HLODSourceCellPlaceholderActor.h"
+
+#if WITH_EDITOR
+void FHLODSourceCellPlaceholderActorDesc::Init(const AActor* InActor)
+{
+	FWorldPartitionActorDesc::Init(InActor);
+
+	if (!bIsDefaultActorDesc)
+	{
+		const AWorldPartitionHLODSourceCellPlaceholder* PlaceholderActor = CastChecked<AWorldPartitionHLODSourceCellPlaceholder>(InActor);
+		HLODActorGuid = PlaceholderActor->GetHLODActorGuid();
+		PlaceholderType = PlaceholderActor->GetPlaceholderType();
+	}
+}
+
+bool FHLODSourceCellPlaceholderActorDesc::Equals(const FWorldPartitionActorDesc* Other) const
+{
+	if (FWorldPartitionActorDesc::Equals(Other))
+	{
+		const FHLODSourceCellPlaceholderActorDesc* OtherDesc = static_cast<const FHLODSourceCellPlaceholderActorDesc*>(Other);
+		return HLODActorGuid == OtherDesc->HLODActorGuid
+			&& PlaceholderType == OtherDesc->PlaceholderType;
+	}
+	return false;
+}
+
+const FGuid& FHLODSourceCellPlaceholderActorDesc::GetHLODActorGuid() const
+{
+	return HLODActorGuid;
+}
+
+EHLODSourceCellPlaceholderType FHLODSourceCellPlaceholderActorDesc::GetPlaceholderType() const
+{
+	return PlaceholderType;
+}
+
+uint32 FHLODSourceCellPlaceholderActorDesc::GetSizeOf() const
+{
+	return sizeof(FHLODSourceCellPlaceholderActorDesc);
+}
+
+void FHLODSourceCellPlaceholderActorDesc::Serialize(FArchive& Ar)
+{
+	ensureMsgf(false, TEXT("FHLODSourceCellPlaceholderActorDesc is not expected to be serialized"));
+}
+#endif
\ No newline at end of file
diff --git a/Engine/Source/Runtime/Engine/Private/WorldPartition/HLOD/StandaloneHLODSubsystem.cpp b/Engine/Source/Runtime/Engine/Private/WorldPartition/HLOD/StandaloneHLODSubsystem.cpp
--- a/Engine/Source/Runtime/Engine/Private/WorldPartition/HLOD/StandaloneHLODSubsystem.cpp
+++ b/Engine/Source/Runtime/Engine/Private/WorldPartition/HLOD/StandaloneHLODSubsystem.cpp
@@ -405,6 +405,19 @@ void UWorldPartitionStandaloneHLODSubsystem::ForEachStandaloneHLODActor(TFunctio
 	}
 }
 
+void UWorldPartitionStandaloneHLODSubsystem::ForEachStandaloneHLODActorWithLevelInstanceGuid(TFunctionRef<void(AWorldPartitionStandaloneHLOD*, const FGuid& LevelInstanceGuid)> Func) const
+{
+	for (const TPair<FGuid, TArray<AWorldPartitionStandaloneHLOD*>>& Pair : StandaloneHLODActors)
+	{
+		const FGuid& LevelInstanceGuid = Pair.Key;
+		const TArray<AWorldPartitionStandaloneHLOD*>& HLODActors = Pair.Value;
+		for (AWorldPartitionStandaloneHLOD* HLODActor : HLODActors)
+		{
+			Func(HLODActor, LevelInstanceGuid);
+		}
+	}
+}
+
 void UWorldPartitionStandaloneHLODSubsystem::ForEachStandaloneHLODActorFiltered(FGuid InGuid, TFunctionRef<void(AWorldPartitionStandaloneHLOD*)> Func) const
 {
 	const TArray<AWorldPartitionStandaloneHLOD*>* HLODActorsPtr = StandaloneHLODActors.Find(InGuid);
diff --git a/Engine/Source/Runtime/Engine/Private/WorldPartition/RuntimeHashSet/WorldPartitionRuntimeHashSetStreamingGeneration.cpp b/Engine/Source/Runtime/Engine/Private/WorldPartition/RuntimeHashSet/WorldPartitionRuntimeHashSetStreamingGeneration.cpp
--- a/Engine/Source/Runtime/Engine/Private/WorldPartition/RuntimeHashSet/WorldPartitionRuntimeHashSetStreamingGeneration.cpp
+++ b/Engine/Source/Runtime/Engine/Private/WorldPartition/RuntimeHashSet/WorldPartitionRuntimeHashSetStreamingGeneration.cpp
@@ -11,8 +11,8 @@
 #include "WorldPartition/HLOD/HLODLayer.h"
 #include "WorldPartition/HLOD/HLODActor.h"
 #include "WorldPartition/HLOD/HLODActorDesc.h"
-#include "WorldPartition/HLOD/CustomHLODPlaceholderActor.h"
-#include "WorldPartition/HLOD/CustomHLODPlaceholderActorDesc.h"
+#include "WorldPartition/HLOD/HLODSourceCellPlaceholderActor.h"
+#include "WorldPartition/HLOD/HLODSourceCellPlaceholderActorDesc.h"
 #include "WorldPartition/HLOD/StandaloneHLODSubsystem.h"
 
 #if WITH_EDITOR
@@ -173,6 +173,7 @@ bool UWorldPartitionRuntimeHashSet::GenerateStreaming(UWorldPartitionStreamingPo
 	};
 
 	CustomHLODActorToSourceCellsMap.Empty();
+	StandaloneHLODActorToSourceCellsMap.Empty();
 
 	TMap<FGuid, FGuid> StandaloneHLODActorToCell;
 	TArray<UWorldPartitionRuntimeCell*> RuntimeCells;
@@ -190,22 +191,33 @@ bool UWorldPartitionRuntimeHashSet::GenerateStreaming(UWorldPartitionStreamingPo
 				RuntimeCell->SetIsAlwaysLoaded(bIsCellAlwaysLoaded);
 				RuntimeCell->SetIsSpatiallyLoaded(CellDescInstance.bIsSpatiallyLoaded);
 
-				// Handle Custom HLOD Placeholder actors.
-				// Save information about mapping from Custom HLOD actor GUID to Source cell and remove placeholder actor from the list of actors to be placed in this cell.
+				// Handle HLOD source cell placeholder actors.
+				// Save information about mapping from HLOD actor GUID to source cell and remove placeholder actor from the list of actors to be placed in this cell.
 				// For more information see FWorldPartitionStreamingGenerator::CreateActorDescViewMap
-				bool bRemovedAnyCustomHLODPlaceholderActors = false;
+				bool bRemovedAnyHLODPlaceholderActors = false;
 				for (TArray<IStreamingGenerationContext::FActorInstance>::TIterator ActorInstanceIterator = CellActorInstances.CreateIterator(); ActorInstanceIterator; ++ActorInstanceIterator)
 				{
-					if (ActorInstanceIterator->GetActorDescView().GetActorNativeClass()->IsChildOf<AWorldPartitionCustomHLODPlaceholder>())
+					if (ActorInstanceIterator->GetActorDescView().GetActorNativeClass()->IsChildOf<AWorldPartitionHLODSourceCellPlaceholder>())
 					{
-						const FCustomHLODPlaceholderActorDesc* CustomHLODPlaceholderActorDesc = static_cast<const FCustomHLODPlaceholderActorDesc*>(ActorInstanceIterator->GetActorDescView().GetActorDesc());
-						CustomHLODActorToSourceCellsMap.Add(CustomHLODPlaceholderActorDesc->GetCustomHLODActorGuid(), RuntimeCell->GetGuid());
+						const FHLODSourceCellPlaceholderActorDesc* PlaceholderDesc = static_cast<const FHLODSourceCellPlaceholderActorDesc*>(ActorInstanceIterator->GetActorDescView().GetActorDesc());
+						const EHLODSourceCellPlaceholderType PlaceholderType = PlaceholderDesc->GetPlaceholderType();
+						check(PlaceholderType != EHLODSourceCellPlaceholderType::None);
+
+						if (PlaceholderType == EHLODSourceCellPlaceholderType::StandaloneHLOD)
+						{
+							StandaloneHLODActorToSourceCellsMap.Add(PlaceholderDesc->GetHLODActorGuid(), RuntimeCell->GetGuid());
+						}
+						else if (PlaceholderType == EHLODSourceCellPlaceholderType::CustomHLOD)
+						{
+							CustomHLODActorToSourceCellsMap.Add(PlaceholderDesc->GetHLODActorGuid(), RuntimeCell->GetGuid());
+						}
+
 						ActorInstanceIterator.RemoveCurrent();
-						bRemovedAnyCustomHLODPlaceholderActors = true;
+						bRemovedAnyHLODPlaceholderActors = true;
 					}
 				}
-				// If, after removing all placeholder actors, there are no actors left in this cell, mark the cell as Custom HLOD Placeholder cell, so that we can still use it at runtime to control Custom HLOD actor's visibility.
-				RuntimeCell->SetIsCustomHLODPlaceholderCell(bRemovedAnyCustomHLODPlaceholderActors && CellActorInstances.IsEmpty());
+				// If, after removing all placeholder actors, there are no actors left in this cell, mark the cell as HLOD Placeholder cell, so that we can still use it at runtime to control HLOD actor's visibility.
+				RuntimeCell->SetIsCustomHLODPlaceholderCell(bRemovedAnyHLODPlaceholderActors && CellActorInstances.IsEmpty());
 
 				PopulateRuntimeCell(RuntimeCell, CellActorInstances, OutPackagesToGenerate);
 
@@ -277,7 +289,6 @@ bool UWorldPartitionRuntimeHashSet::GenerateStreaming(UWorldPartitionStreamingPo
 
 	// Standalone HLOD actors are embedded in the main world streaming cells, so for Standalone HLOD actors that are using other Standalone HLOD actors as source actors,
 	// source cell GUID won't be correct. Save the actual source cell GUID, so that we can use that override at runtime.
-	StandaloneHLODActorToSourceCellsMap.Empty();
 	StreamingGenerationContext->ForEachActorSetContainerInstance([this, &StandaloneHLODActorToCell, World](const IStreamingGenerationContext::FActorSetContainerInstance& ActorSetContainerInstance)
 	{
 		const FStreamingGenerationActorDescViewMap* ActorDescViewMap = ActorSetContainerInstance.ActorDescViewMap;
diff --git a/Engine/Source/Runtime/Engine/Private/WorldPartition/WorldPartitionStreamingGeneration.cpp b/Engine/Source/Runtime/Engine/Private/WorldPartition/WorldPartitionStreamingGeneration.cpp
--- a/Engine/Source/Runtime/Engine/Private/WorldPartition/WorldPartitionStreamingGeneration.cpp
+++ b/Engine/Source/Runtime/Engine/Private/WorldPartition/WorldPartitionStreamingGeneration.cpp
@@ -36,8 +36,9 @@
 #include "WorldPartition/ErrorHandling/WorldPartitionStreamingGenerationNullErrorHandler.h"
 #include "WorldPartition/ErrorHandling/WorldPartitionStreamingGenerationLogErrorHandler.h"
 #include "WorldPartition/HLOD/CustomHLODActor.h"
-#include "WorldPartition/HLOD/CustomHLODPlaceholderActor.h"
+#include "WorldPartition/HLOD/HLODSourceCellPlaceholderActor.h"
 #include "WorldPartition/HLOD/HLODActor.h"
+#include "WorldPartition/HLOD/HLODOnlyLevelInstanceActor.h"
 #include "WorldPartition/HLOD/HLODActorDesc.h"
 #include "WorldPartition/HLOD/HLODLayer.h"
 #include "WorldPartition/HLOD/StandaloneHLODActor.h"
@@ -990,9 +991,10 @@ class FWorldPartitionStreamingGenerator
 					SpawnParameters.NameMode = FActorSpawnParameters::ESpawnActorNameMode::Requested;
 
 					// Spawn placeholder actor and initialize properties based on Custom HLOD actor
-					AWorldPartitionCustomHLODPlaceholder* CustomHLODPlaceholderActor = WorldPartitionContext->GetWorld()->SpawnActor<AWorldPartitionCustomHLODPlaceholder>(SpawnParameters);
-					check(CustomHLODPlaceholderActor);
-					CustomHLODPlaceholderActor->InitFrom(*Iterator);
+					AWorldPartitionHLODSourceCellPlaceholder* HLODPlaceholderActor = WorldPartitionContext->GetWorld()->SpawnActor<AWorldPartitionHLODSourceCellPlaceholder>(SpawnParameters);
+					check(HLODPlaceholderActor);
+					HLODPlaceholderActor->InitFrom(*Iterator);
+					HLODPlaceholderActor->SetPlaceholderType(EHLODSourceCellPlaceholderType::CustomHLOD);
 
 					if (UHLODLayer* HLODLayer = Cast<UHLODLayer>(Iterator->GetHLODLayer().TryLoad()))
 					{
@@ -1000,26 +1002,26 @@ class FWorldPartitionStreamingGenerator
 						{
 							if (UHLODLayer* LinkedHLODLayer = HLODLayer->GetLinkedLayer())
 							{
-								FName ResolvedRuntimeGrid = GetResolvedRuntimeGridForHLODLayer(CustomHLODPlaceholderActor->GetRuntimeGrid(), LinkedHLODLayer);
+								FName ResolvedRuntimeGrid = GetResolvedRuntimeGridForHLODLayer(HLODPlaceholderActor->GetRuntimeGrid(), LinkedHLODLayer);
 								if (!ResolvedRuntimeGrid.IsNone())
 								{
-									CustomHLODPlaceholderActor->SetRuntimeGrid(ResolvedRuntimeGrid);
+									HLODPlaceholderActor->SetRuntimeGrid(ResolvedRuntimeGrid);
 								}
 								else
 								{
-									CustomHLODPlaceholderActor->SetIsSpatiallyLoaded(false);
+									HLODPlaceholderActor->SetIsSpatiallyLoaded(false);
 								}
 							}
 						}
 					}
 				
 					// Register placeholder actor desc view
-					TUniquePtr<FStreamingGenerationTempActorDescInstance>& CustomHLODPlaceholderRef = OutCustomHLODActorPlaceholderInstances.Add_GetRef(MakeUnique<FStreamingGenerationTempActorDescInstance>(Iterator->GetContainerInstance(), CustomHLODPlaceholderActor));
-					RegisterActorDescView(FStreamingGenerationActorDescView(OutActorDescViewMap, CustomHLODPlaceholderRef.Get(), false));
+					TUniquePtr<FStreamingGenerationTempActorDescInstance>& HLODPlaceholderRef = OutCustomHLODActorPlaceholderInstances.Add_GetRef(MakeUnique<FStreamingGenerationTempActorDescInstance>(Iterator->GetContainerInstance(), HLODPlaceholderActor));
+					RegisterActorDescView(FStreamingGenerationActorDescView(OutActorDescViewMap, HLODPlaceholderRef.Get(), false));
 
 					// Destroy placeholder actor since we only need its actor descriptor.
-					WorldPartitionContext->GetWorld()->DestroyActor(CustomHLODPlaceholderActor, false, false);
-					TrashObject(CustomHLODPlaceholderActor);
+					WorldPartitionContext->GetWorld()->DestroyActor(HLODPlaceholderActor, false, false);
+					TrashObject(HLODPlaceholderActor);
 				}
 
 				// If Custom HLOD actor is used with some HLOD layer different than CustomHLODActor layer don't register its actor descriptor
@@ -1120,7 +1122,7 @@ class FWorldPartitionStreamingGenerator
 			if (WorldPartitionStandaloneHLODSubsystem)
 			{
 				// Standalone HLOD actors
-				WorldPartitionStandaloneHLODSubsystem->ForEachStandaloneHLODActor([this, &OutStandaloneHLODActorInstances, &OutActorDescViewMap, &InActorDescCollection, &RegisterActorDescView](AWorldPartitionStandaloneHLOD* StandaloneHLODActor)
+				WorldPartitionStandaloneHLODSubsystem->ForEachStandaloneHLODActorWithLevelInstanceGuid([this, &OutStandaloneHLODActorInstances, &OutCustomHLODActorPlaceholderInstances, &OutActorDescViewMap, &InActorDescCollection, &RegisterActorDescView](AWorldPartitionStandaloneHLOD* StandaloneHLODActor, const FGuid& LevelInstanceGuid)
 				{
 					if (InActorDescCollection.FindHandlingContainerInstance(StandaloneHLODActor))
 					{
@@ -1149,6 +1151,51 @@ class FWorldPartitionStreamingGenerator
 								}
 							}
 						}
+
+						// For HLOD-only Level Instances, create placeholders for HLOD0 actors so they get mapped to the containing map's MainGrid cells.
+						// HLOD1+ actors are already handled via ExternalChildHLODActors (in RuntimeHashSetStreamingGeneration).
+						// Without this, HLOD0 actors fall through to the orphan path at runtime and never hide when the containing map's content cells load.
+						if (const FWorldPartitionActorDescInstance* LevelInstanceActorDescInstance = WorldPartitionContext->GetActorDescInstance(LevelInstanceGuid))
+						{
+							if (LevelInstanceActorDescInstance->GetActorNativeClass()->IsChildOf<AWorldPartitionHLODOnlyLevelInstance>())
+							{
+								if (UActorDescContainerInstance* StandaloneHLODContainerInstance = StandaloneHLODActorRef->GetChildContainerInstance())
+								{
+									for (UActorDescContainerInstance::TConstIterator<AWorldPartitionHLOD> Iterator(StandaloneHLODContainerInstance); Iterator; ++Iterator)
+									{
+										const FWorldPartitionActorDescInstance* HLODActorDescInstance = *Iterator;
+										const FHLODActorDesc* HLODActorDesc = static_cast<const FHLODActorDesc*>(HLODActorDescInstance->GetActorDesc());
+
+										// Only HLOD0 actors (no external child references). HLOD1+ have non-empty ExternalChildHLODActors.
+										if (!HLODActorDesc->GetExternalChildHLODActors().IsEmpty())
+										{
+											continue;
+										}
+
+										FActorSpawnParameters SpawnParameters;
+										SpawnParameters.ObjectFlags = RF_Transient;
+										SpawnParameters.bCreateActorPackage = false;
+										SpawnParameters.OverrideActorGuid = FGuid::NewDeterministicGuid(TEXT("StandaloneHLODPlaceholder") + HLODActorDescInstance->GetActorSoftPath().ToString());
+										SpawnParameters.Name = FName(SpawnParameters.OverrideActorGuid.ToString());
+										SpawnParameters.NameMode = FActorSpawnParameters::ESpawnActorNameMode::Requested;
+
+										AWorldPartitionHLODSourceCellPlaceholder* HLODPlaceholderActor = WorldPartitionContext->GetWorld()->SpawnActor<AWorldPartitionHLODSourceCellPlaceholder>(SpawnParameters);
+										check(HLODPlaceholderActor);
+										HLODPlaceholderActor->InitFrom(HLODActorDescInstance);
+										HLODPlaceholderActor->SetPlaceholderType(EHLODSourceCellPlaceholderType::StandaloneHLOD);
+										// Override runtime grid to NAME_None so the placeholder lands in the containing map's MainGrid rather than the Level Instance's internal HLOD grid.
+										HLODPlaceholderActor->SetRuntimeGrid(NAME_None);
+										HLODPlaceholderActor->SetIsSpatiallyLoaded(true);
+
+										TUniquePtr<FStreamingGenerationTempActorDescInstance>& HLODPlaceholderRef = OutCustomHLODActorPlaceholderInstances.Add_GetRef(MakeUnique<FStreamingGenerationTempActorDescInstance>(StandaloneHLODActorRef->GetContainerInstance(), HLODPlaceholderActor));
+										RegisterActorDescView(FStreamingGenerationActorDescView(OutActorDescViewMap, HLODPlaceholderRef.Get(), false));
+
+										WorldPartitionContext->GetWorld()->DestroyActor(HLODPlaceholderActor, false, false);
+										TrashObject(HLODPlaceholderActor);
+									}
+								}
+							}
+						}
 					}
 				});
 			}
diff --git a/Engine/Source/Runtime/Engine/Public/WorldPartition/HLOD/CustomHLODPlaceholderActor.h b/Engine/Source/Runtime/Engine/Public/WorldPartition/HLOD/CustomHLODPlaceholderActor.h
--- a/Engine/Source/Runtime/Engine/Public/WorldPartition/HLOD/CustomHLODPlaceholderActor.h
+++ b/Engine/Source/Runtime/Engine/Public/WorldPartition/HLOD/CustomHLODPlaceholderActor.h
@@ -2,29 +2,10 @@
 
 #pragma once
 
-#include "GameFramework/Actor.h"
-#include "Misc/Guid.h"
-
-#include "CustomHLODPlaceholderActor.generated.h"
-
-class FWorldPartitionActorDescInstance;
-
-UCLASS(MinimalAPI, NotPlaceable, Transient)
-class AWorldPartitionCustomHLODPlaceholder : public AActor
-{
-	GENERATED_BODY()
-
-public:
-	ENGINE_API AWorldPartitionCustomHLODPlaceholder(const FObjectInitializer& ObjectInitializer);
+UE_DEPRECATED_HEADER(5.8, "Use WorldPartition/HLOD/HLODSourceCellPlaceholderActor.h instead of WorldPartition/HLOD/CustomHLODPlaceholderActor.h")
+#include "WorldPartition/HLOD/HLODSourceCellPlaceholderActor.h"
 
 #if WITH_EDITOR
-	void InitFrom(const FWorldPartitionActorDescInstance* InCustomHLODActorDescInstance);
-	virtual void GetStreamingBounds(FBox& OutRuntimeBounds, FBox& OutEditorBounds) const override;
-	ENGINE_API virtual TUniquePtr<class FWorldPartitionActorDesc> CreateClassActorDesc() const override;
-	ENGINE_API const FGuid& GetCustomHLODActorGuid() const;
-
-private:
-	FGuid CustomHLODActorGuid;
-	const FWorldPartitionActorDescInstance* CustomHLODActorDescInstance;
+UE_DEPRECATED(5.8, "AWorldPartitionCustomHLODPlaceholder has been renamed to AWorldPartitionHLODSourceCellPlaceholder")
+typedef AWorldPartitionHLODSourceCellPlaceholder AWorldPartitionCustomHLODPlaceholder;
 #endif
-};
\ No newline at end of file
diff --git a/Engine/Source/Runtime/Engine/Public/WorldPartition/HLOD/CustomHLODPlaceholderActorDesc.h b/Engine/Source/Runtime/Engine/Public/WorldPartition/HLOD/CustomHLODPlaceholderActorDesc.h
--- a/Engine/Source/Runtime/Engine/Public/WorldPartition/HLOD/CustomHLODPlaceholderActorDesc.h
+++ b/Engine/Source/Runtime/Engine/Public/WorldPartition/HLOD/CustomHLODPlaceholderActorDesc.h
@@ -2,20 +2,10 @@
 
 #pragma once
 
-#include "WorldPartition/WorldPartitionActorDesc.h"
+UE_DEPRECATED_HEADER(5.8, "Use WorldPartition/HLOD/HLODSourceCellPlaceholderActorDesc.h instead of WorldPartition/HLOD/CustomHLODPlaceholderActorDesc.h")
+#include "WorldPartition/HLOD/HLODSourceCellPlaceholderActorDesc.h"
 
 #if WITH_EDITOR
-class FCustomHLODPlaceholderActorDesc : public FWorldPartitionActorDesc
-{
-public:
-	virtual void Init(const AActor* InActor) override;
-	virtual bool Equals(const FWorldPartitionActorDesc* Other) const override;
-	const FGuid& GetCustomHLODActorGuid() const;
-
-protected:
-	virtual uint32 GetSizeOf() const override;
-	virtual void Serialize(FArchive& Ar) override;
-
-	FGuid CustomHLODActorGuid;
-};
-#endif
\ No newline at end of file
+UE_DEPRECATED(5.8, "FCustomHLODPlaceholderActorDesc has been renamed to FHLODSourceCellPlaceholderActorDesc")
+typedef FHLODSourceCellPlaceholderActorDesc FCustomHLODPlaceholderActorDesc;
+#endif
diff --git a/Engine/Source/Runtime/Engine/Public/WorldPartition/HLOD/HLODSourceCellPlaceholderActor.h b/Engine/Source/Runtime/Engine/Public/WorldPartition/HLOD/HLODSourceCellPlaceholderActor.h
new file mode 100644
--- /dev/null
+++ b/Engine/Source/Runtime/Engine/Public/WorldPartition/HLOD/HLODSourceCellPlaceholderActor.h
@@ -0,0 +1,45 @@
+// Copyright Epic Games, Inc. All Rights Reserved.
+
+#pragma once
+
+#include "GameFramework/Actor.h"
+#include "Misc/Guid.h"
+
+#include "HLODSourceCellPlaceholderActor.generated.h"
+
+class FWorldPartitionActorDescInstance;
+
+#if WITH_EDITOR
+enum class EHLODSourceCellPlaceholderType : uint8
+{
+	None,
+	CustomHLOD,
+	StandaloneHLOD
+};
+#endif
+
+UCLASS(MinimalAPI, NotPlaceable, Transient)
+class AWorldPartitionHLODSourceCellPlaceholder : public AActor
+{
+	GENERATED_BODY()
+
+public:
+	ENGINE_API AWorldPartitionHLODSourceCellPlaceholder(const FObjectInitializer& ObjectInitializer);
+
+#if WITH_EDITOR
+	void InitFrom(const FWorldPartitionActorDescInstance* InHLODActorDescInstance);
+	virtual void GetStreamingBounds(FBox& OutRuntimeBounds, FBox& OutEditorBounds) const override;
+	ENGINE_API virtual TUniquePtr<class FWorldPartitionActorDesc> CreateClassActorDesc() const override;
+	ENGINE_API const FGuid& GetHLODActorGuid() const;
+	ENGINE_API EHLODSourceCellPlaceholderType GetPlaceholderType() const;
+	void SetPlaceholderType(EHLODSourceCellPlaceholderType InPlaceholderType);
+
+	UE_DEPRECATED(5.8, "GetCustomHLODActorGuid has been renamed to GetHLODActorGuid")
+	const FGuid& GetCustomHLODActorGuid() const { return GetHLODActorGuid(); }
+
+private:
+	FGuid HLODActorGuid;
+	EHLODSourceCellPlaceholderType PlaceholderType = EHLODSourceCellPlaceholderType::None;
+	const FWorldPartitionActorDescInstance* HLODActorDescInstance;
+#endif
+};
diff --git a/Engine/Source/Runtime/Engine/Public/WorldPartition/HLOD/HLODSourceCellPlaceholderActorDesc.h b/Engine/Source/Runtime/Engine/Public/WorldPartition/HLOD/HLODSourceCellPlaceholderActorDesc.h
new file mode 100644
--- /dev/null
+++ b/Engine/Source/Runtime/Engine/Public/WorldPartition/HLOD/HLODSourceCellPlaceholderActorDesc.h
@@ -0,0 +1,24 @@
+// Copyright Epic Games, Inc. All Rights Reserved.
+
+#pragma once
+
+#include "WorldPartition/WorldPartitionActorDesc.h"
+#include "WorldPartition/HLOD/HLODSourceCellPlaceholderActor.h"
+
+#if WITH_EDITOR
+class FHLODSourceCellPlaceholderActorDesc : public FWorldPartitionActorDesc
+{
+public:
+	virtual void Init(const AActor* InActor) override;
+	virtual bool Equals(const FWorldPartitionActorDesc* Other) const override;
+	const FGuid& GetHLODActorGuid() const;
+	EHLODSourceCellPlaceholderType GetPlaceholderType() const;
+
+protected:
+	virtual uint32 GetSizeOf() const override;
+	virtual void Serialize(FArchive& Ar) override;
+
+	FGuid HLODActorGuid;
+	EHLODSourceCellPlaceholderType PlaceholderType = EHLODSourceCellPlaceholderType::None;
+};
+#endif
diff --git a/Engine/Source/Runtime/Engine/Public/WorldPartition/HLOD/StandaloneHLODSubsystem.h b/Engine/Source/Runtime/Engine/Public/WorldPartition/HLOD/StandaloneHLODSubsystem.h
--- a/Engine/Source/Runtime/Engine/Public/WorldPartition/HLOD/StandaloneHLODSubsystem.h
+++ b/Engine/Source/Runtime/Engine/Public/WorldPartition/HLOD/StandaloneHLODSubsystem.h
@@ -45,6 +45,7 @@ public:
 	ENGINE_API void DeleteStandaloneHLODActorsRecursive(const FLevelInstanceActorDesc& InLevelInstanceActorDesc);
 
 	ENGINE_API void ForEachStandaloneHLODActor(TFunctionRef<void(AWorldPartitionStandaloneHLOD*)> Func) const;
+	ENGINE_API void ForEachStandaloneHLODActorWithLevelInstanceGuid(TFunctionRef<void(AWorldPartitionStandaloneHLOD*, const FGuid& LevelInstanceGuid)> Func) const;
 	ENGINE_API void ForEachStandaloneHLODActorFiltered(FGuid InGuid, TFunctionRef<void(AWorldPartitionStandaloneHLOD*)> Func) const;
 
 	static ENGINE_API bool GetStandaloneHLODFolderPathAndPackagePrefix(const FString& InWorldPackageName, FString& OutFolderPath, FString& OutPackagePrefix);
