[HWRT] Prevent stale RayTracingGeometry when FSkeletalMeshObjectGPUSkin changes from GPUSkinCache to MeshDeformer.
- we were releasing the SkinCacheEntry but not the RayTracingGeometry.

#Synced-CL 52071961

diff --git a/Engine/Source/Runtime/Engine/Private/SkeletalRenderGPUSkin.cpp b/Engine/Source/Runtime/Engine/Private/SkeletalRenderGPUSkin.cpp
--- a/Engine/Source/Runtime/Engine/Private/SkeletalRenderGPUSkin.cpp
+++ b/Engine/Source/Runtime/Engine/Private/SkeletalRenderGPUSkin.cpp
@@ -1031,25 +1031,28 @@ void FSkeletalMeshObjectGPUSkin::ProcessUpdatedDynamicData(FRHICommandList& RHIC
 	// Immediately release any stale entry if we've recently switched to a LOD level that disallows skin cache
 	if (!bGPUSkinCacheEnabled)
 	{
-#if RHI_RAYTRACING
-		if (Mode == EGPUSkinCacheEntryMode::Raster)
-#endif
+		bool bReleasedSkinCacheEntry = false;
+
+		if (SkinCacheEntry)
 		{
-			if (SkinCacheEntry)
-			{
-				FGPUSkinCache::Release(SkinCacheEntry);
-			}
+			FGPUSkinCache::Release(SkinCacheEntry);
+			bReleasedSkinCacheEntry = true;
 		}
+
 #if RHI_RAYTRACING
-		else
+		if (SkinCacheEntryForRayTracing)
 		{
-			check(Mode == EGPUSkinCacheEntryMode::RayTracing);
-			if (SkinCacheEntryForRayTracing)
-			{
-				FGPUSkinCache::Release(SkinCacheEntryForRayTracing);
-			}
+			FGPUSkinCache::Release(SkinCacheEntryForRayTracing);
+			bReleasedSkinCacheEntry = true;
 		}
-#endif
+
+		if (bReleasedSkinCacheEntry)
+		{
+			// When SkinCacheEntry is gone, clear geometry
+			RayTracingGeometry.ReleaseRHI();
+			RayTracingGeometry.SetInitializer(FRayTracingGeometryInitializer{});
+		}
+#endif // RHI_RAYTRACING
 	}
 
 	TArray<FGPUSkinCache::FProcessEntrySection, TInlineAllocator<8>> ProcessEntrySections;
