Fixed crash in UGameUserSettings::IsGlobalIlluminationAllowed for worlds which have been initialized with InitializeScenes(false)

#rnx
#Synced-CL 52050999

diff --git a/Engine/Source/Runtime/Engine/Private/GameUserSettings.cpp b/Engine/Source/Runtime/Engine/Private/GameUserSettings.cpp
--- a/Engine/Source/Runtime/Engine/Private/GameUserSettings.cpp
+++ b/Engine/Source/Runtime/Engine/Private/GameUserSettings.cpp
@@ -1258,7 +1258,7 @@ bool UGameUserSettings::IsGlobalIlluminationAllowed(UObject* WorldContextObject)
 {
 	UWorld* World = GEngine->GetWorldFromContextObject(WorldContextObject, EGetWorldErrorMode::ReturnNull);
 
-	const EShaderPlatform ShaderPlatform = World ? World->Scene->GetShaderPlatform() : GMaxRHIShaderPlatform;
+	const EShaderPlatform ShaderPlatform = World && World->Scene ? World->Scene->GetShaderPlatform() : GMaxRHIShaderPlatform;
 
 	static const auto CVarLumenAllowed = IConsoleManager::Get().FindTConsoleVariableDataInt(TEXT("r.Lumen.DiffuseIndirect.Allow"));
 
