Fix for non-square virtual texture resizing during compilation.
During texture compilation we need to take tile size into account.
Previously we did this but didn't consider the effect of any LODScale.
Also previously the scale was always uniform which was patholgically bad for the 1D style texture case.
#Synced-CL 52056483

diff --git a/Engine/Source/Runtime/Engine/Private/TextureDerivedData.cpp b/Engine/Source/Runtime/Engine/Private/TextureDerivedData.cpp
--- a/Engine/Source/Runtime/Engine/Private/TextureDerivedData.cpp
+++ b/Engine/Source/Runtime/Engine/Private/TextureDerivedData.cpp
@@ -252,6 +252,16 @@ static void SerializeForKey(FArchive& Ar, const FTextureBuildSettings& Settings)
 			TempGuid = FGuid(0xA227BEFC,0x9F8643C6,0x81580369,0xC4C6F73E);
 			Ar << TempGuid;
 		}
+
+		// LODBias affects VT dimensions in InitializeFromBuildSettings()
+		// Only emit when non-zero to preserve existing DDC entries.
+		if (Settings.LODBias != 0)
+		{
+			TempGuid = FGuid(0x7E4A2F81, 0xC3D59B46, 0xA8071E2D, 0x5F93C084);
+			Ar << TempGuid;
+			TempUint32 = (uint32)Settings.LODBias;
+			Ar << TempUint32;
+		}
 	}
 
 	// Avoid changing key if texture is not being downscaled
diff --git a/Engine/Source/Runtime/Engine/Private/VT/VirtualTextureDataBuilder.cpp b/Engine/Source/Runtime/Engine/Private/VT/VirtualTextureDataBuilder.cpp
--- a/Engine/Source/Runtime/Engine/Private/VT/VirtualTextureDataBuilder.cpp
+++ b/Engine/Source/Runtime/Engine/Private/VT/VirtualTextureDataBuilder.cpp
@@ -314,14 +314,14 @@ bool FVirtualTextureBuilderDerivedInfo::InitializeFromBuildSettings(const FTextu
 		return false;
 	}
 
-	// Ensure block size is at least 1 tile, while preserving aspect ratio
-	BlockSizeScale = 1;
-	while (BlockSizeX < TileSize || BlockSizeY < TileSize)
-	{
-		BlockSizeX *= 2;
-		BlockSizeY *= 2;
-		BlockSizeScale *= 2;
-	}
+	// Ensure block size is at least 1 tile. Take LODBias into account when doing that.
+	const uint32 LodBias = FMath::Max(BuildSettingsLayer0.LODBias, 0u);
+	int32 BiasedSizeX = FMath::Max(BlockSizeX >> LodBias, 1);
+	int32 BiasedSizeY = FMath::Max(BlockSizeY >> LodBias, 1);
+	BlockSizeScaleX = 1;
+	BlockSizeScaleY = 1;
+	while (BiasedSizeX < TileSize) { BiasedSizeX *= 2; BlockSizeX *= 2; BlockSizeScaleX *= 2; }
+	while (BiasedSizeY < TileSize) { BiasedSizeY *= 2; BlockSizeY *= 2; BlockSizeScaleY *= 2; }
 
 	SizeInBlocksX = InSourceData.SizeInBlocksX;
 	SizeInBlocksY = InSourceData.SizeInBlocksY;
@@ -962,7 +962,7 @@ void FVirtualTextureDataBuilder::BuildLayerBlocks(FSlowTask& BuildTask, uint32 L
 			TBSettings.MipGenSettings = TMGS_SimpleAverage;
 		}
 
-		// in case image is smaller than tile size, we need PowerOfTwoMode adjusted size to correctly calculate MipBias & LocalBlockSizeScale in a while loop below
+		// in case image is smaller than tile size, we need PowerOfTwoMode adjusted size to correctly calculate MipBias & LocalBlockSizeScaleX/Y in a while loop below
 		int32 AdjustedSizeX, AdjustedSizeY, AdjustedSizeZ;
 		UE::TextureBuildUtilities::GetPowerOfTwoTargetTextureSize(
 			SourceMips[0].SizeX, SourceMips[0].SizeY, 1,
@@ -973,12 +973,14 @@ void FVirtualTextureDataBuilder::BuildLayerBlocks(FSlowTask& BuildTask, uint32 L
 
 		// For multi-block images, we may have scaled the max block size to be tile-sized, but individual blocks may still be smaller than 1 tile
 		// These need to be scaled up as well (scaling up individual blocks has the effect of reducing the block's mip-bias)
-		int32 LocalBlockSizeScale = DerivedInfo.BlockSizeScale;
-		while (AdjustedSizeX * LocalBlockSizeScale < TileSize || AdjustedSizeY * LocalBlockSizeScale < TileSize)
+		int32 LocalBlockSizeScaleX = DerivedInfo.BlockSizeScaleX;
+		int32 LocalBlockSizeScaleY = DerivedInfo.BlockSizeScaleY;
+		while (AdjustedSizeX * LocalBlockSizeScaleX < TileSize || AdjustedSizeY * LocalBlockSizeScaleY < TileSize)
 		{
 			check(BlockData.MipBias > 0u);
 			--BlockData.MipBias;
-			LocalBlockSizeScale *= 2;
+			LocalBlockSizeScaleX *= 2;
+			LocalBlockSizeScaleY *= 2;
 		}
 
 		// give each tile a unique DebugTexturePathName for DebugDump option :
@@ -988,7 +990,7 @@ void FVirtualTextureDataBuilder::BuildLayerBlocks(FSlowTask& BuildTask, uint32 L
 		// this is the Build to Uncompressed to apply processing to create the source for the tiles
 		TArray<FCompressedImage2D> CompressedMips;
 		bool bBuildTextureResult = false;
-		if (LocalBlockSizeScale == 1)
+		if (LocalBlockSizeScaleX == 1 && LocalBlockSizeScaleY == 1)
 		{
 			uint32 NumMipsInTail, ExtData;
 			bBuildTextureResult = Compressor->BuildTexture(SourceMips, *CompositeSourceMips, TBSettings, CurDebugTexturePathName, CompressedMips, NumMipsInTail, ExtData, nullptr);
@@ -1010,7 +1012,7 @@ void FVirtualTextureDataBuilder::BuildLayerBlocks(FSlowTask& BuildTask, uint32 L
 				FImage* ScaledMip = new(ScaledSourceMips) FImage;
 				// Pow22 cannot be used as a destination gamma, so change it to sRGB now :
 				EGammaSpace GammaSpace = (SrcMip.GammaSpace == EGammaSpace::Pow22) ? EGammaSpace::sRGB : SrcMip.GammaSpace;
-				SrcMip.ResizeTo(*ScaledMip, SrcMip.SizeX * LocalBlockSizeScale, SrcMip.SizeY * LocalBlockSizeScale, SrcMip.Format, GammaSpace);
+				SrcMip.ResizeTo(*ScaledMip, SrcMip.SizeX * LocalBlockSizeScaleX, SrcMip.SizeY * LocalBlockSizeScaleY, SrcMip.Format, GammaSpace);
 			}
 
 			for (const FImage& SrcMip : *CompositeSourceMips)
@@ -1018,7 +1020,7 @@ void FVirtualTextureDataBuilder::BuildLayerBlocks(FSlowTask& BuildTask, uint32 L
 				FImage* ScaledMip = new(ScaledCompositeMips) FImage;
 				// Pow22 cannot be used as a destination gamma, so change it to sRGB now :
 				EGammaSpace GammaSpace = (SrcMip.GammaSpace == EGammaSpace::Pow22) ? EGammaSpace::sRGB : SrcMip.GammaSpace;
-				SrcMip.ResizeTo(*ScaledMip, SrcMip.SizeX * LocalBlockSizeScale, SrcMip.SizeY * LocalBlockSizeScale, SrcMip.Format, GammaSpace);
+				SrcMip.ResizeTo(*ScaledMip, SrcMip.SizeX * LocalBlockSizeScaleX, SrcMip.SizeY * LocalBlockSizeScaleY, SrcMip.Format, GammaSpace);
 			}
 
 			// Pow22 was converted to sRGB by Resize :
diff --git a/Engine/Source/Runtime/Engine/Private/VT/VirtualTextureDataBuilder.h b/Engine/Source/Runtime/Engine/Private/VT/VirtualTextureDataBuilder.h
--- a/Engine/Source/Runtime/Engine/Private/VT/VirtualTextureDataBuilder.h
+++ b/Engine/Source/Runtime/Engine/Private/VT/VirtualTextureDataBuilder.h
@@ -89,7 +89,8 @@ struct FVirtualTextureBuilderDerivedInfo
 	int32 SizeInBlocksY = 0;
 	int32 BlockSizeX = 0;
 	int32 BlockSizeY = 0;
-	int32 BlockSizeScale = 1;
+	int32 BlockSizeScaleX = 1;
+	int32 BlockSizeScaleY = 1;
 	int32 SizeX = 0;
 	int32 SizeY = 0;
 	int32 NumMips = 0;
