[Audio] Fix NumberOfLoadsInFlight assert in FAudioChunkCache destructor

CL 43738639 fixed a deadlock in CancelAllPendingLoads but left the destructor
calling CachePool.Reset() without first cancelling in-flight loads. Added
CancelAllPendingLoads() call before Reset() to close the race window.

#Synced-CL 52025567
#Synced-CL 52025567

diff --git a/Engine/Source/Runtime/Engine/Private/AudioStreamingCache.cpp b/Engine/Source/Runtime/Engine/Private/AudioStreamingCache.cpp
--- a/Engine/Source/Runtime/Engine/Private/AudioStreamingCache.cpp
+++ b/Engine/Source/Runtime/Engine/Private/AudioStreamingCache.cpp
@@ -820,7 +820,10 @@ FAudioChunkCache::FAudioChunkCache(uint32 InMaxChunkSize, uint32 NumChunks, uint
 FAudioChunkCache::~FAudioChunkCache()
 {
 	LLM_SCOPE(ELLMTag::AudioStreamCache);
-	// While this is handled by the default destructor, we do this to ensure that we don't leak async read operations.
+
+	CancelAllPendingLoads();
+	
+	// Ensure all in-flight loads are cancelled and waited on before the pool is destroyed.
 	CachePool.Reset();
 	check(NumberOfLoadsInFlight.GetValue() == 0);
 }
