Use weak pointers to track PostProcessVolumes to prevent use-after-free.
- PPVs are now stored as TWeakInterfacePtr<IInterface_PostProcessVolume>.
- On use, the pointer is first checked to see if it's still valid.
- On insertion or removal of a PPV, if stale entries are encountered in the list, they may be removed.
- An array with the old signature is maintained to retain API compatibility, but will be removed in a future release.

#Synced-CL 51959746

diff --git a/Engine/Plugins/MovieScene/MovieRenderPipeline/Source/MovieRenderPipelineCore/Private/Graph/MovieGraphCoreTimeStep.cpp b/Engine/Plugins/MovieScene/MovieRenderPipeline/Source/MovieRenderPipelineCore/Private/Graph/MovieGraphCoreTimeStep.cpp
--- a/Engine/Plugins/MovieScene/MovieRenderPipeline/Source/MovieRenderPipelineCore/Private/Graph/MovieGraphCoreTimeStep.cpp
+++ b/Engine/Plugins/MovieScene/MovieRenderPipeline/Source/MovieRenderPipelineCore/Private/Graph/MovieGraphCoreTimeStep.cpp
@@ -763,8 +763,14 @@ float UMovieGraphCoreTimeStep::GetBlendedMotionBlurAmount()
 	{
 		// Apply any motion blur settings from post process volumes in the world
 		FVector ViewLocation = PlayerCameraManager->GetCameraLocation();
-		for (IInterface_PostProcessVolume* PPVolume : GetWorld()->GetPostProcessVolumes())
+		for (const TWeakInterfacePtr<IInterface_PostProcessVolume>& PPVolumeWeak : GetWorld()->GetPostProcessVolumes())
 		{
+			IInterface_PostProcessVolume* PPVolume = PPVolumeWeak.Get();
+			if (!PPVolume)
+			{
+				continue;
+			}
+
 			const FPostProcessVolumeProperties VolumeProperties = PPVolume->GetProperties();
 
 			// Skip any volumes which are either disabled or don't modify blur amount
diff --git a/Engine/Plugins/MovieScene/MovieRenderPipeline/Source/MovieRenderPipelineCore/Private/MoviePipeline.cpp b/Engine/Plugins/MovieScene/MovieRenderPipeline/Source/MovieRenderPipelineCore/Private/MoviePipeline.cpp
--- a/Engine/Plugins/MovieScene/MovieRenderPipeline/Source/MovieRenderPipelineCore/Private/MoviePipeline.cpp
+++ b/Engine/Plugins/MovieScene/MovieRenderPipeline/Source/MovieRenderPipelineCore/Private/MoviePipeline.cpp
@@ -1405,8 +1405,14 @@ MoviePipeline::FFrameConstantMetrics UMoviePipeline::CalculateShotFrameMetrics(c
 	{
 		// Apply any motion blur settings from post process volumes in the world
 		FVector ViewLocation = PlayerCameraManager->GetCameraLocation();
-		for (IInterface_PostProcessVolume* PPVolume : GetWorld()->GetPostProcessVolumes())
+		for (const TWeakInterfacePtr<IInterface_PostProcessVolume>& PPVolumeWeak : GetWorld()->GetPostProcessVolumes())
 		{
+			IInterface_PostProcessVolume* PPVolume = PPVolumeWeak.Get();
+			if (!PPVolume)
+			{
+				continue;
+			}
+
 			const FPostProcessVolumeProperties VolumeProperties = PPVolume->GetProperties();
 
 			// Skip any volumes which are either disabled or don't modify blur amount
diff --git a/Engine/Plugins/MovieScene/MovieRenderPipeline/Source/MovieRenderPipelineCore/Private/MoviePipelineUtils.cpp b/Engine/Plugins/MovieScene/MovieRenderPipeline/Source/MovieRenderPipelineCore/Private/MoviePipelineUtils.cpp
--- a/Engine/Plugins/MovieScene/MovieRenderPipeline/Source/MovieRenderPipelineCore/Private/MoviePipelineUtils.cpp
+++ b/Engine/Plugins/MovieScene/MovieRenderPipeline/Source/MovieRenderPipelineCore/Private/MoviePipelineUtils.cpp
@@ -1260,8 +1260,14 @@ namespace UE
 
 		void DoPostProcessBlend(const FVector& InViewLocation, const UWorld* InWorld, const FMinimalViewInfo& InViewInfo, FSceneView* InOutView)
 		{
-			for (IInterface_PostProcessVolume* PPVolume : InWorld->GetPostProcessVolumes())
+			for (const TWeakInterfacePtr<IInterface_PostProcessVolume>& PPVolumeWeak : InWorld->GetPostProcessVolumes())
 			{
+				IInterface_PostProcessVolume* PPVolume = PPVolumeWeak.Get();
+				if (!PPVolume)
+				{
+					continue;
+				}
+
 				const FPostProcessVolumeProperties VolumeProperties = PPVolume->GetProperties();
 
 				// Skip any volumes which are disabled
diff --git a/Engine/Plugins/Runtime/nDisplay/Source/DisplayClusterMoviePipeline/Private/DisplayClusterMoviePipelineViewportPass.cpp b/Engine/Plugins/Runtime/nDisplay/Source/DisplayClusterMoviePipeline/Private/DisplayClusterMoviePipelineViewportPass.cpp
--- a/Engine/Plugins/Runtime/nDisplay/Source/DisplayClusterMoviePipeline/Private/DisplayClusterMoviePipelineViewportPass.cpp
+++ b/Engine/Plugins/Runtime/nDisplay/Source/DisplayClusterMoviePipeline/Private/DisplayClusterMoviePipelineViewportPass.cpp
@@ -175,8 +175,14 @@ void UDisplayClusterMoviePipelineViewportPassBase::BlendPostProcessSettings(FSce
 			// Use DCRA as a camera
 			FVector ViewLocation = DCViews[InViewportIndex].ViewLocation;
 
-			for (IInterface_PostProcessVolume* PPVolume : GetWorld()->GetPostProcessVolumes())
+			for (const TWeakInterfacePtr<IInterface_PostProcessVolume>& PPVolumeWeak : GetWorld()->GetPostProcessVolumes())
 			{
+				IInterface_PostProcessVolume* PPVolume = PPVolumeWeak.Get();
+				if (!PPVolume)
+				{
+					continue;
+				}
+
 				const FPostProcessVolumeProperties VolumeProperties = PPVolume->GetProperties();
 
 				// Skip any volumes which are disabled
diff --git a/Engine/Source/Runtime/Engine/Classes/Engine/World.h b/Engine/Source/Runtime/Engine/Classes/Engine/World.h
--- a/Engine/Source/Runtime/Engine/Classes/Engine/World.h
+++ b/Engine/Source/Runtime/Engine/Classes/Engine/World.h
@@ -1914,22 +1914,24 @@ public:
 
 	/** Removes a post process volume from the world.  Returns true if it was successfully removed. */
 	UE_API bool RemovePostProcessVolume(IInterface_PostProcessVolume* InVolume);
+	
+	/** Called when a scene view for this world needs the worlds post process settings computed */
+	UE_API void AddPostProcessingSettings(FVector ViewLocation, FSceneView* SceneView);
 
 	/** Return an array of post processing volumes, sorted in ascending order of priority. */
-	const TArray<IInterface_PostProcessVolume*>& GetPostProcessVolumes() const
+	const TArray<TWeakInterfacePtr<IInterface_PostProcessVolume>>& GetPostProcessVolumes() const
 	{
-PRAGMA_DISABLE_DEPRECATION_WARNINGS
-		return PostProcessVolumes;
-PRAGMA_ENABLE_DEPRECATION_WARNINGS
+		return PostProcessVolumesWeak;
 	}
 
 	UE_DEPRECATED(5.8, "Use GetPostProcessVolumes() instead")
 	TArray<IInterface_PostProcessVolume*> PostProcessVolumes;
 
-	/** Called when a scene view for this world needs the worlds post process settings computed */
-	UE_API void AddPostProcessingSettings(FVector ViewLocation, FSceneView* SceneView);
 
 private:
+	/** Must be identical to `PostProcessVolumes`. TWeakInterfacePtr is used here to detect stale (destroyed) volumes that were not removed. */
+	TArray<TWeakInterfacePtr<IInterface_PostProcessVolume>> PostProcessVolumesWeak;
+	void RemovePostProcessVolume(int32 Index); // do not make public, indices are unstable
 	/** Caches the size of each volume at the time of insertion, for sort determinism in case it later changes.  One for one with PostProcessVolumes above. */
 	TArray<double> PostProcessVolumeCachedSizes;
 public:
diff --git a/Engine/Source/Runtime/Engine/Private/World.cpp b/Engine/Source/Runtime/Engine/Private/World.cpp
--- a/Engine/Source/Runtime/Engine/Private/World.cpp
+++ b/Engine/Source/Runtime/Engine/Private/World.cpp
@@ -10396,53 +10396,94 @@ static bool PostProcessVolumeLess(const FPostProcessVolumeProperties& Properties
 
 void UWorld::InsertPostProcessVolume(IInterface_PostProcessVolume* InVolume)
 {
-PRAGMA_DISABLE_DEPRECATION_WARNINGS
-	const int32 NumVolumes = PostProcessVolumes.Num();
+	int32 NumVolumes = PostProcessVolumesWeak.Num();
 	FPostProcessVolumeProperties TargetProperties = InVolume->GetProperties();
 	int32 InsertIndex = 0;
 	// TODO: replace with binary search.
-	for (; InsertIndex < NumVolumes; InsertIndex++)
+	while (InsertIndex < NumVolumes)
 	{
-		IInterface_PostProcessVolume* CurrentVolume = PostProcessVolumes[InsertIndex];
+		IInterface_PostProcessVolume* CurrentVolume = PostProcessVolumesWeak[InsertIndex].Get();
 		if (CurrentVolume == InVolume)
 		{
 			return;
 		}
 
-		FPostProcessVolumeProperties CurrentProperties = CurrentVolume->GetProperties();
-
-		if (PostProcessVolumeLess(TargetProperties, CurrentProperties, PostProcessVolumeCachedSizes[InsertIndex]))
+		if (!CurrentVolume) // stale entry, volume was destroyed
 		{
-			break;
+			RemovePostProcessVolume(InsertIndex);
+			--NumVolumes;
+		}
+		else
+		{
+			FPostProcessVolumeProperties CurrentProperties = CurrentVolume->GetProperties();
+
+			if (PostProcessVolumeLess(TargetProperties, CurrentProperties, PostProcessVolumeCachedSizes[InsertIndex]))
+			{
+				break;
+			}
+			++InsertIndex;
 		}
 	}
 
 	// Check the remainder of the array to see if InVolume is already present.  The caller may have changed TargetProperties.Size
 	// from what was originally cached, making the "Size > CurrentCachedSize" comparison inconsistent, meaning the existing item will not
 	// have been reached in the loop.
-	for (int32 Index = InsertIndex + 1; Index < NumVolumes; Index++)
+	for (int32 Index = InsertIndex + 1; Index < NumVolumes; )
 	{
-		if (PostProcessVolumes[Index] == InVolume)
+		IInterface_PostProcessVolume* CurrentVolume = PostProcessVolumesWeak[Index].Get();
+		if (CurrentVolume == InVolume)
 		{
 			return;
 		}
+
+		if (!CurrentVolume) // stale entry, volume was destroyed
+		{
+			RemovePostProcessVolume(Index);
+			--NumVolumes;
+		}
+		else
+		{
+			++Index;
+		}
 	}
-	PostProcessVolumes.Insert(InVolume, InsertIndex);
+
+	PostProcessVolumesWeak.Insert(TWeakInterfacePtr<IInterface_PostProcessVolume>(InVolume), InsertIndex);
 	PostProcessVolumeCachedSizes.Insert(TargetProperties.Size, InsertIndex);
+PRAGMA_DISABLE_DEPRECATION_WARNINGS
+	PostProcessVolumes.Insert(InVolume, InsertIndex);
 PRAGMA_ENABLE_DEPRECATION_WARNINGS
 }
 
 bool UWorld::RemovePostProcessVolume(IInterface_PostProcessVolume* InVolume)
 {
-PRAGMA_DISABLE_DEPRECATION_WARNINGS
-	int32 RemoveIndex;
-	if (PostProcessVolumes.Find(InVolume, RemoveIndex))
+	for (int32 Index = 0, NumVolumes = PostProcessVolumesWeak.Num(); Index < NumVolumes; )
 	{
-		PostProcessVolumes.RemoveAt(RemoveIndex);
-		PostProcessVolumeCachedSizes.RemoveAt(RemoveIndex);
-		return true;
+		IInterface_PostProcessVolume* CurrentVolume = PostProcessVolumesWeak[Index].Get();
+		if (!CurrentVolume) // stale entry, volume was destroyed
+		{
+			RemovePostProcessVolume(Index);
+			--NumVolumes;
+		}
+		else
+		{
+			if (CurrentVolume == InVolume)
+			{
+				RemovePostProcessVolume(Index);
+				--NumVolumes;
+				return true;
+			}
+			++Index;
+		}
 	}
 	return false;
+}
+
+void UWorld::RemovePostProcessVolume(int32 Index)
+{
+	PostProcessVolumesWeak.RemoveAt(Index);
+	PostProcessVolumeCachedSizes.RemoveAt(Index);
+PRAGMA_DISABLE_DEPRECATION_WARNINGS
+	PostProcessVolumes.RemoveAt(Index);
 PRAGMA_ENABLE_DEPRECATION_WARNINGS
 }
 
@@ -10633,9 +10674,12 @@ void UWorld::AddPostProcessingSettings(FVector ViewLocation, FSceneView* SceneVi
 	SceneView->FinalPostProcessDebugInfo.Reset();
 #endif
 
-	for (IInterface_PostProcessVolume* PPVolume : GetPostProcessVolumes())
+	for (const TWeakInterfacePtr<IInterface_PostProcessVolume>& PPVolumeWeak : PostProcessVolumesWeak)
 	{
-		DoPostProcessVolume(PPVolume, ViewLocation, SceneView);
+		if (IInterface_PostProcessVolume* PPVolume = PPVolumeWeak.Get())
+		{
+			DoPostProcessVolume(PPVolume, ViewLocation, SceneView);
+		}
 	}
 }
 
