RemoteObjects: Fix for processing RPCs at the right time and not stopping replication upon ActorComponent BeginDestroy/EndPlay due to a migration.

We must be registered with Iris to process RPCs, so ensure we're already registered when we attempt to execute them.
Add logging for any dropped RPCs.
ActorComponent BeginDestroy can be executed when reusing the Object's memory (acting as its destructor).  If we know the ActorComponent was migrated-out (exists elsewhere) do not call the 'immediate' StopReplication (it would have a deferred StopReplication in PostMigrate).

#Synced-CL 52055930

diff --git a/Engine/Source/Runtime/CoreUObject/Private/UObject/RemoteObjectSerialization.cpp b/Engine/Source/Runtime/CoreUObject/Private/UObject/RemoteObjectSerialization.cpp
--- a/Engine/Source/Runtime/CoreUObject/Private/UObject/RemoteObjectSerialization.cpp
+++ b/Engine/Source/Runtime/CoreUObject/Private/UObject/RemoteObjectSerialization.cpp
@@ -1966,6 +1966,7 @@ namespace UE::RemoteObject::Serialization::Network
 
 		// Call the RPC function.
 		UE_LOGF(LogRemoteSerialization, Verbose, "Processing queued RPC %ls on %ls", *Function->GetName(), *TargetObject->GetName());
+		ensureMsgf(UE::RemoteObject::Handle::IsOwned(TargetObject), TEXT("%hs: %s is Not Owned so we cannot process the RPC!"), __FUNCTION__, *GetNameSafe(TargetObject));
 		TargetObject->ProcessEvent(Function, Params);
 
 		for (TFieldIterator<FProperty> It(Function); It && It->HasAnyPropertyFlags(CPF_Parm); ++It)
diff --git a/Engine/Source/Runtime/Engine/Private/Actor.cpp b/Engine/Source/Runtime/Engine/Private/Actor.cpp
--- a/Engine/Source/Runtime/Engine/Private/Actor.cpp
+++ b/Engine/Source/Runtime/Engine/Private/Actor.cpp
@@ -1358,13 +1358,6 @@ void AActor::PostMigrate(const FUObjectMigrationContext& MigrationContext)
 	{
 		// Build up the list of Components we have again (even if we previously owned this object, it could have been mutated)
 		EAddComponentForReplicationBehavior AddComponentForRepType = EAddComponentForReplicationBehavior::NeverAdd;
-		
-		// NOTE: This is unnecessary because  UActorComponent::PostMigrate() will explicitly defer start replication for itself.
-		//if (MigrationRecvType == EObjectMigrationRecvType::AssignedOwnership)
-		//{
-		//	AddComponentForRepType = EAddComponentForReplicationBehavior::Deferred;
-		//}
-
 		ResetOwnedComponents(AddComponentForRepType);
 
 		// We can get rid of this check further down the road...
@@ -1402,8 +1395,19 @@ void AActor::PostMigrate(const FUObjectMigrationContext& MigrationContext)
 		}
 		else if (MigrationRecvType == EObjectMigrationRecvType::ReturnedLoan)
 		{
-			// Process any RPCs for this actor and its sub-objects (assumes the actor and its sub-objects have already been registered with Iris).
-			UE::RemoteObject::Serialization::Network::ProcessRPCQueue(this);
+			// Only process queued RPCs if the actor is already registered with Iris (otherwise they will be ignored).
+			// If not yet registered, the RPCs will be processed in FlushDeferred.
+			if (UReplicationSystem* ReplicationSystem = GetNetDriver()->GetReplicationSystem())
+			{
+				if (const UObjectReplicationBridge* Bridge = ReplicationSystem->GetReplicationBridge())
+				{
+					const UE::Net::FNetRefHandle RefHandle = Bridge->GetReplicatedRefHandle(this);
+					if (RefHandle.IsValid())
+					{
+						UE::RemoteObject::Serialization::Network::ProcessRPCQueue(this);
+					}
+				}
+			}
 
 			FActorLoanedState LoanedState = RemoteObjects_ActorLoanedStates.GetAndRemoveAnnotation(this);
 			if (bActorIsBeingDestroyed)
diff --git a/Engine/Source/Runtime/Engine/Private/Components/ActorComponent.cpp b/Engine/Source/Runtime/Engine/Private/Components/ActorComponent.cpp
--- a/Engine/Source/Runtime/Engine/Private/Components/ActorComponent.cpp
+++ b/Engine/Source/Runtime/Engine/Private/Components/ActorComponent.cpp
@@ -765,10 +765,10 @@ void UActorComponent::PostMigrate(const struct FUObjectMigrationContext& Migrati
 			// NOTE: This should be re-enabled if components ever migrate separately to actors
 			//
 			//UE::Net::FDeferredReplicationSystemCalls::StopReplicatingActorComponentDefer(this, EEndReplicationFlags::Migrated);
-
-			// Prevent calling StopReplicatingActorComponent() in EndPlay(). 
-			bCallStopReplicationInEndPlay = false;
 		}
+
+		// Prevent calling StopReplicatingActorComponent() in EndPlay(). 
+		bCallStopReplicationInEndPlay = false;
 	}
 	else if (ensureMsgf(MigrationContext.MigrationSide == EObjectMigrationSide::Receive, TEXT("Invalid MigrationContext: Side was not Send or Receive")))
 	{
@@ -905,6 +905,9 @@ void UActorComponent::PostMigrate(const struct FUObjectMigrationContext& Migrati
 			}
 
 			SetComponentTickEnabled(bComponentTickEnabled);
+
+			// We now own this object, so if we EndPlay, then actually stop replication.
+			bCallStopReplicationInEndPlay = true;
 		}
 	}
 }
diff --git a/Engine/Source/Runtime/Engine/Private/NetDriver.cpp b/Engine/Source/Runtime/Engine/Private/NetDriver.cpp
--- a/Engine/Source/Runtime/Engine/Private/NetDriver.cpp
+++ b/Engine/Source/Runtime/Engine/Private/NetDriver.cpp
@@ -8210,6 +8210,10 @@ void UNetDriver::ProcessRemoteFunction(
 						return;
 					}
 				}
+				else
+				{
+					UE_LOGF(LogNet, Verbose, "SendRPC %ls::%ls dropped because Actor does not have a NetConnection", *Actor->GetName(), *Function->GetName());
+				}
 			}
 
 			// If we are using Iris replication, we should never fall back on normal replication path
diff --git a/Engine/Source/Runtime/Engine/Private/PlayerController.cpp b/Engine/Source/Runtime/Engine/Private/PlayerController.cpp
--- a/Engine/Source/Runtime/Engine/Private/PlayerController.cpp
+++ b/Engine/Source/Runtime/Engine/Private/PlayerController.cpp
@@ -5093,13 +5093,33 @@ void APlayerController::PostMigrate(const FUObjectMigrationContext& MigrationCon
 	// is expected to set a new value of NetConnection if and when required.
 	ensure(NetConnection == nullptr);
 
-	Super::PostMigrate(MigrationContext);
-	
-	if (!ensure(GetWorld()) || !ensure(GetWorld()->GetNetDriver()))
+	// Make sure we have the prerequisites
+	UWorld* World = GetWorld();
+	UNetDriver* NetDriver = World ? World->GetNetDriver() : nullptr;
+	if (!ensureMsgf(NetDriver, TEXT("%s::%hs: World or NetDriver not set."), *GetName(), __FUNCTION__))
 	{
 		return;
 	}
-	
+
+	// Update the NetConnection on the player controller because it will be set to null when serialized. It's important that the NetConnection variable
+	// is always updated, even when borrowed, because there is code that depends on that value and won't work if it's null.
+	// We need to do this before Super::PostMigrate because RPCs will not Queue without Player set.
+	if (MigrationContext.MigrationSide == EObjectMigrationSide::Receive)
+	{
+		// Set NetConnection to the local connection that matches the player id in the migrated player controller.
+		NetConnection = NetDriver->FindConnection(CachedConnectionPlayerId);
+			
+		if (!ensureMsgf(NetConnection, TEXT("APlayerController::PostMigrate: Failed to find a NetConnection associated with player id %s"), *CachedConnectionPlayerId.ToDebugString()))
+		{
+			return;
+		}
+
+		// Set Player at the same time as NetConnection because it's used in APlayerController::GetConnection().
+		Player = NetConnection;
+	}
+
+	Super::PostMigrate(MigrationContext);
+		
 	const bool bHasOwnership = UE::RemoteObject::Handle::IsOwned(this);
 	
 	if (MigrationContext.MigrationSide == EObjectMigrationSide::Receive)
@@ -5110,29 +5130,6 @@ void APlayerController::PostMigrate(const FUObjectMigrationContext& MigrationCon
 			   *FRemoteServerId::GetLocalServerId().ToString(), *FRemoteObjectId(this).ToString(), *GetName(), bHasOwnership, ToString(MigrationRecvType), *CachedConnectionPlayerId.ToDebugString());
 
 		const bool bAttachControllerToConnection = (MigrationRecvType == EObjectMigrationRecvType::AssignedOwnership);
-
-		// Update the NetConnection on the player controller because it will be set to null when serialized. It's important that the NetConnection variable
-		// is always updated, even when borrowed, because there is code that depends on that value and won't work if it's null.
-		const bool bUpdateNetConnection = (
-			bAttachControllerToConnection || 
-			MigrationRecvType == EObjectMigrationRecvType::ReturnedLoan || 
-			MigrationRecvType == EObjectMigrationRecvType::Borrowed
-		);
-
-		if (bUpdateNetConnection)
-		{
-			// Set NetConnection to the local connection that matches the player id in the migrated player controller.
-			NetConnection = GetWorld()->GetNetDriver()->FindConnection(CachedConnectionPlayerId);
-			
-			if (!ensureMsgf(NetConnection, TEXT("APlayerController::PostMigrate: Failed to find a NetConnection associated with player id %s"), *CachedConnectionPlayerId.ToDebugString()))
-			{
-				return;
-			}
-
-			// Set Player at the same time as NetConnection because it's used in APlayerController::GetConnection().
-			Player = NetConnection;
-		}
-
 		if (bAttachControllerToConnection)
 		{	
 			// It's assumed the current connection for the player id should be using an ANoPawnPlayerController.
@@ -5165,7 +5162,7 @@ void APlayerController::PostMigrate(const FUObjectMigrationContext& MigrationCon
 				CachedConnectionPlayerId = PrevNoPawnPlayerController->CachedConnectionPlayerId;
 
 				// Stop replicating the old ANoPawnPlayerController since it's no longer used.
-				RemoveFromReplicationSystem(GetWorld()->GetNetDriver(), PrevNoPawnPlayerController);
+				RemoveFromReplicationSystem(NetDriver, PrevNoPawnPlayerController);
 
 				UpdateOwningNetConnection();
 
@@ -5174,8 +5171,8 @@ void APlayerController::PostMigrate(const FUObjectMigrationContext& MigrationCon
 					   *GetName(), *NetConnection->GetName());
 
 				// There should only be one registered controller in the world for this connection and associated player.
-				GetWorld()->RemoveController(PrevNoPawnPlayerController);
-				GetWorld()->AddController(this);
+				World->RemoveController(PrevNoPawnPlayerController);
+				World->AddController(this);
 			}
 		}
 	}
@@ -5191,7 +5188,7 @@ void APlayerController::PostMigrate(const FUObjectMigrationContext& MigrationCon
 		if (bDetachControllerToConnection)
 		{
 			// Lookup the local connection that matches the player id in the player controller being migrated away.
-			UNetConnection* PrevNetConnection = GetWorld()->GetNetDriver()->FindConnection(CachedConnectionPlayerId);
+			UNetConnection* PrevNetConnection = NetDriver->FindConnection(CachedConnectionPlayerId);
 		
 			if (!ensureMsgf(PrevNetConnection, TEXT("APlayerController::PostMigrate: Failed to find a connection associated with player id %s"), *CachedConnectionPlayerId.ToDebugString()))
 			{
@@ -5234,7 +5231,7 @@ void APlayerController::PostMigrate(const FUObjectMigrationContext& MigrationCon
 				// can happen during seamless travel and the world is destroyed but the connection remains open.
 				if (PrevNetConnection->CachedNoPawnPlayerController)
 				{
-					if (PrevNetConnection->CachedNoPawnPlayerController->GetWorld() != GetWorld())
+					if (PrevNetConnection->CachedNoPawnPlayerController->GetWorld() != World)
 					{
 						UE_LOG(LogPlayerController, Verbose, TEXT("APlayerController::PostMigrate(): Destroying cached instance ANoPawnPlayerController because it was spawned in a different world."));
 
@@ -5251,7 +5248,7 @@ void APlayerController::PostMigrate(const FUObjectMigrationContext& MigrationCon
 					FActorSpawnParameters SpawnInfo;
 					SpawnInfo.ObjectFlags |= RF_Transient;
 					SpawnInfo.bDeferConstruction = true;
-					ANoPawnPlayerController* NewNoPawnPlayerController = GetWorld()->SpawnActor<ANoPawnPlayerController>(ANoPawnPlayerController::StaticClass(), FVector(), FRotator(), SpawnInfo);
+					ANoPawnPlayerController* NewNoPawnPlayerController = World->SpawnActor<ANoPawnPlayerController>(ANoPawnPlayerController::StaticClass(), FVector(), FRotator(), SpawnInfo);
 					if (ensure(NewNoPawnPlayerController))
 					{
 						// Attach this new instance of ANoPawnPlayerController to the network connection.
@@ -5280,7 +5277,7 @@ void APlayerController::PostMigrate(const FUObjectMigrationContext& MigrationCon
 					BindAndInitializeConnectionAndController(PrevNetConnection, CachedNoPawnPlayerController, this);
 				
 					// Start replicating the ANoPawnPlayerController.
-					AddToReplicationSystem(GetWorld()->GetNetDriver(), CachedNoPawnPlayerController);
+					AddToReplicationSystem(NetDriver, CachedNoPawnPlayerController);
 
 					// Since the player controller is being re-used, ensure that its initial state is the same as if it were just spawned.
 					//
@@ -5303,8 +5300,8 @@ void APlayerController::PostMigrate(const FUObjectMigrationContext& MigrationCon
 				ensure(PrevNetConnection->OwningActor);
 
 				// There should only be one registered controller in the world for this connection and associated player.
-				GetWorld()->RemoveController(this);
-				GetWorld()->AddController(PrevNetConnection->PlayerController);
+				World->RemoveController(this);
+				World->AddController(PrevNetConnection->PlayerController);
 			}
 		}
 	}
diff --git a/Engine/Source/Runtime/Net/Iris/Private/Iris/ReplicationSystem/NetBlob/NetRPC.cpp b/Engine/Source/Runtime/Net/Iris/Private/Iris/ReplicationSystem/NetBlob/NetRPC.cpp
--- a/Engine/Source/Runtime/Net/Iris/Private/Iris/ReplicationSystem/NetBlob/NetRPC.cpp
+++ b/Engine/Source/Runtime/Net/Iris/Private/Iris/ReplicationSystem/NetBlob/NetRPC.cpp
@@ -359,12 +359,19 @@ bool FNetRPC::ResolveFunctionAndObject(FNetSerializationContext& Context)
 	if (!NetObjectReference.GetRefHandle().IsValid())
 	{
 		// This can occur if sending side had queued up rpcs to object being invalidated
+		UE_LOGF(LogIrisRpc, Verbose, "%s FAILED (invalid handle): NetRef=%ls TargetRef=%ls FuncLocator=(%u:%u)",
+			__FUNCTION__, *NetObjectReference.ToString(), *TargetObjectReference.ToString(), FunctionLocator.DescriptorIndex, FunctionLocator.FunctionIndex);
 		return false;
 	}
 
 	const FReplicationStateMemberFunctionDescriptor* FunctionDescriptor = nullptr;
 	if (!NetRPC_GetFunctionAndObject(Context, NetObjectReference, TargetObjectReference, FunctionLocator, FunctionDescriptor, ObjectPtr))
 	{
+		const FString DebugFuncName = NetRPC_GetDebugFunctionName(Context.GetInternalContext()->ReplicationSystem, NetObjectReference, FunctionLocator);
+		UObject* RefObject = NetRPC_GetObject(Context, NetObjectReference, TargetObjectReference);
+		UE_LOGF(LogIrisRpc, Warning, "%s FAILED (GetFunctionAndObject): NetRef=%ls TargetRef=%ls Func=%ls ResolvedObject=%ls FuncLocator=(%u:%u)",
+			__FUNCTION__, *NetObjectReference.ToString(), *TargetObjectReference.ToString(), *DebugFuncName, *GetNameSafe(RefObject),
+			FunctionLocator.DescriptorIndex, FunctionLocator.FunctionIndex);
 		return false;
 	}
 
@@ -504,18 +511,19 @@ void FNetRPC::CallFunction(FNetRPCCallContext& CallContext)
 
 	if (Object == nullptr || Function == nullptr)
 	{
+		const FString DebugFuncName = Function ? Function->GetName() : NetRPC_GetDebugFunctionName(ReplicationSystem, NetObjectReference, FunctionLocator);
 		if (!TargetObjectReference.IsValid())
 		{
-			UE_LOG(LogIrisRpc, Error, TEXT("Rejected RPC (%u:%u) %s due to missing object or function for object: %s."),
+			UE_LOG(LogIrisRpc, Error, TEXT("Rejected RPC (%u:%u) '%s' due to missing object or function for object: %s."),
 				FunctionLocator.DescriptorIndex, FunctionLocator.FunctionIndex,
-				(Function ? ToCStr(Function->GetName()) : TEXT("\'unknown\'")),
+				*DebugFuncName,
 				*NetRPC_GetDebugObjectRefName(Context, NetObjectReference, Object));
 		}
 		else
 		{
-			UE_LOG(LogIrisRpc, Error, TEXT("Rejected RPC (%u:%u) %s due to missing object or function for subobject: %s of rootobject: %s"),
+			UE_LOG(LogIrisRpc, Error, TEXT("Rejected RPC (%u:%u) '%s' due to missing object or function for subobject: %s of rootobject: %s"),
 				FunctionLocator.DescriptorIndex, FunctionLocator.FunctionIndex,
-				(Function ? ToCStr(Function->GetName()) : TEXT("\'unknown\'")),
+				*DebugFuncName,
 				*NetRPC_GetDebugObjectRefName(Context, TargetObjectReference, Object), *NetRPC_GetDebugObjectRefName(Context, NetObjectReference));
 		}
 		return;
@@ -757,6 +765,8 @@ static bool NetRPC_GetFunctionAndObject(FNetSerializationContext& Context, const
 	if (!RefObject)
 	{
 		// Ignore this RPC and continue processing the rest of the data
+		UE_LOGF(LogIrisRpc, Warning, "%s: Object resolution FAILED. ObjRef=%ls SubObjRef=%ls Handle=%ls",
+			__FUNCTION__, *ObjectReference.ToString(), *SubObjectReference.ToString(), *ObjectReference.GetRefHandle().ToString());
         return false;
 	}
 
