Added FExternFromMaterialDecl(const FMaterialExternalCodeDeclaration* Declaration) constructor.

#rnx
#Synced-CL 51827487

#Synced-CL 51827487

#Synced-CL 51827487

diff --git a/Engine/Source/Runtime/Engine/Public/Materials/MaterialIRExtern.cpp b/Engine/Source/Runtime/Engine/Public/Materials/MaterialIRExtern.cpp
--- a/Engine/Source/Runtime/Engine/Public/Materials/MaterialIRExtern.cpp
+++ b/Engine/Source/Runtime/Engine/Public/Materials/MaterialIRExtern.cpp
@@ -47,7 +47,12 @@ bool FExternAnalysisContext::CheckStage(EStage Stage, EStageMask AllowedStages,
 //                        FExternFromMaterialDecl
 //------------------------------------------------------------------------
 
-// @massimo.tristano should FExternFromMaterialDecl be cached inside the module so that Extern referencing them are all unique?
+FExternFromMaterialDecl::FExternFromMaterialDecl(const FMaterialExternalCodeDeclaration* Declaration)
+	: Name{ Declaration->Name }
+	, Declaration{ Declaration }
+{
+}
+
 FExternFromMaterialDecl::FExternFromMaterialDecl(FName InName) 
 	: Name { InName }
 {
diff --git a/Engine/Source/Runtime/Engine/Public/Materials/MaterialIRExtern.h b/Engine/Source/Runtime/Engine/Public/Materials/MaterialIRExtern.h
--- a/Engine/Source/Runtime/Engine/Public/Materials/MaterialIRExtern.h
+++ b/Engine/Source/Runtime/Engine/Public/Materials/MaterialIRExtern.h
@@ -259,6 +259,7 @@ struct FExternFromMaterialDecl
 	const FMaterialExternalCodeDeclaration* Declaration;
 
 	// Implementation of the CExternDefinition concept.
+	FExternFromMaterialDecl(const FMaterialExternalCodeDeclaration* Declaration);
 	FExternFromMaterialDecl(FName Name);
 	MIR::FExternInfo GetInfo() const;
 	void Analyze(FExternAnalysisContext& Context);
