Makes material aggregate attribute search by name case sensitive in  UMaterialAggregate::FindAttributeIndexByName().

#rnx
#Synced-CL 51827487

diff --git a/Engine/Source/Runtime/Engine/Private/Materials/MaterialAggregate.cpp b/Engine/Source/Runtime/Engine/Private/Materials/MaterialAggregate.cpp
--- a/Engine/Source/Runtime/Engine/Private/Materials/MaterialAggregate.cpp
+++ b/Engine/Source/Runtime/Engine/Private/Materials/MaterialAggregate.cpp
@@ -106,9 +106,10 @@ static void ClearCyclicDependencies(UMaterialAggregate* Aggregate, bool bDisplay
 
 int32 UMaterialAggregate::FindAttributeIndexByName(FName Name) const
 {
+	// @massimo.tristano We might want to make this constant time
 	return Attributes.IndexOfByPredicate([Name](const FMaterialAggregateAttribute& Attribute)
 	{
-		return Attribute.Name == Name;
+		return Attribute.Name.IsEqual(Name, ENameCase::CaseSensitive);
 	});
 }
 
