Fixed handling of folding boolean Not and IsNan/IsInfinite/IsFinite in the new translator.

#rnx

#Synced-CL 51827487

diff --git a/Engine/Source/Runtime/Engine/Private/Materials/MaterialIR.cpp b/Engine/Source/Runtime/Engine/Private/Materials/MaterialIR.cpp
--- a/Engine/Source/Runtime/Engine/Private/Materials/MaterialIR.cpp
+++ b/Engine/Source/Runtime/Engine/Private/Materials/MaterialIR.cpp
@@ -698,26 +698,6 @@ int32 FInstruction::GetOutputPriority() const
 	return 1024 + Kind;
 }
 
-bool IsComparisonOperator(EOperator Op)
-{
-	switch (Op)
-	{
-		case UO_Not:
-		case UO_IsFinite:
-		case UO_IsInf:
-		case UO_IsNan:
-		case BO_Equals:
-		case BO_GreaterThan:
-		case BO_GreaterThanOrEquals:
-		case BO_LessThan:
-		case BO_LessThanOrEquals:
-		case BO_NotEquals:
-			return true;
-		default:
-			return false;
-	}
-}
-
 bool IsUnaryOperator(EOperator Op)
 {
 	return Op >= UO_FirstUnaryOperator && Op < BO_FirstBinaryOperator;
diff --git a/Engine/Source/Runtime/Engine/Private/Materials/MaterialIREmitter.cpp b/Engine/Source/Runtime/Engine/Private/Materials/MaterialIREmitter.cpp
--- a/Engine/Source/Runtime/Engine/Private/Materials/MaterialIREmitter.cpp
+++ b/Engine/Source/Runtime/Engine/Private/Materials/MaterialIREmitter.cpp
@@ -1640,24 +1640,25 @@ FValueRef FEmitter::Aggregate(const UMaterialAggregate* InAggregate, FValueRef I
 }
 
 /*--------------------------------- Operator emission ---------------------------------*/
-
+	
 template <typename T>
-static bool FoldComparisonOperatorScalar(EOperator Operator, T A, T B)
+static bool FoldFloatToBoolOperatorScalar(EOperator Operator, T A)
 {
-	if constexpr (std::is_floating_point_v<T>)
+	static_assert(std::is_floating_point_v<T>);
+	switch (Operator)
 	{
-		switch (Operator)
-		{
-			case UO_IsFinite: return FGenericPlatformMath::IsFinite(A);
-			case UO_IsInf: return !FGenericPlatformMath::IsFinite(A);
-			case UO_IsNan: return FGenericPlatformMath::IsNaN(A);
-			default: break;
-		}
+		case UO_IsFinite: return FGenericPlatformMath::IsFinite(A);
+		case UO_IsInf:    return !FGenericPlatformMath::IsFinite(A);
+		case UO_IsNan:    return FGenericPlatformMath::IsNaN(A);
+		default: UE_MIR_UNREACHABLE();
 	}
+}
 
+template <typename T>
+static bool FoldComparisonOperatorScalar(EOperator Operator, T A, T B)
+{
 	switch (Operator)
 	{
-		case UO_Not: return !A;
 		case BO_GreaterThan: return A > B;
 		case BO_GreaterThanOrEquals: return A >= B;
 		case BO_LessThan: return A < B;
@@ -2052,6 +2053,37 @@ static FValueRef TrySimplifyOperator(FEmitter& Emitter, EOperator& Op, FValueRef
 	return {};
 }
 
+// Whether the given operator is unary and returns a boolean.
+static bool IsUnaryOperatorToBool(EOperator Op)
+{
+	switch (Op)
+	{
+		case UO_IsFinite:
+		case UO_IsInf:
+		case UO_IsNan:
+			return true;
+		default:
+			return false;
+	}
+}
+
+// Whether the given operator identifies a comparison operation (e.g., ">=", "==").
+static bool IsComparisonOperator(EOperator Op)
+{
+	switch (Op)
+	{
+		case BO_Equals:
+		case BO_GreaterThan:
+		case BO_GreaterThanOrEquals:
+		case BO_LessThan:
+		case BO_LessThanOrEquals:
+		case BO_NotEquals:
+			return true;
+		default:
+			return false;
+	}
+}
+
 // Tries to fold (statically evaluate) the operator, assuming that the arguments are all scalar.
 // It returns either the result of the operator or null if it could not be folded.
 static FValueRef TryFoldOperatorScalar(FEmitter& Emitter, EOperator Op, FValueRef A, FValueRef B, FValueRef C, FType ResultType)
@@ -2080,7 +2112,29 @@ static FValueRef TryFoldOperatorScalar(FEmitter& Emitter, EOperator Op, FValueRe
 	}
 
 	// Call the appropriate helper function depending on what type of operator this is
-	if (IsComparisonOperator(Op))
+	if (Op == UO_Not)
+	{
+		return Emitter.ConstantBool(!AConstant->Boolean);
+	}
+	else if (IsUnaryOperatorToBool(Op))
+	{
+		bool Result;
+		switch (PrimitiveType->ScalarKind)
+		{
+			case EScalarKind::Float:
+				Result = FoldFloatToBoolOperatorScalar<TFloat>(Op, AConstant->Float);
+				break;
+
+			case EScalarKind::Double:
+				Result = FoldFloatToBoolOperatorScalar<TDouble>(Op, AConstant->Double);
+				break;
+			
+			default:
+				UE_MIR_UNREACHABLE();
+		}
+		return Emitter.ConstantBool(Result);
+	}
+	else if (IsComparisonOperator(Op))
 	{
 		bool Result;
 		switch (PrimitiveType->ScalarKind)
diff --git a/Engine/Source/Runtime/Engine/Public/Materials/MaterialIR.h b/Engine/Source/Runtime/Engine/Public/Materials/MaterialIR.h
--- a/Engine/Source/Runtime/Engine/Public/Materials/MaterialIR.h
+++ b/Engine/Source/Runtime/Engine/Public/Materials/MaterialIR.h
@@ -719,9 +719,6 @@ enum EOperator
 	OperatorCount,
 };
 
-// Whether the given operator identifies a comparison operation (e.g., ">=", "==").
-bool IsComparisonOperator(EOperator Op);
-
 // Whether the given operator identifies a unary operator.
 bool IsUnaryOperator(EOperator Op);
 
