Nullptr crash fix in new material translator.

#rnx
#Synced-CL 51827487

diff --git a/Engine/Source/Runtime/Engine/Private/Materials/MaterialIR.cpp b/Engine/Source/Runtime/Engine/Private/Materials/MaterialIR.cpp
--- a/Engine/Source/Runtime/Engine/Private/Materials/MaterialIR.cpp
+++ b/Engine/Source/Runtime/Engine/Private/Materials/MaterialIR.cpp
@@ -654,7 +654,7 @@ bool FComposite::AreComponentsConstant() const
 {
 	for (FValue const* Component : GetComponents())
 	{
-		if (!Component->As<FConstant>())
+		if (Component && !Component->As<FConstant>())
 		{
 			return false;
 		}
