[NetPhysics] Fix crash from legacy delta serialization flow in NetworkPhysicsComponent

#Synced-CL 51928567

diff --git a/Engine/Plugins/Experimental/ChaosModularVehicle/Source/ChaosModularVehicleEngine/Private/ChaosModularVehicle/ChaosSimModuleManagerAsyncCallback.cpp b/Engine/Plugins/Experimental/ChaosModularVehicle/Source/ChaosModularVehicleEngine/Private/ChaosModularVehicle/ChaosSimModuleManagerAsyncCallback.cpp
--- a/Engine/Plugins/Experimental/ChaosModularVehicle/Source/ChaosModularVehicleEngine/Private/ChaosModularVehicle/ChaosSimModuleManagerAsyncCallback.cpp
+++ b/Engine/Plugins/Experimental/ChaosModularVehicle/Source/ChaosModularVehicleEngine/Private/ChaosModularVehicle/ChaosSimModuleManagerAsyncCallback.cpp
@@ -280,19 +280,26 @@ void FModularVehicleAsyncInput::ProcessInputs()
 
 bool FNetworkModularVehicleInputs::NetSerialize(FArchive& Ar, class UPackageMap* Map, bool& bOutSuccess)
 {
+	static FNetworkModularVehicleInputs LocalDefaultDeltaSource;
 	EModuleInputQuantizationType InputQuantizationType = EModuleInputQuantizationType::Default_16Bits;
 	if (const UModularVehicleBaseComponent* ModularVehicleBaseComponent = Cast<UModularVehicleBaseComponent>(ImplementationComponent.Get()))
 	{
 		InputQuantizationType = ModularVehicleBaseComponent->InputQuantizationType;	
 	}
 
+	// If we should use delta source but the source is null, use a default source so that we at least consume the replicated bits correctly and doesn't cause a crash.
+	if (bIsUsingDeltaSerialization && !DeltaSourceData)
+	{
+		DeltaSourceData = &LocalDefaultDeltaSource;
+	}
+
 	FNetworkPhysicsData::SerializeFrames(Ar);
 
 	Ar.SerializeBits(&VehicleInputs.Reverse, 1);
 	Ar.SerializeBits(&VehicleInputs.KeepAwake, 1);
 	bOutSuccess = true;
 	FNetworkModularVehicleInputs* DeltaSource = static_cast<FNetworkModularVehicleInputs*>(DeltaSourceData);
-	if (DeltaSource && ChaosModularVehicleCVars::bEnableInputReducedBandwidth)
+	if (ChaosModularVehicleCVars::bEnableInputReducedBandwidth && bIsUsingDeltaSerialization && DeltaSource)
 	{
 		using namespace UE::Net;
 		bOutSuccess = true;
@@ -415,7 +422,7 @@ void FNetworkModularVehicleStateNetTokenData::Init(const Chaos::FModuleNetDataAr
 
 bool FNetworkModularVehicleStates::NetSerialize(FArchive& Ar, class UPackageMap* Map, bool& bOutSuccess)
 {
-	if (ChaosModularVehicleCVars::bEnableStateReducedBandwidth && DeltaSourceData != nullptr)
+	if (ChaosModularVehicleCVars::bEnableStateReducedBandwidth && bIsUsingDeltaSerialization)
 	{
 		return DeltaNetSerialize(Ar, Map, bOutSuccess);
 	}
@@ -507,11 +514,18 @@ bool FNetworkModularVehicleStates::NetSerialize(FArchive& Ar, class UPackageMap*
 
 bool FNetworkModularVehicleStates::DeltaNetSerialize(FArchive& Ar, UPackageMap* Map, bool& bOutSuccess)
 {
+	static FNetworkModularVehicleStates LocalDefaultDeltaSource;
 	const bool bPrintDebugInfo = ChaosModularVehicleCVars::bEnableStateNetSerializeDebugPrinting;
 	FBitReader* BitReader = (FBitReader*)&Ar;
 	FBitWriter* BitWriter = (FBitWriter*)&Ar;
 	using namespace UE::Net;
 
+	// If we should use delta source but the source is null, use a default source so that we at least consume the replicated bits correctly and doesn't cause a crash.
+	if (!DeltaSourceData)
+	{
+		DeltaSourceData = &LocalDefaultDeltaSource;
+	}
+
 	FNetworkPhysicsData::SerializeFrames(Ar);
 
 	FNetworkModularVehicleStates* DeltaSource = static_cast<FNetworkModularVehicleStates*>(DeltaSourceData);
diff --git a/Engine/Source/Runtime/Engine/Private/PhysicsEngine/NetworkPhysicsComponent.cpp b/Engine/Source/Runtime/Engine/Private/PhysicsEngine/NetworkPhysicsComponent.cpp
--- a/Engine/Source/Runtime/Engine/Private/PhysicsEngine/NetworkPhysicsComponent.cpp
+++ b/Engine/Source/Runtime/Engine/Private/PhysicsEngine/NetworkPhysicsComponent.cpp
@@ -299,6 +299,7 @@ bool FNetworkPhysicsRewindDataProxy::NetSerializeBase(FArchive& Ar, class UPacka
 							// Set the component pointer to the implementation that uses this data
 							NetData->SetImplementationComponent(Owner.Get()->ActorComponent.Get());
 
+							NetData->bIsUsingDeltaSerialization = false;
 							// Only use stateful delta source for the first entry in history, the following entries will use the previous entry as delta source
 							if (PhysicsReplicationCVars::ResimulationCVars::bEnableStatefulDeltaSerialization && DataIndex == 0 && GetLatestDeltaSourceIndex && GetDeltaSourceData)
 							{
@@ -313,12 +314,24 @@ bool FNetworkPhysicsRewindDataProxy::NetSerializeBase(FArchive& Ar, class UPacka
 									{
 										ensureMsgf(false, TEXT("Delta Serialization failed to get the latest delta source when sending, should not happen. On the first send the latest index should be populated with a default value, not null."));
 										NetData->SetDeltaSourceData(GetDeltaSourceData(/*Default*/ -2));
+
+										// Set "index" to the buffer size, meaning it's invalid
+										DeltaSourceIndex = NetworkPhysicsComponentConstants::DeltaSourceBufferSize;
 									}
 
-									constexpr uint32 NumBitsDeltaBufferSize = FMath::CeilLogTwo(NetworkPhysicsComponentConstants::DeltaSourceBufferSize);
-									Ar.SerializeBits(&DeltaSourceIndex, NumBitsDeltaBufferSize);
+									Ar.SerializeBits(&NetData->bIsUsingDeltaSerialization, 1);
+
+									if (NetData->bIsUsingDeltaSerialization)
+									{
+										constexpr uint32 NumBitsDeltaBufferSize = FMath::CeilLogTwo(NetworkPhysicsComponentConstants::DeltaSourceBufferSize);
+										Ar.SerializeBits(&DeltaSourceIndex, NumBitsDeltaBufferSize);
+									}
 								}
 							}
+							else
+							{
+								Ar.SerializeBits(&NetData->bIsUsingDeltaSerialization, 1);
+							}
 						}
 					}
 				});
@@ -359,17 +372,25 @@ bool FNetworkPhysicsRewindDataProxy::NetSerializeBase(FArchive& Ar, class UPacka
 							// Set the component pointer to the implementation that uses this data
 							NetData->SetImplementationComponent(Owner.Get()->ActorComponent.Get());
 
-							// Only use stateful delta source for the first entry in history, the following entries will use the previous entry as delta source
-							if (PhysicsReplicationCVars::ResimulationCVars::bEnableStatefulDeltaSerialization && DataIndex == 0 && GetDeltaSourceData)
+							Ar.SerializeBits(&NetData->bIsUsingDeltaSerialization, 1);
+
+							// Stateful Delta Serialization
+							if (NetData->bIsUsingDeltaSerialization)
 							{
-								// Stateful Delta Serialization
+								uint32 DeltaSourceIndex = 0;
+								constexpr uint32 NumBitsDeltaBufferSize = FMath::CeilLogTwo(NetworkPhysicsComponentConstants::DeltaSourceBufferSize);
+								Ar.SerializeBits(&DeltaSourceIndex, NumBitsDeltaBufferSize); 
+
+								// Only use stateful delta source for the first entry in history, the following entries will use the previous entry as delta source
+								if (PhysicsReplicationCVars::ResimulationCVars::bEnableStatefulDeltaSerialization && GetDeltaSourceData)
 								{
-									uint32 DeltaSourceIndex = 0;
-									constexpr uint32 NumBitsDeltaBufferSize = FMath::CeilLogTwo(NetworkPhysicsComponentConstants::DeltaSourceBufferSize);
-									Ar.SerializeBits(&DeltaSourceIndex, NumBitsDeltaBufferSize);
 
-									FNetworkPhysicsData* DeltaSourceData = GetDeltaSourceData(static_cast<int32>(DeltaSourceIndex));;
-									if (!DeltaSourceData)
+									FNetworkPhysicsData* DeltaSourceData = nullptr;
+									if (DeltaSourceIndex < NetworkPhysicsComponentConstants::DeltaSourceBufferSize)
+									{
+										DeltaSourceData = GetDeltaSourceData(static_cast<int32>(DeltaSourceIndex));;
+									}
+									else
 									{
 										// Sender used default as delta source
 										DeltaSourceData = GetDeltaSourceData(/*Default*/ -2);
@@ -385,7 +406,8 @@ bool FNetworkPhysicsRewindDataProxy::NetSerializeBase(FArchive& Ar, class UPacka
 										DeltaSourceData = GetDeltaSourceData(/*Default*/ -2);
 									}
 
-									NetData->SetDeltaSourceData(DeltaSourceData);
+									// Don't use the SetDeltaSourceData API since it also override bIsUsingDeltaSerialization depending on if delta source is null or not, but here we know delta was used even if we can't find a valid delta source.
+									NetData->DeltaSourceData = DeltaSourceData;
 								}
 							}
 						}
diff --git a/Engine/Source/Runtime/Engine/Public/Physics/NetworkPhysicsComponent.h b/Engine/Source/Runtime/Engine/Public/Physics/NetworkPhysicsComponent.h
--- a/Engine/Source/Runtime/Engine/Public/Physics/NetworkPhysicsComponent.h
+++ b/Engine/Source/Runtime/Engine/Public/Physics/NetworkPhysicsComponent.h
@@ -33,8 +33,10 @@ DECLARE_MULTICAST_DELEGATE_OneParam(FAddResimulationRequestInternal, const int32
 
 namespace NetworkPhysicsComponentConstants
 {
-	/** Size of the array caching Delta Sources for delta serialization. */
-	static constexpr int32 DeltaSourceBufferSize = 4;
+	/** Size of the array caching Delta Sources for delta serialization.
+	* NOTE: This is also used to set the number of bits to serialize when sending the delta source index and we need one extra value for "invalid" delta source 
+	* So if changing this, set it to a PowerOfTwo - 1 to use all bits available for example: 3 will give index values of 0, 1, 2 and we will use 3 as the invalid value meaning we will have 4 values and 4 values fits int 2 bits */
+	static constexpr int32 DeltaSourceBufferSize = 3;
 
 	/** Max number of elements when replicating inputs and states. Example, latest input along with the last N inputs will be sent in an array to cover for packet loss.
 	* See CVar: np2.Resim.RedundantInputs or np2.Resim.DynamicInputReplicationScaling.Enabled */
@@ -1165,11 +1167,18 @@ struct FNetworkPhysicsData : public FNetworkPhysicsPayload
 	/** Pointer to a previous FNetworkPhysicsData which is valid during NetSerialize() to be used for delta serialization */
 	FNetworkPhysicsData* DeltaSourceData = nullptr;
 
-	void SetDeltaSourceData(FNetworkPhysicsData* IntDeltaSourceData)
+	bool bIsUsingDeltaSerialization = false;
+
+	void SetDeltaSourceData(FNetworkPhysicsData* InDeltaSourceData)
+	{
+		DeltaSourceData = InDeltaSourceData;
+		bIsUsingDeltaSerialization = DeltaSourceData != nullptr;
+	}
+	void ClearDeltaSourceData()
 	{
-		DeltaSourceData = IntDeltaSourceData ? IntDeltaSourceData : nullptr;
+		DeltaSourceData = nullptr;
+		bIsUsingDeltaSerialization = false;
 	}
-	void ClearDeltaSourceData() { DeltaSourceData = nullptr; }
 
 	// Serialize the data into/from the archive
 	void SerializeFrames(FArchive& Ar)
