[NetPhysics] Implement new networked Actions feature into the NetworkPhysicsComponent.
- Actions are a separate concept compared to inputs and states.
- Actions can be forced to only be authored on the server and synchronized to the clients
- Actions can be scheduled for a set time from "now" or a specific frame
- Actions can be predicted on clients and those predictions will get stomped by confirmed actions from the server
- Actions can be proposed by clients to the server. The client will apply the action predictively and replicate it to the server which will need to produce the same action and then compare that action with the proposed action from the client, if the proposed action is "close enough" then use the proposed action. As a way of reducing desyncing of the autonomous proxy.
- Actions are re-predicted and applied during resimulation
- If actions are being proposed they are compared during resimulation and if the resim produces a new action not yet proposed that action will be replicated to the server as a new proposal for the corresponding frame.

#Synced-CL 51758418

diff --git a/Engine/Source/Runtime/Engine/Private/PhysicsEngine/NetworkPhysicsComponent.cpp b/Engine/Source/Runtime/Engine/Private/PhysicsEngine/NetworkPhysicsComponent.cpp
--- a/Engine/Source/Runtime/Engine/Private/PhysicsEngine/NetworkPhysicsComponent.cpp
+++ b/Engine/Source/Runtime/Engine/Private/PhysicsEngine/NetworkPhysicsComponent.cpp
@@ -14,6 +14,11 @@
 #include "Iris/ReplicationState/PropertyNetSerializerInfoRegistry.h"
 #include "Iris/ReplicationState/ReplicationStateDescriptor.h"
 #include "Iris/ReplicationSystem/ReplicationOperations.h"
+#include "Iris/ReplicationSystem/ObjectReplicationBridge.h"
+#include "Net/Iris/ReplicationSystem/EngineReplicationBridge.h"
+#include "Net/Iris/ReplicationSystem/ReplicationSystemUtil.h"
+#include "Iris/ReplicationSystem/NetRefHandle.h"
+#include "Misc/Crc.h"
 
 #if UE_WITH_REMOTE_OBJECT_HANDLE
 #include "UObject/UObjectMigrationContext.h"
@@ -231,6 +236,11 @@ namespace PhysicsReplicationCVars
 		bool bEnableLagScalingSimProxyRuntimeInputDecay = false;
 		static FAutoConsoleVariableRef CVarEnableLagScalingSimProxyRuntimeInputDecay(TEXT("np2.Resim.EnableLagScalingSimProxyRuntimeInputDecay"), bEnableLagScalingSimProxyRuntimeInputDecay, TEXT("When true, scale 'runtime' simulated proxy input decay with measured lag, just like resim input decay"));
 
+		bool bActionsEnableDebugLogs = false;
+		static FAutoConsoleVariableRef CVarNetworkedActionsEnableDebugLogs(TEXT("np2.Resim.NetworkedActions.EnableDebugLogs"), bActionsEnableDebugLogs, TEXT("Enable logs for the networked Actions flow in non-shipping builds."));
+		int32 bActionsEquivalenceFrameWindow = 1;
+		static FAutoConsoleVariableRef CVarNetworkedActionsEquivalenceFrameWindow(TEXT("np2.Resim.NetworkedActions.EquivalenceFrameWindow"), bActionsEquivalenceFrameWindow, TEXT("Number of frames to look forward and behind when looking for equivalent actions in proposals when server produces an action or client re-produces an acton during resim, as well as in predicted actions when receiving a confirmed action from the server. A value of N means look at (CurrentFrame - N), CurrentFrame and (CurrentFrame + N)."));
+
 		bool bSimulationDecayEnable = false;
 		static FAutoConsoleVariableRef CVarEnableSimulationDecay(TEXT("np2.Resim.SimulationDecay.Enable"), bSimulationDecayEnable, TEXT("When enabled, apply simulation decay on simulated proxies by essentially multiplying their physics step by a time scale value. Effectively forward predicting sim-proxies less, which can make desync less during forward prediction."));
 		bool bSimulationDecayApplyAtRuntime = true;
@@ -1012,6 +1022,7 @@ void UNetworkPhysicsComponent::GetLifetimeReplicatedProps(TArray< FLifetimePrope
 	DOREPLIFETIME_WITH_PARAMS_FAST(UNetworkPhysicsComponent, ReplicatedInputCollection, ReplicatedParamsOwner);
 	DOREPLIFETIME_WITH_PARAMS_FAST(UNetworkPhysicsComponent, ReplicatedRemoteInputCollection, ReplicatedParamsRemote);
 	DOREPLIFETIME_WITH_PARAMS_FAST(UNetworkPhysicsComponent, ReplicatedStateCollection, ReplicatedParamsAll);
+	DOREPLIFETIME_WITH_PARAMS_FAST(UNetworkPhysicsComponent, ReplicatedActionsCollection, ReplicatedParamsAll);
 }
 
 void UNetworkPhysicsComponent::PreReplication(IRepChangedPropertyTracker& ChangedPropertyTracker)
@@ -1029,6 +1040,7 @@ void UNetworkPhysicsComponent::PreReplication(IRepChangedPropertyTracker& Change
 	DOREPLIFETIME_ACTIVE_OVERRIDE_FAST(UNetworkPhysicsComponent, ReplicatedInputCollection, !bIsUsingLegacyData);
 	DOREPLIFETIME_ACTIVE_OVERRIDE_FAST(UNetworkPhysicsComponent, ReplicatedRemoteInputCollection, !bIsUsingLegacyData);
 	DOREPLIFETIME_ACTIVE_OVERRIDE_FAST(UNetworkPhysicsComponent, ReplicatedStateCollection, !bIsUsingLegacyData);
+	DOREPLIFETIME_ACTIVE_OVERRIDE_FAST(UNetworkPhysicsComponent, ReplicatedActionsCollection, true);
 }
 
 // Called every Game Thread frame
@@ -1283,6 +1295,26 @@ void UNetworkPhysicsComponent::NetworkMarshaledData()
 			}
 		}
 
+		// Unreliable Actions
+		if (AsyncOutput->ActionData.Num() > 0)
+		{
+			ReplicatedActionsCollection.DataArray.SetNum(AsyncOutput->ActionData.Num(), EAllowShrinking::Yes);
+			int32 Idx = 0;
+			for (const TInstancedStruct<FNetworkPhysicsActionPayload>& ActionInstance : AsyncOutput->ActionData)
+			{
+				ReplicatedActionsCollection.DataArray[Idx++] = ActionInstance;
+			}
+
+			if (bIsServer)
+			{
+				MARK_PROPERTY_DIRTY_FROM_NAME(UNetworkPhysicsComponent, ReplicatedActionsCollection, this);
+			}
+			else
+			{
+				ServerReceiveActionCollection(ReplicatedActionsCollection);
+			}
+		}
+
 		if (bStopRelayingLocalInputsDeferred)
 		{
 			bIsRelayingLocalInputs = false;
@@ -1793,6 +1825,25 @@ void UNetworkPhysicsComponent::ServerReceiveInputCollection_Implementation(const
 	}
 }
 
+void UNetworkPhysicsComponent::ServerReceiveActionCollection_Implementation(const FNetworkPhysicsActionCollection& ClientActionCollection)
+{
+	if (!NetworkPhysicsComponent_Internal || ClientActionCollection.DataArray.Num() <= 0)
+	{
+		return;
+	}
+
+	if (FAsyncNetworkPhysicsComponentInput* AsyncInput = NetworkPhysicsComponent_Internal->GetProducerInputData_External())
+	{
+		for (const TInstancedStruct<FNetworkPhysicsActionPayload>& ActionInstance : ClientActionCollection.DataArray)
+		{
+			if (ActionInstance.IsValid())
+			{
+				AsyncInput->ReplicatedActionData.Add(ActionInstance);
+			}
+		}
+	}
+}
+
 // repnotify for inputs on owner client
 void UNetworkPhysicsComponent::OnRep_SetReplicatedInputCollection()
 {
@@ -1862,6 +1913,112 @@ void UNetworkPhysicsComponent::OnRep_SetReplicatedStateCollection()
 	}
 }
 
+void UNetworkPhysicsComponent::OnRep_SetReplicatedActionsCollection()
+{
+	if (!NetworkPhysicsComponent_Internal || ReplicatedActionsCollection.DataArray.Num() <= 0)
+	{
+		return;
+	}
+
+	if (FAsyncNetworkPhysicsComponentInput* AsyncInput = NetworkPhysicsComponent_Internal->GetProducerInputData_External())
+	{
+		for (const TInstancedStruct<FNetworkPhysicsActionPayload>& ActionInstance : ReplicatedActionsCollection.DataArray)
+		{
+			if (ActionInstance.IsValid())
+			{
+				AsyncInput->ReplicatedActionData.Add(ActionInstance);
+			}
+		}
+	}
+}
+
+void UNetworkPhysicsComponent::EnqueueImmediateActionInstance_External(TInstancedStruct<FNetworkPhysicsActionPayload>&& ActionInstance, const UObject* SourceObject)
+{
+	const uint32 SourceId = UE::NetworkPhysicsUtils::GetNetworkStableHash_External(SourceObject);
+	ensureMsgf(SourceId > 0u, TEXT("EnqueueImmediateActionInstance_External could not find a network stable source to base an ID from, the object needs to either be replicated or baked in the level. Else it's recommended to pass in a custom unique SourceId as a parameter instead of the source object. Action: %ls  ---  Source: %ls")
+		, ActionInstance.GetScriptStruct() ? *ActionInstance.GetScriptStruct()->GetName() : TEXT("UNKNOWN ACTION")
+		, SourceObject ? *SourceObject->GetName() : TEXT("UNKNOWN SORUCE OBJECT"));
+
+	EnqueueImmediateActionInstance_External(MoveTemp(ActionInstance), SourceId);
+}
+
+void UNetworkPhysicsComponent::EnqueueImmediateActionInstance_External(TInstancedStruct<FNetworkPhysicsActionPayload>&& ActionInstance, const uint32 SourceId)
+{
+	EnqueueScheduledActionInstanceAtFrame_External(MoveTemp(ActionInstance), SourceId, -1);
+}
+
+void UNetworkPhysicsComponent::EnqueueScheduledActionInstance_External(TInstancedStruct<FNetworkPhysicsActionPayload>&& ActionInstance, const UObject* SourceObject, const float DelaySeconds)
+{
+	const uint32 SourceId = UE::NetworkPhysicsUtils::GetNetworkStableHash_External(SourceObject);
+	ensureMsgf(SourceId > 0u, TEXT("EnqueueScheduledActionInstance_External could not find a network stable source to base an ID from, the object needs to either be replicated or baked in the level. Else it's recommended to pass in a custom unique SourceId as a parameter instead of the source object. Action: %ls  ---  Source: %ls")
+		, ActionInstance.GetScriptStruct() ? *ActionInstance.GetScriptStruct()->GetName() : TEXT("UNKNOWN ACTION")
+		, SourceObject ? *SourceObject->GetName() : TEXT("UNKNOWN SORUCE OBJECT"));
+
+	EnqueueScheduledActionInstance_External(MoveTemp(ActionInstance), SourceId, DelaySeconds);
+}
+
+void UNetworkPhysicsComponent::EnqueueScheduledActionInstance_External(TInstancedStruct<FNetworkPhysicsActionPayload>&& ActionInstance, const uint32 SourceId, const float DelaySeconds)
+{
+	int32 LocalFrame = INDEX_NONE;
+	if (UWorld* World = GetWorld())
+	{
+		if (FPhysScene* PhysScene = World->GetPhysicsScene())
+		{
+			if (Chaos::FPhysicsSolverBase* Solver = static_cast<Chaos::FPhysicsSolverBase*>(PhysScene->GetSolver()))
+			{
+				if (Solver->IsUsingFixedDt())
+				{
+					// Only apply a frame if we are actually delaying, else leave as INDEX_NONE which will queue up the action ASAP and handle edge-cases.
+					const int32 DelayFrames = FMath::CeilToInt32(DelaySeconds / Solver->GetAsyncDeltaTime());
+					if (DelayFrames > 0)
+					{
+						LocalFrame = Solver->GetMarshallingManager().GetInternalStep_External() + DelayFrames;
+					}
+				}
+			}
+		}
+	}
+	EnqueueScheduledActionInstanceAtFrame_External(MoveTemp(ActionInstance), SourceId, LocalFrame);
+}
+
+void UNetworkPhysicsComponent::EnqueueScheduledActionInstanceAtFrame_External(TInstancedStruct<FNetworkPhysicsActionPayload>&& ActionInstance, const UObject* SourceObject, const int32 LocalFrame)
+{
+	const uint32 SourceId = UE::NetworkPhysicsUtils::GetNetworkStableHash_External(SourceObject);
+	ensureMsgf(SourceId > 0u, TEXT("EnqueueScheduledActionInstanceAtFrame_External could not find a network stable source to base an ID from, the object needs to either be replicated or baked in the level. Else it's recommended to pass in a custom unique SourceId as a parameter instead of the source object. Action: %ls  ---  Source: %ls")
+		, ActionInstance.GetScriptStruct() ? *ActionInstance.GetScriptStruct()->GetName() : TEXT("UNKNOWN ACTION")
+		, SourceObject ? *SourceObject->GetName() : TEXT("UNKNOWN SORUCE OBJECT"));
+
+	EnqueueScheduledActionInstanceAtFrame_External(MoveTemp(ActionInstance), SourceId, LocalFrame);
+}
+
+void UNetworkPhysicsComponent::EnqueueScheduledActionInstanceAtFrame_External(TInstancedStruct<FNetworkPhysicsActionPayload>&& ActionInstance, const uint32 SourceId, const int32 LocalFrame)
+{
+	if (!NetworkPhysicsComponent_Internal)
+	{
+		return;
+	}
+
+	FNetworkPhysicsActionPayload& Action = ActionInstance.GetMutable<FNetworkPhysicsActionPayload>();
+	Action.SourceId = SourceId;
+
+	// If we have a LocalFrame passed in, convert that to ServerFrame, if INDEX_NONE set ServerFrame to that which indicates that the action should be applied ASAP
+	Action.ServerFrame = (LocalFrame == INDEX_NONE) ? INDEX_NONE : (LocalFrame + UE::NetworkPhysicsUtils::GetNetworkPhysicsTickOffset_External(GetWorld()));
+
+#if !(UE_BUILD_SHIPPING || UE_BUILD_TEST)
+	if (PhysicsReplicationCVars::ResimulationCVars::bActionsEnableDebugLogs)
+	{
+		UE_LOGF(LogChaos, Log, "GT %ls EnqueueActionAtFrame - Type: %ls - SourceId: %ld - LocalFrame: %d + Offset: %d = ServerFrame: %d"
+			, (HasServerWorld() ? TEXT(" -- SERVER -- ") : TEXT(" ++ CLIENT ++ ")), (ActionInstance.GetScriptStruct() ? *ActionInstance.GetScriptStruct()->GetName() : TEXT("UNKNOWN ACTION"))
+			, Action.SourceId, LocalFrame, UE::NetworkPhysicsUtils::GetNetworkPhysicsTickOffset_External(GetWorld()), Action.ServerFrame);
+	}
+#endif
+
+	if (FAsyncNetworkPhysicsComponentInput* AsyncInput = NetworkPhysicsComponent_Internal->GetProducerInputData_External())
+	{
+		AsyncInput->LocalActionData.Add(MoveTemp(ActionInstance));
+	}
+}
+
 
 bool UNetworkPhysicsComponent::HasServerWorld() const
 {
@@ -2350,7 +2507,7 @@ void FAsyncNetworkPhysicsComponent::ConsumeAsyncInput_Internal(const int32 Physi
 					}
 					else
 					{
-						if (InputHelper && ensure(ImplementationInterface_Internal))
+						if (InputHelper && ImplementationInterface_Internal)
 						{
 							InputHelper->ValidateData(AsyncInput->InputData.Get(), *ImplementationInterface_Internal);
 						}
@@ -2414,10 +2571,408 @@ void FAsyncNetworkPhysicsComponent::ConsumeAsyncInput_Internal(const int32 Physi
 					InputPredictionFramesAverage += (static_cast<float>(InputPredictedFrames) - InputPredictionFramesAverage) * AverageAlpha;
 				}
 			}
+
+			// Receive Replicated Actions
+			if (AsyncInput->ReplicatedActionData.Num() > 0 && !bTriggerResim) // bTriggerResim TEMPORARY, this function gets called twice, remove when we can add multiple of the same actions without duplication
+			{
+				for (const TInstancedStruct<FNetworkPhysicsActionPayload>& ActionInstance : AsyncInput->ReplicatedActionData)
+				{
+					RegisterAction_Internal(ActionInstance, /*bWasReceivedFromReplication*/ true);
+				}
+			}
+
+			// Receive Local Actions
+			if (AsyncInput->LocalActionData.Num() > 0 && !bTriggerResim) // bTriggerResim TEMPORARY, this function gets called twice, remove when we can add multiple of the same actions without duplication
+			{
+				for (const TInstancedStruct<FNetworkPhysicsActionPayload>& ActionInstance : AsyncInput->LocalActionData)
+				{
+					RegisterAction_Internal(ActionInstance, /*bWasReceivedFromReplication*/ false);
+				}
+			}
 		}
 	}
 }
 
+void FAsyncNetworkPhysicsComponent::RegisterAction_Internal(
+	const TInstancedStruct<FNetworkPhysicsActionPayload>& NewActionInstance,
+	const bool bWasReceivedFromReplication)
+{
+	using EActionAuthorStyle = FNetworkPhysicsActionPayload::EActionAuthorStyle;
+
+	if (!NewActionInstance.IsValid() || !GetRigidSolver() || !GetEvolution())
+	{
+		return;
+	}
+
+	if (ConfirmedActions.IsInitialized() == false || PredictedActions.IsInitialized() == false || PendingProposedActions.IsInitialized() == false)
+	{
+		const int32 NumFrames = SetupRewindData();
+		ConfirmedActions.Initialize(NumFrames);
+		PredictedActions.Initialize(NumFrames);
+		PendingProposedActions.Initialize(NumFrames);
+	}
+
+	// Helper function to get the LocalFrame for an action, if an action does not have a ServerFrame set that means it should be applied ASAP, i.e. use current LocalFrame.
+	auto GetActionLocalFrameHelper = [&](const FNetworkPhysicsActionPayload& Action, const bool bWasReceivedFromReplication, const int32 CurrentLocalFrame, const int32 TickOffset) -> int32
+		{
+			int32 ActionLocalFrame = Action.ServerFrame != INDEX_NONE ? (Action.ServerFrame - TickOffset) : CurrentLocalFrame;
+
+			// Local actions can be enqueued after this frame has already applied actions, LastAppliedActionLocalFrame gets set when we reach the point where we apply actions during the frame
+			// In that case they need to get registered for the next frame instead of the current one.
+			if (bWasReceivedFromReplication == false && Action.ServerFrame == INDEX_NONE && LastAppliedActionLocalFrame >= ActionLocalFrame)
+			{
+				ActionLocalFrame = LastAppliedActionLocalFrame + 1;
+			}
+
+			return ActionLocalFrame;
+		};
+
+	// General Properties
+	const bool bServer = IsServer();
+	const bool bLocalControl = IsLocallyControlled();
+	const bool bIsResim = GetEvolution()->IsResimming();
+	const int32 CurrentLocalFrame = GetRigidSolver()->GetCurrentFrame();
+	const int32 TickOffset = UE::NetworkPhysicsUtils::GetNetworkPhysicsTickOffset_Internal(GetSolver());
+
+	// Action Properties
+	const FNetworkPhysicsActionPayload& NewAction = FActionHistory::GetRef(NewActionInstance);
+	const int32 ActionLocalFrame = GetActionLocalFrameHelper(NewAction, bWasReceivedFromReplication, CurrentLocalFrame, TickOffset);
+	const EActionAuthorStyle AuthorStyle = NewAction.GetAuthorStyle();
+	const bool bIsProposedStyle = AuthorStyle == EActionAuthorStyle::Proposed || AuthorStyle == EActionAuthorStyle::ProposedAutonomousOnly;
+	const bool bIsAutonomousOnlyStyle = AuthorStyle == EActionAuthorStyle::PredictedAutonomousOnly || AuthorStyle == EActionAuthorStyle::ProposedAutonomousOnly;
+
+	// Proposal reuse should be exact on forward-simulated client prediction, but allowed a small frame window during resim and on the server when adopting proposals.
+	const int32 ProposalFrameWindow = (bIsResim || bServer) ? FMath::Max(PhysicsReplicationCVars::ResimulationCVars::bActionsEquivalenceFrameWindow, 0) : 0;
+
+	// Confirmed actions should suppress equivalent local prediction in a small window.
+	const int32 PredictionFrameWindow = FMath::Max(PhysicsReplicationCVars::ResimulationCVars::bActionsEquivalenceFrameWindow, 0);
+
+	if (bWasReceivedFromReplication)
+	{
+		if (bServer) // Server received replicated action
+		{
+			// Server only accepts replicated action proposals
+			if (bIsProposedStyle == false)
+			{
+				ensure(false);
+				return;
+			}
+
+			// Proposed actions should already carry a client-generated ActionId.
+			if (NewAction.ActionId == 0u)
+			{
+				ensure(false);
+				return;
+			}
+
+			// Ignore duplicate client sends of the same proposal.
+			if (PendingProposedActions.FindActiveActionById(NewAction.ActionId, /*bIncludeApplied*/true).IsValid())
+			{
+				return;
+			}
+
+			// Ignore proposal that has already been confirmed as an authoritative action.
+			if (ConfirmedActions.FindActiveActionById(NewAction.ActionId, /*bIncludeApplied*/true).IsValid())
+			{
+				return;
+			}
+
+			// Store the proposal exactly as received.
+			PendingProposedActions.AddAction(NewActionInstance, ActionLocalFrame, TickOffset);
+
+#if !(UE_BUILD_SHIPPING || UE_BUILD_TEST)
+			if (PhysicsReplicationCVars::ResimulationCVars::bActionsEnableDebugLogs)
+			{
+				UE_LOGF(LogChaos, Log, "PT %ls RegisterAction_Internal - PendingProposedActions.AddAction - Server Receive Proposal - ActionId: %d - ServerFrame: %d - bIsResim: %d", (bServer ? TEXT(" -- SERVER -- ") : TEXT(" ++ CLIENT ++ ")), NewActionInstance.Get().ActionId, (ActionLocalFrame + TickOffset), bIsResim);
+			}
+#endif
+			return;
+		}
+		else // Client received replicated action
+		{
+			// Incoming replicated actions on the client from the server are authoritative confirmed actions.
+
+			// Store the confirmed action exactly as received.
+			ConfirmedActions.AddAction(NewActionInstance, ActionLocalFrame, TickOffset);
+
+			// Retire the predicted actions equivalent to the received confirmed action
+			const bool bPredictedActionDeactivated = PredictedActions.DeactivateEquivalentActionsInRange(NewActionInstance, (ActionLocalFrame - PredictionFrameWindow), (ActionLocalFrame + PredictionFrameWindow));
+
+#if !(UE_BUILD_SHIPPING || UE_BUILD_TEST)
+			if (PhysicsReplicationCVars::ResimulationCVars::bActionsEnableDebugLogs)
+			{
+				UE_LOGF(LogChaos, Log, "PT %ls RegisterAction_Internal - ConfirmedActions.AddAction - Client Receive Confirmed - ActionId: %d - ServerFrame : %d - bIsResim: %d", (bServer ? TEXT(" --SERVER-- ") : TEXT(" ++CLIENT++ ")), NewActionInstance.Get().ActionId, (ActionLocalFrame + TickOffset), bIsResim);
+				UE_LOGF(LogChaos, Log, "	PT %ls RegisterAction_Internal - PredictedActions.DeactivateEquivalentActionsInRange - Frames %d -> %d - Succeeded: %d - bIsResim: %d", (bServer ? TEXT(" -- SERVER -- ") : TEXT(" ++ CLIENT ++ ")), (ActionLocalFrame - PredictionFrameWindow), (ActionLocalFrame + PredictionFrameWindow), bPredictedActionDeactivated, bIsResim);
+			}
+#endif
+
+
+			if (bLocalControl == true && bIsProposedStyle == true && NewAction.ActionId != 0u)
+			{
+				// The autonomous proxy proposal has come back from the server as authority.
+				// Retire the pending proposal.
+				const bool bProposedActionDeactivated = PendingProposedActions.DeactivateAllByActionId(NewAction.ActionId);
+
+#if !(UE_BUILD_SHIPPING || UE_BUILD_TEST)
+				if (PhysicsReplicationCVars::ResimulationCVars::bActionsEnableDebugLogs)
+				{
+					UE_LOGF(LogChaos, Log, "	PT %ls RegisterAction_Internal - PendingProposedActions.DeactivateAllByActionId - ActionId: %d - Succeeded: %d - bIsResim: %d", (bServer ? TEXT(" -- SERVER -- ") : TEXT(" ++ CLIENT ++ ")), NewAction.ActionId, bProposedActionDeactivated, bIsResim);
+				}
+#endif
+			}
+
+			return;
+		}
+	}
+	else /*Locally Produced Action*/
+	{
+		if (bServer) // Server locally produced action
+		{
+			if (bIsProposedStyle == true)
+			{
+				bool bRejectedProposal = false;
+
+				// Check for a matching proposed action from the client
+				FActionHistoryLocation MatchingProposalLocation = PendingProposedActions.FindEquivalentActionInRange(NewActionInstance, (ActionLocalFrame - ProposalFrameWindow), (ActionLocalFrame + ProposalFrameWindow), /*bIncludeApplied*/false, /*bCheckSourceEquivalent*/false);
+				
+				// If we did not find an equivalent action, look for an action that matches type and source only and stomp the action data but keep the ActionId
+				if (MatchingProposalLocation.IsValid() == false)
+				{
+					MatchingProposalLocation = PendingProposedActions.FindEquivalentActionInRange(NewActionInstance, (ActionLocalFrame - ProposalFrameWindow), (ActionLocalFrame + ProposalFrameWindow), /*bIncludeApplied*/false, /*bCheckSourceEquivalent*/true);
+					bRejectedProposal = true;
+				}
+
+				if (MatchingProposalLocation.IsValid())
+				{
+					const TInstancedStruct<FNetworkPhysicsActionPayload>* ProposedActionInstance = PendingProposedActions.GetAction(MatchingProposalLocation);
+
+					if (ProposedActionInstance != nullptr)
+					{
+						if (bRejectedProposal)
+						{
+							// Override proposal but keep ActionId
+							ConfirmedActions.AddAction(NewActionInstance, ActionLocalFrame, TickOffset, ProposedActionInstance->Get().ActionId);
+
+#if !(UE_BUILD_SHIPPING || UE_BUILD_TEST)
+							if (PhysicsReplicationCVars::ResimulationCVars::bActionsEnableDebugLogs)
+							{
+								UE_LOGF(LogChaos, Log, "PT %ls RegisterAction_Internal - ConfirmedActions.AddAction - Server Rejected Proposal - ActionId: %d - ServerFrame: %d - bIsResim: %d", (bServer ? TEXT(" -- SERVER -- ") : TEXT(" ++ CLIENT ++ ")), ProposedActionInstance->Get().ActionId, (ActionLocalFrame + TickOffset), bIsResim);
+							}
+#endif
+						}
+						else // Accepted proposal
+						{
+							const int32 ProposedLocalFrame = PendingProposedActions.GetLocalFrame(MatchingProposalLocation);
+
+							// Store the confirmed proposed action exactly as received from client (keeping the ActionId that the client set)
+							ConfirmedActions.AddAction(*ProposedActionInstance, ProposedLocalFrame, TickOffset);
+
+#if !(UE_BUILD_SHIPPING || UE_BUILD_TEST)
+							if (PhysicsReplicationCVars::ResimulationCVars::bActionsEnableDebugLogs)
+							{
+								UE_LOGF(LogChaos, Log, "PT %ls RegisterAction_Internal - ConfirmedActions.AddAction - Server Confirmed Proposal - ActionId: %d - ServerFrame: %d - bIsResim: %d", (bServer ? TEXT(" -- SERVER -- ") : TEXT(" ++ CLIENT ++ ")), ProposedActionInstance->Get().ActionId, (ProposedLocalFrame + TickOffset), bIsResim);
+							}
+#endif
+						}
+
+						const bool bProposedActionDeleted = PendingProposedActions.Deactivate(MatchingProposalLocation);
+
+#if !(UE_BUILD_SHIPPING || UE_BUILD_TEST)
+						if (PhysicsReplicationCVars::ResimulationCVars::bActionsEnableDebugLogs)
+						{
+							UE_LOGF(LogChaos, Log, "	PT %ls RegisterAction_Internal - PendingProposedActions.Deactivate - Succeeded: %d - bIsResim: %d", (bServer ? TEXT(" -- SERVER -- ") : TEXT(" ++ CLIENT ++ ")), bProposedActionDeleted, bIsResim);
+						}
+#endif
+						return;
+					}
+				}
+			}
+
+			// Store the confirmed action and set a unique ActionId
+			const uint32 NewActionId = GetUniqueActionId_Internal();
+			ConfirmedActions.AddAction(NewActionInstance, ActionLocalFrame, TickOffset, NewActionId);
+			
+#if !(UE_BUILD_SHIPPING || UE_BUILD_TEST)
+			if (PhysicsReplicationCVars::ResimulationCVars::bActionsEnableDebugLogs)
+			{
+				UE_LOGF(LogChaos, Log, "PT %ls RegisterAction_Internal - ConfirmedActions.AddAction - Server Local Action - ActionId: %d - ServerFrame: %d - bIsResim: %d", (bServer ? TEXT(" -- SERVER -- ") : TEXT(" ++ CLIENT ++ ")), NewActionId, (ActionLocalFrame + TickOffset), bIsResim);
+			}
+#endif
+			return;
+		}
+		else // Client locally produced action
+		{
+			const bool bCanPredictOnThisClient = AuthorStyle != EActionAuthorStyle::Authority && (bLocalControl || bIsAutonomousOnlyStyle == false);
+
+			if (bCanPredictOnThisClient == false)
+			{
+				return;
+			}
+
+			// If the server has already established equivalent authority for this frame or just ahead of it, do not add another local prediction.
+			if (ConfirmedActions.HasEquivalentActionInRange(NewActionInstance, ActionLocalFrame, (ActionLocalFrame + PredictionFrameWindow)))
+			{
+				return;
+			}
+
+			if (bIsProposedStyle == true && bLocalControl == true)
+			{
+				// Autonomous-proxy proposing an action
+
+				if (bIsResim)
+				{
+					// During resim, reuse the existing pending proposal if the same logical action is regenerated within the allowed frame window.
+					const FActionHistoryLocation ExistingPendingProposalLocation = PendingProposedActions.FindEquivalentActionInRange(NewActionInstance, (ActionLocalFrame - ProposalFrameWindow), (ActionLocalFrame + ProposalFrameWindow), /*bIncludeApplied*/false, /*bCheckSourceEquivalent*/true);
+
+					if (ExistingPendingProposalLocation.IsValid())
+					{
+						const TInstancedStruct<FNetworkPhysicsActionPayload>* ExistingPendingProposal = PendingProposedActions.GetAction(ExistingPendingProposalLocation);
+
+						if (ExistingPendingProposal != nullptr)
+						{
+							const int32 ExistingPendingProposalLocalFrame = PendingProposedActions.GetLocalFrame(ExistingPendingProposalLocation);
+
+							// Store the equivalent proposed action as predicted and clear ActionId since no predicted action has an ActionId
+							PredictedActions.AddAction(*ExistingPendingProposal, ExistingPendingProposalLocalFrame, TickOffset
+								, /*OverrideActionId*/0u, /*bClearActionId=*/true, /*bCheckDuplicateInSameFrame=*/true);
+
+#if !(UE_BUILD_SHIPPING || UE_BUILD_TEST)
+							if (PhysicsReplicationCVars::ResimulationCVars::bActionsEnableDebugLogs)
+							{
+								UE_LOGF(LogChaos, Log, "PT %ls RegisterAction_Internal - PredictedActions.AddAction - Client Reuse Proposal - ActionId: %d - ServerFrame : %d - bIsResim: %d", (bServer ? TEXT(" --SERVER-- ") : TEXT(" ++CLIENT++ ")), ExistingPendingProposal->Get().ActionId, (ExistingPendingProposalLocalFrame + TickOffset), bIsResim);
+							}
+#endif
+							return;
+						}
+					}
+				}
+
+				const uint32 NewActionId = GetUniqueActionId_Internal();
+				// Store the proposal action and set a unique ActionId
+				PendingProposedActions.AddAction(NewActionInstance, ActionLocalFrame, TickOffset, NewActionId);
+
+#if !(UE_BUILD_SHIPPING || UE_BUILD_TEST)
+				if (PhysicsReplicationCVars::ResimulationCVars::bActionsEnableDebugLogs)
+				{
+					UE_LOGF(LogChaos, Log, "	PT %ls RegisterAction_Internal - PendingProposedActions.AddAction - Client Local Proposal - ActionId: %d - ServerFrame: %d - bIsResim: %d", (bServer ? TEXT(" -- SERVER -- ") : TEXT(" ++ CLIENT ++ ")), NewActionId, (ActionLocalFrame + TickOffset), bIsResim);
+				}
+#endif
+			}
+
+			// Store the predicted action
+			PredictedActions.AddAction(NewActionInstance, ActionLocalFrame, TickOffset
+				, /*OverrideActionId*/0u, /*bClearActionId=*/true, /*bCheckDuplicateInSameFrame=*/true);
+
+#if !(UE_BUILD_SHIPPING || UE_BUILD_TEST)
+			if (PhysicsReplicationCVars::ResimulationCVars::bActionsEnableDebugLogs)
+			{
+				UE_LOGF(LogChaos, Log, "PT %ls RegisterAction_Internal - PredictedActions.AddAction - Client Local Prediction - ActionId: %d - ServerFrame: %d - bIsResim: %d", (bServer ? TEXT(" -- SERVER -- ") : TEXT(" ++ CLIENT ++ ")), NewActionInstance.Get().ActionId, (ActionLocalFrame + TickOffset), bIsResim);
+			}
+#endif
+			return;
+		}
+	}
+}
+
+void FAsyncNetworkPhysicsComponent::ApplyActionsForFrame(const int32 LocalFrame)
+{
+	if (!ImplementationInterface_Internal || !GetRigidSolver() || !GetEvolution())
+	{
+		return;
+	}
+
+	if (ConfirmedActions.IsInitialized() == false || PredictedActions.IsInitialized() == false)
+	{
+		return;
+	}
+
+	const bool bIsResim = GetEvolution()->IsResimming();
+
+	// From this point onward, any new enqueued local action should land on the next frame.
+	LastAppliedActionLocalFrame = LocalFrame;
+
+	// Apply authoritative confirmed actions for this exact local frame.
+	FActionHistoryFrameData* ConfirmedFrameData = ConfirmedActions.FindFrameData(LocalFrame);
+	if (ConfirmedFrameData != nullptr)
+	{
+		for (FActionHistoryEntry& ConfirmedActionEntry : ConfirmedFrameData->Entries)
+		{
+			if (FActionHistory::IsUsableEntry(ConfirmedActionEntry, /*bAllowApplied*/bIsResim) == false)
+			{
+				continue;
+			}
+
+#if !(UE_BUILD_SHIPPING || UE_BUILD_TEST)
+			if (PhysicsReplicationCVars::ResimulationCVars::bActionsEnableDebugLogs)
+			{
+				UE_LOGF(LogChaos, Log, "PT %ls ApplyActionsForFrame - Apply Confirmed Action - Type: %ls - ServerFrame: %d - bIsResim: %d"
+					, (IsServer() ? TEXT(" -- SERVER -- ") : TEXT(" ++ CLIENT ++ "))
+					, ConfirmedActionEntry.ActionInstance.GetScriptStruct() ? *ConfirmedActionEntry.ActionInstance.GetScriptStruct()->GetName() : TEXT("UNKNOWN ACTION")
+					, (LocalFrame + UE::NetworkPhysicsUtils::GetNetworkPhysicsTickOffset_Internal(GetSolver())), bIsResim);
+			}
+#endif
+
+			ImplementationInterface_Internal->ApplyAction(ConfirmedActionEntry.ActionInstance);
+
+			ConfirmedActionEntry.bApplied = true;
+		}
+	}
+
+	// Apply local predicted actions for this exact local frame.
+	FActionHistoryFrameData* PredictedFrameData = PredictedActions.FindFrameData(LocalFrame);
+	if (PredictedFrameData != nullptr)
+	{
+		for (FActionHistoryEntry& PredictedActionEntry : PredictedFrameData->Entries)
+		{
+			if (FActionHistory::IsUsableEntry(PredictedActionEntry, /*bAllowApplied*/false) == false)
+			{
+				continue;
+			}
+
+#if !(UE_BUILD_SHIPPING || UE_BUILD_TEST)
+			if (PhysicsReplicationCVars::ResimulationCVars::bActionsEnableDebugLogs)
+			{
+				UE_LOGF(LogChaos, Log, "PT %ls ApplyActionsForFrame - Apply Predicted Action - Type: %ls - ServerFrame: %d - bIsResim: %d"
+					, (IsServer() ? TEXT(" -- SERVER -- ") : TEXT(" ++ CLIENT ++ "))
+					, PredictedActionEntry.ActionInstance.GetScriptStruct() ? *PredictedActionEntry.ActionInstance.GetScriptStruct()->GetName() : TEXT("UNKNOWN ACTION")
+					, (LocalFrame + UE::NetworkPhysicsUtils::GetNetworkPhysicsTickOffset_Internal(GetSolver())), bIsResim);
+			}
+#endif
+
+			ImplementationInterface_Internal->ApplyAction(PredictedActionEntry.ActionInstance);
+
+			PredictedActionEntry.bApplied = true;
+			PredictedActionEntry.bActive = false;
+		}
+	}
+}
+
+uint16 FAsyncNetworkPhysicsComponent::GetUniqueActionId_Internal()
+{
+	// Lower 15 bits are the ID sequence, top bit identifies server/client origin.
+	constexpr uint16 ServerBit = static_cast<uint16>(1u) << 15; // 1000 0000 0000 0000
+	constexpr uint16 SequenceMask = static_cast<uint16>(~ServerBit); // 0111 1111 1111 1111
+
+	// Increment the ActionIdIterator and only keep the first 15 bits
+	ActionIdIterator = static_cast<uint16>(++ActionIdIterator & SequenceMask);
+	
+	// ActionId of 0 is invalid, when looping around increment one more time
+	if (ActionIdIterator == 0u)
+	{
+		ActionIdIterator++;
+	}
+
+	if (IsServer())
+	{
+		// Set the last 16th bit as 1 for ActionId created by the server
+		return uint16(ActionIdIterator | ServerBit);
+	}
+	else // Client
+	{
+		return ActionIdIterator;
+	}
+}
+
 FAsyncNetworkPhysicsComponentOutput& FAsyncNetworkPhysicsComponent::GetAsyncOutput_Internal()
 {
 	FAsyncNetworkPhysicsComponentOutput& AsyncOutput = GetProducerOutputData_Internal();
@@ -2509,7 +3064,7 @@ void FAsyncNetworkPhysicsComponent::OnPreProcessInputs_Internal(const int32 Phys
 				}
 				else
 				{
-					if (ensure(ImplementationInterface_Internal))
+					if (ImplementationInterface_Internal)
 					{
 						ImplementationInterface_Internal->ApplyState(*PhysicsData);
 					}
@@ -2557,7 +3112,7 @@ void FAsyncNetworkPhysicsComponent::OnPreProcessInputs_Internal(const int32 Phys
 /* // Importance is not implemented into the new network flow (yet?)
 					else
 					{
-						if (InputHelper && ensure(ImplementationInterface_Internal))
+						if (InputHelper && ImplementationInterface_Internal)
 						{
 							InputHelper->ApplyDataRange(InputHistory.Get(), NewImportantInputFrame, NextExpectedInputFrame - 1, *ImplementationInterface_Internal);
 						}
@@ -2660,7 +3215,7 @@ void FAsyncNetworkPhysicsComponent::OnPreProcessInputs_Internal(const int32 Phys
 								}
 								else
 								{
-									if (ensure(ImplementationInterface_Internal))
+									if (ImplementationInterface_Internal)
 									{
 										ImplementationInterface_Internal->ApplyInput(*PhysicsData);
 									}
@@ -2704,7 +3259,7 @@ void FAsyncNetworkPhysicsComponent::OnPreProcessInputs_Internal(const int32 Phys
 				}
 				else
 				{
-					if (ensure(ImplementationInterface_Internal))
+					if (ImplementationInterface_Internal)
 					{
 						ImplementationInterface_Internal->ApplyInput(*PhysicsData);
 					}
@@ -2732,6 +3287,10 @@ void FAsyncNetworkPhysicsComponent::OnPreProcessInputs_Internal(const int32 Phys
 #endif
 		}
 	}
+
+	 // Apply Actions
+	ApplyActionsForFrame(PhysicsStep);
+
 	NewImportantInputFrame = INT_MAX;
 }
 
@@ -2776,7 +3335,7 @@ void FAsyncNetworkPhysicsComponent::OnPostProcessInputs_Internal(const int32 Phy
 		}
 		else
 		{
-			if (ensure(ImplementationInterface_Internal))
+			if (ImplementationInterface_Internal)
 			{
 				// Gather input data from implementation
 				ImplementationInterface_Internal->BuildInput(*PhysicsData);
@@ -2817,7 +3376,7 @@ void FAsyncNetworkPhysicsComponent::OnPostProcessInputs_Internal(const int32 Phy
 			}
 			else
 			{
-				if (ensure(ImplementationInterface_Internal))
+				if (ImplementationInterface_Internal)
 				{
 					// Gather state data from implementation
 					ImplementationInterface_Internal->BuildState(*PhysicsData);
@@ -2842,6 +3401,7 @@ void FAsyncNetworkPhysicsComponent::OnPostProcessInputs_Internal(const int32 Phy
 		FAsyncNetworkPhysicsComponentOutput& AsyncOutput = GetAsyncOutput_Internal();
 		SendInputData_Internal(AsyncOutput, PhysicsStep);
 		SendStateData_Internal(AsyncOutput, PhysicsStep);
+		SendActionData_Internal(AsyncOutput, PhysicsStep);
 		FinalizeOutputData_Internal();
 	}
 }
@@ -2961,7 +3521,7 @@ void FAsyncNetworkPhysicsComponent::OnPostSolve_Internal()
 			}
 			else
 			{
-				if (ensure(ImplementationInterface_Internal))
+				if (ImplementationInterface_Internal)
 				{
 					// Gather state data from implementation
 					ImplementationInterface_Internal->BuildState(*PhysicsData);
@@ -2986,6 +3546,7 @@ void FAsyncNetworkPhysicsComponent::OnPostSolve_Internal()
 		FAsyncNetworkPhysicsComponentOutput& AsyncOutput = GetAsyncOutput_Internal();
 		SendInputData_Internal(AsyncOutput, PhysicsStep);
 		SendStateData_Internal(AsyncOutput, NextPhysicsStep); // We cache state at the end of the solve which is the state that next frame starts with
+		SendActionData_Internal(AsyncOutput, PhysicsStep);
 		FinalizeOutputData_Internal();
 	}
 }
@@ -3131,6 +3692,18 @@ void FAsyncNetworkPhysicsComponent::SendStateData_Internal(FAsyncNetworkPhysicsC
 	}
 }
 
+void FAsyncNetworkPhysicsComponent::SendActionData_Internal(FAsyncNetworkPhysicsComponentOutput& AsyncOutput, const int32 PhysicsStep)
+{
+	if (IsServer())
+	{
+		ConfirmedActions.GetActionsToReplicate(AsyncOutput.ActionData);
+	}
+	else if (IsLocallyControlled()) // Autonomous Proxy Client
+	{
+		PendingProposedActions.GetActionsToReplicate(AsyncOutput.ActionData);
+	}
+}
+
 Chaos::FPBDRigidsSolver* FAsyncNetworkPhysicsComponent::GetRigidSolver()
 {
 	return static_cast<Chaos::FPBDRigidsSolver*>(GetSolver());
@@ -3565,7 +4138,6 @@ int32 FAsyncNetworkPhysicsComponent::GetLatestReceivedStateFrame() const
 		const Chaos::FPhysicsSolverBase* BaseSolver = GetSolver();
 		if (const Chaos::FRewindData* RewindData = BaseSolver ? BaseSolver->GetRewindData() : nullptr)
 		{
-			// Todo: we should not have to get a read lock to get this information bu instead measure the latest lag upon reception
 			Chaos::FReadPhysicsObjectInterface_Internal Interface = Chaos::FPhysicsObjectInternalInterface::GetRead();
 			if (Chaos::FGeometryParticleHandle* Particle = Interface.GetParticle(PhysicsObject))
 			{
@@ -3576,6 +4148,51 @@ int32 FAsyncNetworkPhysicsComponent::GetLatestReceivedStateFrame() const
 	return INDEX_NONE;
 }
 
+void FAsyncNetworkPhysicsComponent::EnqueueImmediateActionInstance_Internal(TInstancedStruct<FNetworkPhysicsActionPayload>&& ActionInstance, const uint32 SourceId)
+{
+	EnqueueScheduledActionInstanceAtFrame_Internal(MoveTemp(ActionInstance), SourceId, -1);
+}
+
+void FAsyncNetworkPhysicsComponent::EnqueueScheduledActionInstance_Internal(TInstancedStruct<FNetworkPhysicsActionPayload>&& ActionInstance, const uint32 SourceId, const float DelaySeconds)
+{
+	int32 LocalFrame = INDEX_NONE;
+	if (Chaos::FPBDRigidsSolver* RigidSolver = GetRigidSolver())
+	{
+		if (RigidSolver->IsUsingFixedDt())
+		{
+			// Only apply a frame if we are actually delaying, else leave as INDEX_NONE which will queue up the action ASAP and handle edge-cases.
+			int32 DelayFrames = FMath::CeilToInt32(DelaySeconds / RigidSolver->GetAsyncDeltaTime());
+			if (DelayFrames > 0)
+			{
+				LocalFrame = RigidSolver->GetCurrentFrame() + DelayFrames;
+			}
+		}
+	}
+	EnqueueScheduledActionInstanceAtFrame_Internal(MoveTemp(ActionInstance), SourceId, LocalFrame);
+}
+
+void FAsyncNetworkPhysicsComponent::EnqueueScheduledActionInstanceAtFrame_Internal(TInstancedStruct<FNetworkPhysicsActionPayload>&& ActionInstance, const uint32 SourceId, const int32 LocalFrame)
+{
+	FNetworkPhysicsActionPayload& Action = ActionInstance.GetMutable<FNetworkPhysicsActionPayload>();
+	Action.SourceId = SourceId;
+
+	// If we have a frame passed in, convert that to ServerFrame, if INDEX_NONE set ServerFrame to that which indicates that the action should be applied ASAP
+	Action.ServerFrame = (LocalFrame == INDEX_NONE) ? INDEX_NONE : (LocalFrame + UE::NetworkPhysicsUtils::GetNetworkPhysicsTickOffset_Internal(GetSolver()));
+
+#if !(UE_BUILD_SHIPPING || UE_BUILD_TEST)
+	if (PhysicsReplicationCVars::ResimulationCVars::bActionsEnableDebugLogs)
+	{
+		const bool bIsResim = GetEvolution() ? GetEvolution()->IsResimming() : false;
+
+		UE_LOGF(LogChaos, Log, "PT %ls EnqueueActionAtFrame - Type: %ls - SourceId: %ld - LocalFrame: %d + Offset: %d = ServerFrame: %d - bIsResim: %d"
+			, (IsServer() ? TEXT(" -- SERVER -- ") : TEXT(" ++ CLIENT ++ ")), (ActionInstance.GetScriptStruct() ? *ActionInstance.GetScriptStruct()->GetName() : TEXT("UNKNOWN ACTION"))
+			, Action.SourceId, LocalFrame, UE::NetworkPhysicsUtils::GetNetworkPhysicsTickOffset_Internal(GetSolver()), Action.ServerFrame, bIsResim);
+	}
+#endif
+
+	RegisterAction_Internal(ActionInstance, /*bWasReceivedFromReplication*/ false);
+}
+
 namespace UE::NetworkPhysicsUtils
 {
 	int32 GetUpcomingServerFrame_External(UWorld* World)
@@ -3630,4 +4247,71 @@ namespace UE::NetworkPhysicsUtils
 
 		return 0;
 	}
-}
+
+	const UObject* GetOwnerFromUObject_External(const UObject* SourceObject)
+	{
+		if (!IsValid(SourceObject))
+		{
+			return nullptr;
+		}
+
+		if (const AActor* Actor = Cast<AActor>(SourceObject))
+		{
+			return Actor;
+		}
+
+		if (const UActorComponent* Component = Cast<UActorComponent>(SourceObject))
+		{
+			if (const AActor* Owner = Component->GetOwner())
+			{
+				return Owner;
+			}
+		}
+
+		if (const AActor* OuterActor = SourceObject->GetTypedOuter<AActor>())
+		{
+			return OuterActor;
+		}
+
+		return SourceObject;
+	}
+
+	const uint32 GetNetworkStableHash_External(const UObject* InSourceObject)
+	{
+		const UObject* OwnerObject = GetOwnerFromUObject_External(InSourceObject);
+		if (!IsValid(OwnerObject))
+		{
+			return 0u;
+		}
+
+		// If object is replicated, use Iris NetRepHandle
+		if (const UWorld* World = OwnerObject->GetWorld())
+		{
+			if (UEngineReplicationBridge* ReplicationBridge = UE::Net::FReplicationSystemUtil::GetEngineReplicationBridge(World))
+			{
+				// If your branch scopes this enum differently, adjust this token only.
+				const UE::Net::FNetRefHandle NetRefHandle = ReplicationBridge->GetReplicatedRefHandle(OwnerObject, UE::Net::EGetRefHandleFlags::None);
+
+				if (NetRefHandle.IsValid())
+				{
+					return UE::Net::GetTypeHash(NetRefHandle);
+				}
+			}
+		}
+
+		// If object has a network stable name, that that
+		if (OwnerObject->IsNameStableForNetworking())
+		{
+			const UObject* PackageOuter = OwnerObject->GetOutermost();
+			const FString StablePath = OwnerObject->GetPathName(PackageOuter);
+
+			if (!StablePath.IsEmpty())
+			{
+				return FCrc::StrCrc32(*StablePath);
+			}
+		}
+
+		// No network stable option to use as a SourceId
+		return 0u;
+	}
+}
\ No newline at end of file
diff --git a/Engine/Source/Runtime/Engine/Public/Physics/NetworkPhysicsComponent.h b/Engine/Source/Runtime/Engine/Public/Physics/NetworkPhysicsComponent.h
--- a/Engine/Source/Runtime/Engine/Public/Physics/NetworkPhysicsComponent.h
+++ b/Engine/Source/Runtime/Engine/Public/Physics/NetworkPhysicsComponent.h
@@ -1,12 +1,13 @@
 // Copyright Epic Games, Inc. All Rights Reserved.
 
 #pragma once
+#include "Engine/EngineTypes.h"
+#include "Engine/NetConnection.h"
+#include "Engine/World.h"
 #include "RewindData.h"
 #include "Components/ActorComponent.h"
 #include "GameFramework/PawnMovementComponent.h"
 #include "GameFramework/PlayerController.h"
-#include "Engine/NetConnection.h"
-#include "Engine/World.h"
 #include "Physics/Experimental/PhysScene_Chaos.h"
 #include "Subsystems/WorldSubsystem.h"
 #include "PhysicsEngine/PhysicsSettings.h"
@@ -83,6 +84,12 @@ namespace UE::NetworkPhysicsUtils
 
 	// Returns the network physics frame offset on the physics thread
 	ENGINE_API const int32 GetNetworkPhysicsTickOffset_Internal(const Chaos::FPhysicsSolverBase* Solver);
+
+	// Returns the owning UObject of a UObject
+	ENGINE_API const UObject* GetOwnerFromUObject_External(const UObject* SourceObject);
+
+	// Returns a network stable hash for a UObject, requires it to be either replicated or baked in the level, else returns 0
+	ENGINE_API const uint32 GetNetworkStableHash_External(const UObject* InSourceObject);
 }
 
 /** Templated data history, holding a data buffer */
@@ -625,6 +632,715 @@ private :
 	}
 };
 
+
+/**
+*	Start Implementing Action Structs and helpers
+* 
+*	Action "history" is setup like this:
+*	FActionHistory - Circular frame buffer
+*		Frame 1 - FActionHistoryFrameData
+*			Action 1 - FActionHistoryEntry
+*				FActionInstance ActionInstance
+*			Action 2 - FActionHistoryEntry
+*				FActionInstance ActionInstance
+*		Frame 2 - FActionHistoryFrameData
+*			Action 1 - FActionHistoryEntry
+*				FActionInstance ActionInstance
+*/
+
+USTRUCT()
+struct FNetworkPhysicsActionPayload
+{
+	GENERATED_BODY()
+
+	virtual ~FNetworkPhysicsActionPayload() = default;
+
+	enum EActionAuthorStyle : uint8
+	{
+		Authority = 0, // Performed server-side, replicated and applied to clients
+		Predicted = 1, // Predicted and applied client-side, performed server-side, replicated to clients where the predicted action gets overwritten
+		PredictedAutonomousOnly = 2, // Predicted and applied on the autonomous client, performed server-side, replicated to clients where the predicted action gets overwritten
+		Proposed = 3, // Predicted and applied client-side, replicated to server, performed server-side and compared with the proposed action from the autonomous proxy and if close enough the server uses the clients proposed action, then replicated to where the proposed action gets overwritten
+		ProposedAutonomousOnly = 4 // Predicted and applied on the autonomous client, replicated to server, performed server-side and compared with the proposed action from the autonomous proxy and if close enough the server uses the clients proposed action, then replicated to where the proposed action gets overwritten
+	};
+	virtual const EActionAuthorStyle GetAuthorStyle() const { return EActionAuthorStyle::Authority; };
+
+	UPROPERTY()
+	int32 ServerFrame = INDEX_NONE;
+
+	UPROPERTY()
+	uint32 SourceId = 0u;
+
+	UPROPERTY()
+	uint16 ActionId = 0u;
+
+	virtual bool IsNearlyEqual(const FNetworkPhysicsActionPayload& Other) const
+	{
+		ensureMsgf(false, TEXT("FNetworkPhysicsActionPayload::IsNearlyEqual is not overridden for an action."));
+		return false;
+	}
+
+	virtual void DebugAction() const { UE_LOGF(LogChaos, Log, "		DEBUG || ActionPayload || DebugAction() Not Implemented!"); }
+};
+
+// Struct holding the Action Instance along with helper data
+struct FActionHistoryEntry
+{
+	using FActionInstance = TInstancedStruct<FNetworkPhysicsActionPayload>;
+
+	FActionInstance ActionInstance;
+	bool bActive = false;
+	bool bApplied = false;
+	bool bHasBeenQueuedForReplication = false;
+
+	void Reset()
+	{
+		ActionInstance.Reset();
+		bActive = false;
+		bApplied = false;
+		bHasBeenQueuedForReplication = false;
+	}
+};
+
+// Struct holding actions per-frame in FActionHistory
+struct FActionHistoryFrameData
+{
+	int32 LocalFrame = INDEX_NONE;
+	TArray<FActionHistoryEntry> Entries;
+
+	void ResetForFrame(const int32 InLocalFrame)
+	{
+		LocalFrame = InLocalFrame;
+		Entries.Reset();
+	}
+
+	bool IsForFrame(const int32 InLocalFrame) const
+	{
+		return LocalFrame == InLocalFrame;
+	}
+};
+
+// Struct that points to a specific action in FActionHistory
+struct FActionHistoryLocation
+{
+	int32 FrameIndex = INDEX_NONE;
+	int32 EntryIndex = INDEX_NONE;
+
+	bool IsValid() const
+	{
+		return FrameIndex != INDEX_NONE && EntryIndex != INDEX_NONE;
+	}
+};
+
+struct FActionHistory
+{
+	using FActionInstance = TInstancedStruct<FNetworkPhysicsActionPayload>;
+
+	int32 HistorySize = 0;
+	TArray<FActionHistoryFrameData> FrameHistory;
+
+	void Initialize(const int32 InHistorySize)
+	{
+		HistorySize = FMath::Max(1, InHistorySize);
+		FrameHistory.SetNum(HistorySize);
+
+		for (FActionHistoryFrameData& FrameData : FrameHistory)
+		{
+			FrameData.ResetForFrame(INDEX_NONE);
+		}
+	}
+
+	void Reset()
+	{
+		for (FActionHistoryFrameData& FrameData : FrameHistory)
+		{
+			FrameData.ResetForFrame(INDEX_NONE);
+		}
+	}
+
+	bool IsInitialized() const
+	{
+		return HistorySize > 0 && FrameHistory.Num() == HistorySize;
+	}
+
+	int32 GetFrameIndex(const int32 LocalFrame) const
+	{
+		check(IsInitialized());
+
+		int32 FrameIndex = LocalFrame % HistorySize;
+		if (FrameIndex < 0)
+		{
+			FrameIndex += HistorySize;
+		}
+
+		return FrameIndex;
+	}
+
+	FActionHistoryFrameData* FindFrameData(const int32 LocalFrame)
+	{
+		if (IsInitialized() == false)
+		{
+			return nullptr;
+		}
+
+		FActionHistoryFrameData& FrameData = FrameHistory[GetFrameIndex(LocalFrame)];
+		if (FrameData.IsForFrame(LocalFrame))
+		{
+			return &FrameData;
+		}
+
+		return nullptr;
+	}
+
+	const FActionHistoryFrameData* FindFrameData(const int32 LocalFrame) const
+	{
+		if (IsInitialized() == false)
+		{
+			return nullptr;
+		}
+
+		const FActionHistoryFrameData& FrameData = FrameHistory[GetFrameIndex(LocalFrame)];
+		if (FrameData.IsForFrame(LocalFrame))
+		{
+			return &FrameData;
+		}
+
+		return nullptr;
+	}
+
+	FActionHistoryFrameData& FindOrAddFrameData(const int32 LocalFrame)
+	{
+		check(IsInitialized());
+
+		FActionHistoryFrameData& FrameData = FrameHistory[GetFrameIndex(LocalFrame)];
+		if (FrameData.IsForFrame(LocalFrame) == false)
+		{
+			// Override old entry in circle buffer with new frame
+			FrameData.ResetForFrame(LocalFrame);
+		}
+
+		return FrameData;
+	}
+
+	/**
+	*	Start static helper functions
+	*/
+	static const FNetworkPhysicsActionPayload& GetRef(const FActionInstance& ActionInstance)
+	{
+		return ActionInstance.Get<FNetworkPhysicsActionPayload>();
+	}
+
+	static FNetworkPhysicsActionPayload& GetRefMutable(FActionInstance& ActionInstance)
+	{
+		return ActionInstance.GetMutable<FNetworkPhysicsActionPayload>();
+	}
+
+	/* Check if active, not applied and valid.
+	* @param bAllowApplied - If false, return false if action has been applied. */
+	static bool IsUsableEntry(const FActionHistoryEntry& Entry, const bool bAllowApplied)
+	{
+		if (Entry.bActive == false)
+		{
+			return false;
+		}
+
+		// Check if we care about the entry already being applied or not, for example during resimulation we allow re-application of actions
+		if (bAllowApplied == false && Entry.bApplied == true)
+		{
+			return false;
+		}
+
+		return Entry.ActionInstance.IsValid();
+	}
+
+	static bool SameConcreteType(const FActionInstance& AActionInstance, const FActionInstance& BActionInstance)
+	{
+		if (AActionInstance.IsValid() == false || BActionInstance.IsValid() == false)
+		{
+			return false;
+		}
+
+		return AActionInstance.GetScriptStruct() == BActionInstance.GetScriptStruct();
+	}
+
+	static bool HasSameValidActionId(const FActionInstance& AActionInstance, const FActionInstance& BActionInstance)
+	{
+		const FNetworkPhysicsActionPayload& AAction = GetRef(AActionInstance);
+		const FNetworkPhysicsActionPayload& BAction = GetRef(BActionInstance);
+
+		if (AAction.ActionId == 0u)
+		{
+			return false;
+		}
+
+		return AAction.ActionId == BAction.ActionId;
+	}
+
+	// Check if two actions are of the same type and produced by the same source
+	static bool AreSourceEquivalent(const FActionInstance& ExistingActionInstance, const FActionInstance& CandidateActionInstance)
+	{
+		if (SameConcreteType(ExistingActionInstance, CandidateActionInstance) == false)
+		{
+			return false;
+		}
+
+		const FNetworkPhysicsActionPayload& ExistingAction = GetRef(ExistingActionInstance);
+		const FNetworkPhysicsActionPayload& CandidateAction = GetRef(CandidateActionInstance);
+
+		return ExistingAction.SourceId == CandidateAction.SourceId;
+	}
+
+	// Check if two actions are equal enough to be considered the same action
+	static bool AreEquivalent(const FActionInstance& ExistingActionInstance, const FActionInstance& CandidateActionInstance)
+	{
+		if (SameConcreteType(ExistingActionInstance, CandidateActionInstance) == false)
+		{
+			return false;
+		}
+
+		const FNetworkPhysicsActionPayload& ExistingAction = GetRef(ExistingActionInstance);
+		const FNetworkPhysicsActionPayload& CandidateAction = GetRef(CandidateActionInstance);
+
+		if (ExistingAction.SourceId != CandidateAction.SourceId)
+		{
+			return false;
+		}
+
+		return ExistingAction.IsNearlyEqual(CandidateAction);
+	}
+
+	// Check if two actions are duplications, not just equivalent enough
+	static bool IsStrongDuplicate(const FActionInstance& ExistingActionInstance, const int32 ExistingLocalFrame, const FActionInstance& CandidateActionInstance, const int32 CandidateLocalFrame)
+	{
+		if (HasSameValidActionId(ExistingActionInstance, CandidateActionInstance))
+		{
+			return true;
+		}
+
+		const FNetworkPhysicsActionPayload& ExistingAction = GetRef(ExistingActionInstance);
+		const FNetworkPhysicsActionPayload& CandidateAction = GetRef(CandidateActionInstance);
+
+		// If either side already has a valid but not equal ActionId, they are not duplicates
+		if (ExistingAction.ActionId != 0u || CandidateAction.ActionId != 0u)
+		{
+			return false;
+		}
+		if (SameConcreteType(ExistingActionInstance, CandidateActionInstance) == false)
+		{
+			return false;
+		}
+		if (ExistingAction.SourceId != CandidateAction.SourceId)
+		{
+			return false;
+		}
+		if (ExistingLocalFrame != CandidateLocalFrame)
+		{
+			return false;
+		}
+
+		return ExistingAction.IsNearlyEqual(CandidateAction);
+	}
+
+	/**
+	*	End static helper functions
+	*/
+
+
+	FActionInstance MakeCopy(const FActionInstance& ActionInstance, const int32 LocalFrame, const int32 TickOffset, const uint16 OverrideActionId = 0u, const bool bClearActionId = false) const
+	{
+		FActionInstance ActionInstanceCopy = ActionInstance;
+		FNetworkPhysicsActionPayload& ActionCopy = GetRefMutable(ActionInstanceCopy);
+
+		// If no frame number is set, apply LocalFrame converted into ServerFrame (i.e. the Action will be applied ASAP)
+		if (ActionCopy.ServerFrame == INDEX_NONE)
+		{
+			ActionCopy.ServerFrame = LocalFrame + TickOffset;
+		}
+
+		if (OverrideActionId != 0u)
+		{
+			ActionCopy.ActionId = OverrideActionId;
+		}
+		else if (bClearActionId)
+		{
+			ActionCopy.ActionId = 0u;
+		}
+
+		return ActionInstanceCopy;
+	}
+
+	FActionHistoryLocation FindActiveActionById(const uint16 ActionId, const bool bIncludeApplied = false) const
+	{
+		FActionHistoryLocation FoundLocation;
+
+		if (IsInitialized() == false || ActionId == 0u)
+		{
+			return FoundLocation;
+		}
+
+		for (int32 FrameIndex = 0; FrameIndex < FrameHistory.Num(); ++FrameIndex)
+		{
+			const FActionHistoryFrameData& FrameData = FrameHistory[FrameIndex];
+			if (FrameData.LocalFrame == INDEX_NONE)
+			{
+				continue;
+			}
+
+			for (int32 EntryIndex = 0; EntryIndex < FrameData.Entries.Num(); ++EntryIndex)
+			{
+				const FActionHistoryEntry& Entry = FrameData.Entries[EntryIndex];
+				if (IsUsableEntry(Entry, bIncludeApplied) == false)
+				{
+					continue;
+				}
+
+				if (GetRef(Entry.ActionInstance).ActionId == ActionId)
+				{
+					FoundLocation.FrameIndex = FrameIndex;
+					FoundLocation.EntryIndex = EntryIndex;
+					return FoundLocation;
+				}
+			}
+		}
+
+		return FoundLocation;
+	}
+
+	FActionHistoryLocation FindEquivalentActionInRange(const FActionInstance& CandidateActionInstance, const int32 MinLocalFrame, const int32 MaxLocalFrame, const bool bIncludeApplied = false, const bool bCheckSourceEquivalent = false) const
+	{
+		FActionHistoryLocation FoundLocation;
+
+		if (IsInitialized() == false || CandidateActionInstance.IsValid() == false)
+		{
+			return FoundLocation;
+		}
+
+		const int32 StartLocalFrame = FMath::Min(MinLocalFrame, MaxLocalFrame);
+		const int32 EndLocalFrame = FMath::Max(MinLocalFrame, MaxLocalFrame);
+
+		for (int32 QueryLocalFrame = StartLocalFrame; QueryLocalFrame <= EndLocalFrame; ++QueryLocalFrame)
+		{
+			const FActionHistoryFrameData* FrameData = FindFrameData(QueryLocalFrame);
+			if (FrameData == nullptr)
+			{
+				continue;
+			}
+
+			for (int32 EntryIndex = 0; EntryIndex < FrameData->Entries.Num(); ++EntryIndex)
+			{
+				const FActionHistoryEntry& Entry = FrameData->Entries[EntryIndex];
+				if (IsUsableEntry(Entry, bIncludeApplied) == false)
+				{
+					continue;
+				}
+
+				if (bCheckSourceEquivalent)
+				{
+					if (AreSourceEquivalent(Entry.ActionInstance, CandidateActionInstance))
+					{
+						FoundLocation.FrameIndex = GetFrameIndex(QueryLocalFrame);
+						FoundLocation.EntryIndex = EntryIndex;
+						return FoundLocation;
+					}
+				}
+				else
+				{
+					if (AreEquivalent(Entry.ActionInstance, CandidateActionInstance))
+					{
+						FoundLocation.FrameIndex = GetFrameIndex(QueryLocalFrame);
+						FoundLocation.EntryIndex = EntryIndex;
+						return FoundLocation;
+					}
+				}
+			}
+		}
+
+		return FoundLocation;
+	}
+
+	bool HasEquivalentActionInRange(const FActionInstance& CandidateActionInstance, const int32 MinLocalFrame, const int32 MaxLocalFrame, const bool bIncludeApplied = false, const bool bCheckSourceEquivalent = false) const
+	{
+		return FindEquivalentActionInRange(CandidateActionInstance, MinLocalFrame, MaxLocalFrame, bIncludeApplied).IsValid();
+	}
+
+	bool AddAction(const FActionInstance& ActionInstance, const int32 LocalFrame, const int32 TickOffset
+		, const uint16 OverrideActionId = 0u, const bool bClearActionId = false, const bool bCheckDuplicateInSameFrame = false)
+	{
+		if (IsInitialized() == false || ActionInstance.IsValid() == false)
+		{
+			return false;
+		}
+
+		FActionInstance ActionInstanceCopy = MakeCopy(ActionInstance, LocalFrame, TickOffset, OverrideActionId, bClearActionId);
+		const FNetworkPhysicsActionPayload& ActionCopy = GetRef(ActionInstanceCopy);
+
+		// If we already have registered an action with this ActionId, early out. We might have received the same replicated action twice in a row.
+		if (ActionCopy.ActionId != 0u && FindActiveActionById(ActionCopy.ActionId, true).IsValid())
+		{
+			return false;
+		}
+
+		FActionHistoryFrameData& FrameData = FindOrAddFrameData(LocalFrame);
+
+		// In some cases we want to stop duplicates from being registered, a duplicate action is a stronger check than equivalent action.
+		if (bCheckDuplicateInSameFrame)
+		{
+			for (const FActionHistoryEntry& ExistingEntry : FrameData.Entries)
+			{
+				if (IsUsableEntry(ExistingEntry, /*bAllowApplied*/true) == false)
+				{
+					continue;
+				}
+
+				if (IsStrongDuplicate(ExistingEntry.ActionInstance, FrameData.LocalFrame, ActionInstanceCopy, LocalFrame))
+				{
+					return false;
+				}
+			}
+		}
+
+		FActionHistoryEntry& NewEntry = FrameData.Entries.Emplace_GetRef();
+		NewEntry.ActionInstance = MoveTemp(ActionInstanceCopy);
+		NewEntry.bActive = true;
+		NewEntry.bApplied = false;
+		NewEntry.bHasBeenQueuedForReplication = false;
+		return true;
+	}
+
+	const FActionInstance* GetAction(const FActionHistoryLocation& Location) const
+	{
+		if (Location.IsValid() == false)
+		{
+			return nullptr;
+		}
+
+		if (FrameHistory.IsValidIndex(Location.FrameIndex) == false)
+		{
+			return nullptr;
+		}
+
+		const FActionHistoryFrameData& FrameData = FrameHistory[Location.FrameIndex];
+
+		if (FrameData.Entries.IsValidIndex(Location.EntryIndex) == false)
+		{
+			return nullptr;
+		}
+
+		const FActionHistoryEntry& Entry = FrameData.Entries[Location.EntryIndex];
+
+		if (Entry.ActionInstance.IsValid() == false)
+		{
+			return nullptr;
+		}
+
+		return &Entry.ActionInstance;
+	}
+
+	int32 GetLocalFrame(const FActionHistoryLocation& Location) const
+	{
+		if (Location.IsValid() == false)
+		{
+			return INDEX_NONE;
+		}
+
+		if (FrameHistory.IsValidIndex(Location.FrameIndex) == false)
+		{
+			return INDEX_NONE;
+		}
+
+		return FrameHistory[Location.FrameIndex].LocalFrame;
+	}
+
+	bool Deactivate(const FActionHistoryLocation& Location)
+	{
+		if (Location.IsValid() == false)
+		{
+			return false;
+		}
+
+		if (FrameHistory.IsValidIndex(Location.FrameIndex) == false)
+		{
+			return false;
+		}
+
+		FActionHistoryFrameData& FrameData = FrameHistory[Location.FrameIndex];
+
+		if (FrameData.Entries.IsValidIndex(Location.EntryIndex) == false)
+		{
+			return false;
+		}
+
+		FrameData.Entries[Location.EntryIndex].bActive = false;
+		return true;
+	}
+
+	bool DeactivateAllByActionId(const uint16 ActionId)
+	{
+		bool bDeactivatedAny = false;
+
+		if (IsInitialized() == false || ActionId == 0u)
+		{
+			return false;
+		}
+
+		for (FActionHistoryFrameData& FrameData : FrameHistory)
+		{
+			if (FrameData.LocalFrame == INDEX_NONE)
+			{
+				continue;
+			}
+
+			for (FActionHistoryEntry& Entry : FrameData.Entries)
+			{
+				if (Entry.ActionInstance.IsValid() == false)
+				{
+					continue;
+				}
+
+				if (GetRef(Entry.ActionInstance).ActionId == ActionId)
+				{
+					Entry.bActive = false;
+					bDeactivatedAny = true;
+				}
+			}
+		}
+
+		return bDeactivatedAny;
+	}
+
+	bool DeactivateEquivalentActionsInRange(const FActionInstance& CandidateActionInstance, const int32 MinLocalFrame, const int32 MaxLocalFrame)
+	{
+		bool bDeactivatedAny = false;
+
+		if (IsInitialized() == false || CandidateActionInstance.IsValid() == false)
+		{
+			return false;
+		}
+
+		const int32 StartLocalFrame = FMath::Min(MinLocalFrame, MaxLocalFrame);
+		const int32 EndLocalFrame = FMath::Max(MinLocalFrame, MaxLocalFrame);
+		const uint16 CandidateActionId = GetRef(CandidateActionInstance).ActionId;
+
+		for (int32 QueryLocalFrame = StartLocalFrame; QueryLocalFrame <= EndLocalFrame; ++QueryLocalFrame)
+		{
+			FActionHistoryFrameData* FrameData = FindFrameData(QueryLocalFrame);
+			if (FrameData == nullptr)
+			{
+				continue;
+			}
+
+			for (FActionHistoryEntry& Entry : FrameData->Entries)
+			{
+				if (IsUsableEntry(Entry, /*bAllowApplied*/false) == false)
+				{
+					continue;
+				}
+
+				const bool bSameActionId = CandidateActionId != 0u && GetRef(Entry.ActionInstance).ActionId == CandidateActionId;
+
+				if (bSameActionId || AreEquivalent(Entry.ActionInstance, CandidateActionInstance))
+				{
+					Entry.bActive = false;
+					bDeactivatedAny = true;
+				}
+			}
+		}
+
+		return bDeactivatedAny;
+	}
+
+	void GetActionsToReplicate(TArray<FActionInstance>& OutActions, const bool bResetOutArray = true, const bool bMarkAsQueuedForReplication = true, const bool bOnlyActiveActions = true)
+	{
+		if (IsInitialized() == false)
+		{
+			return;
+		}
+
+		if (bResetOutArray == true)
+		{
+			OutActions.Reset();
+		}
+
+		struct FPendingReplicationLocation
+		{
+			int32 FrameIndex = INDEX_NONE;
+			int32 EntryIndex = INDEX_NONE;
+			int32 LocalFrame = INDEX_NONE;
+		};
+
+		TArray<FPendingReplicationLocation, TInlineAllocator<16>> PendingReplicationLocations;
+
+		for (int32 FrameIndex = 0; FrameIndex < FrameHistory.Num(); ++FrameIndex)
+		{
+			const FActionHistoryFrameData& FrameData = FrameHistory[FrameIndex];
+
+			if (FrameData.LocalFrame == INDEX_NONE)
+			{
+				continue;
+			}
+
+			for (int32 EntryIndex = 0; EntryIndex < FrameData.Entries.Num(); ++EntryIndex)
+			{
+				const FActionHistoryEntry& Entry = FrameData.Entries[EntryIndex];
+
+				if (Entry.bHasBeenQueuedForReplication == true)
+				{
+					continue;
+				}
+				if (bOnlyActiveActions == true && Entry.bActive == false)
+				{
+					continue;
+				}
+				if (Entry.ActionInstance.IsValid() == false)
+				{
+					continue;
+				}
+
+				FPendingReplicationLocation& Location = PendingReplicationLocations.Emplace_GetRef();
+
+				Location.FrameIndex = FrameIndex;
+				Location.EntryIndex = EntryIndex;
+				Location.LocalFrame = FrameData.LocalFrame;
+			}
+		}
+
+		// The circular buffer storage order is not chronological, sort by local frame, this makes delta serialization work better when multiple actions are being replicated
+		PendingReplicationLocations.Sort([](const FPendingReplicationLocation& A, const FPendingReplicationLocation& B)
+			{
+				if (A.LocalFrame != B.LocalFrame)
+				{
+					return A.LocalFrame < B.LocalFrame;
+				}
+
+				if (A.FrameIndex != B.FrameIndex)
+				{
+					return A.FrameIndex < B.FrameIndex;
+				}
+
+				return A.EntryIndex < B.EntryIndex;
+			});
+
+		for (const FPendingReplicationLocation& Location : PendingReplicationLocations)
+		{
+			FActionHistoryFrameData& FrameData = FrameHistory[Location.FrameIndex];
+			FActionHistoryEntry& Entry = FrameData.Entries[Location.EntryIndex];
+
+			// Populate OUT array with the sorted actions to be replicated
+			OutActions.Add(Entry.ActionInstance);
+
+			if (bMarkAsQueuedForReplication == true)
+			{
+				Entry.bHasBeenQueuedForReplication = true;
+			}
+		}
+	}
+};
+
+/**
+* End Implementing Action Structs and helpers
+*/
+
+
 /**
  * Base struct for replicated rewind history properties
  */
@@ -1000,6 +1716,15 @@ struct FNetworkPhysicsPayload
 	}
 };
 
+USTRUCT()
+struct FNetworkPhysicsActionCollection
+{
+	GENERATED_BODY()
+
+	UPROPERTY()
+	TArray<TInstancedStruct<FNetworkPhysicsActionPayload>> DataArray;
+};
+
 USTRUCT()
 struct FNetworkPhysicsDataCollection
 {
@@ -1013,8 +1738,10 @@ struct FNetworkPhysicsDataCollection
 	{
 		for (const TInstancedStruct<FNetworkPhysicsPayload>& DataInstance : DataArray)
 		{
-			const FNetworkPhysicsPayload& Data = DataInstance.Get();
-			Data.LocalFrame = Data.ServerFrame - ClientFrameOffset;
+			if (const FNetworkPhysicsPayload* Data = DataInstance.GetPtr<FNetworkPhysicsPayload>())
+			{
+				Data->LocalFrame = Data->ServerFrame - ClientFrameOffset;
+			}
 		}
 	}
 
@@ -1091,6 +1818,8 @@ public:
 
 	virtual void BuildState(FNetworkPhysicsPayload& State) const = 0;
 	virtual void ApplyState(const FNetworkPhysicsPayload& State) = 0;
+
+	virtual void ApplyAction(const TInstancedStruct<FNetworkPhysicsActionPayload>& ActionInstance) = 0;
 };
 
 template<typename InputType, typename StateType>
@@ -1107,6 +1836,7 @@ private:
 
 	void BuildState(FNetworkPhysicsPayload& State) const override { BuildState_Internal(static_cast<StateType&>(State)); }
 	void ApplyState(const FNetworkPhysicsPayload& State) override { ApplyState_Internal(static_cast<const StateType&>(State)); }
+	void ApplyAction(const TInstancedStruct<FNetworkPhysicsActionPayload>& ActionInstance) override { ApplyAction_Internal(ActionInstance); }
 
 public:
 	/** Populate the input struct with current input data */
@@ -1125,6 +1855,9 @@ public:
 
 	/** Apply state struct to implementation */
 	virtual void ApplyState_Internal(const StateType& State) = 0;
+
+	/** Apply action struct to implementation */
+	virtual void ApplyAction_Internal(const TInstancedStruct<FNetworkPhysicsActionPayload>& ActionInstance) { }
 };
 
 /** END Iris Compliant Data */
@@ -1585,6 +2318,90 @@ public:
 	// Check if networked physics is setup with a synchronized physics tick offset
 	ENGINE_API bool IsNetworkPhysicsTickOffsetAssigned() const;
 
+	/** Enqueue a TInstancedStruct with an action based on FNetworkPhysicsActionPayload, to be executed as soon as possible.
+	* @param SourceObject is the object that is the source of this action, not necessarily the object that the action should run on. */
+	ENGINE_API void EnqueueImmediateActionInstance_External(TInstancedStruct<FNetworkPhysicsActionPayload>&& ActionInstance, const UObject* SourceObject);
+	
+	/** Enqueue a TInstancedStruct with an action based on FNetworkPhysicsActionPayload, to be executed as soon as possible.
+	* @param SourceId can be left at 0 for actions triggered by the owner of this implementation, else if triggered by an external actor pass in a network stable hash.
+	* Note: On the game thread a network stable hash can be retrieved from replicated UObjects or level baked UObjects (for example actors) via UE::NetworkPhysicsUtils::GetNetworkStableHash_External(UObject); */
+	ENGINE_API void EnqueueImmediateActionInstance_External(TInstancedStruct<FNetworkPhysicsActionPayload>&& ActionInstance, const uint32 SourceId);
+
+	/** Enqueue an action based on FNetworkPhysicsActionPayload, to be executed as soon as possible.
+	* @param SourceObject is the object that is the source of this action, not necessarily the object that the action should run on. */
+	template<typename TActionType>
+	void EnqueueImmediateAction_External(TActionType& Action, const UObject* SourceObject)
+	{
+		static_assert(TIsDerivedFrom<TActionType, FNetworkPhysicsActionPayload>::Value, "EnqueueImmediateAction_External - Action needs to derive from FNetworkPhysicsActionPayload");
+		EnqueueImmediateActionInstance_External(TInstancedStruct<FNetworkPhysicsActionPayload>::Make<TActionType>(Action), SourceObject);
+	}
+
+	/** Enqueue an action based on FNetworkPhysicsActionPayload, to be executed as soon as possible.
+	* @param SourceId can be left at 0 for actions triggered by the owner of this implementation, else if triggered by an external actor pass in a network stable hash.
+	* Note: On the game thread a network stable hash can be retrieved from replicated UObjects or level baked UObjects (for example actors) via UE::NetworkPhysicsUtils::GetNetworkStableHash_External(UObject); */
+	template<typename TActionType>
+	void EnqueueImmediateAction_External(TActionType& Action, const uint32 SourceId)
+	{
+		static_assert(TIsDerivedFrom<TActionType, FNetworkPhysicsActionPayload>::Value, "EnqueueImmediateAction_External - Action needs to derive from FNetworkPhysicsActionPayload");
+		EnqueueImmediateActionInstance_External(TInstancedStruct<FNetworkPhysicsActionPayload>::Make<TActionType>(Action), SourceId);
+	}
+
+	/** Enqueue a TInstancedStruct with an action based on FNetworkPhysicsActionPayload, to be executed after a delay of @param DelaySeconds.
+	* @param SourceObject is the object that is the source of this action, not necessarily the object that the action should run on. */
+	ENGINE_API void EnqueueScheduledActionInstance_External(TInstancedStruct<FNetworkPhysicsActionPayload>&& ActionInstance, const UObject* SourceObject, const float DelaySeconds);
+	
+	/** Enqueue a TInstancedStruct with an action based on FNetworkPhysicsActionPayload, to be executed after a delay of @param DelaySeconds.
+	* @param SourceId can be left at 0 for actions triggered by the owner of this implementation, else if triggered by an external actor pass in a network stable hash.
+	* Note: On the game thread a network stable hash can be retrieved from replicated UObjects or level baked UObjects (for example actors) via UE::NetworkPhysicsUtils::GetNetworkStableHash_External(UObject); */
+	ENGINE_API void EnqueueScheduledActionInstance_External(TInstancedStruct<FNetworkPhysicsActionPayload>&& ActionInstance, const uint32 SourceId, const float DelaySeconds);
+
+	/** Enqueue an action based on FNetworkPhysicsActionPayload, to be executed after a delay of @param DelaySeconds.
+	* @param SourceObject is the object that is the source of this action, not necessarily the object that the action should run on. */
+	template<typename TActionType>
+	void EnqueueScheduledAction_External(TActionType& Action, const UObject* SourceObject, const float DelaySeconds)
+	{
+		static_assert(TIsDerivedFrom<TActionType, FNetworkPhysicsActionPayload>::Value, "EnqueueScheduledAction_External - Action needs to derive from FNetworkPhysicsActionPayload");
+		EnqueueScheduledActionInstance_External(TInstancedStruct<FNetworkPhysicsActionPayload>::Make<TActionType>(Action), SourceObject, DelaySeconds);
+	}
+
+	/** Enqueue an action based on FNetworkPhysicsActionPayload, to be executed after a delay of @param DelaySeconds.
+	* @param SourceId can be left at 0 for actions triggered by the owner of this implementation, else if triggered by an external actor pass in a network stable hash.
+	* Note: On the game thread a network stable hash can be retrieved from replicated UObjects or level baked UObjects (for example actors) via UE::NetworkPhysicsUtils::GetNetworkStableHash_External(UObject); */
+	template<typename TActionType>
+	void EnqueueScheduledAction_External(TActionType& Action, const uint32 SourceId, const float DelaySeconds)
+	{
+		static_assert(TIsDerivedFrom<TActionType, FNetworkPhysicsActionPayload>::Value, "EnqueueScheduledAction_External - Action needs to derive from FNetworkPhysicsActionPayload");
+		EnqueueScheduledActionInstance_External(TInstancedStruct<FNetworkPhysicsActionPayload>::Make<TActionType>(Action), SourceId, DelaySeconds);
+	}
+
+	/** Enqueue a TInstancedStruct with an action based on FNetworkPhysicsActionPayload, to be executed at @param LocalFrame.
+	* @param SourceObject is the object that is the source of this action, not necessarily the object that the action should run on. */
+	ENGINE_API void EnqueueScheduledActionInstanceAtFrame_External(TInstancedStruct<FNetworkPhysicsActionPayload>&& ActionInstance, const UObject* SourceObject, const int32 LocalFrame);
+	
+	/** Enqueue a TInstancedStruct with an action based on FNetworkPhysicsActionPayload, to be executed at @param LocalFrame.
+	* @param SourceId can be left at 0 for actions triggered by the owner of this implementation, else if triggered by an external actor pass in a network stable hash.
+	* Note: On the game thread a network stable hash can be retrieved from replicated UObjects or level baked UObjects (for example actors) via UE::NetworkPhysicsUtils::GetNetworkStableHash_External(UObject); */
+	ENGINE_API void EnqueueScheduledActionInstanceAtFrame_External(TInstancedStruct<FNetworkPhysicsActionPayload>&& ActionInstance, const uint32 SourceId, const int32 LocalFrame);
+
+	/** Enqueue an action based on FNetworkPhysicsActionPayload, to be executed at @param LocalFrame.
+	* @param SourceObject is the object that is the source of this action, not necessarily the object that the action should run on. */
+	template<typename TActionType>
+	void EnqueueScheduledActionAtFrame_External(TActionType& Action, const UObject* SourceObject, const int32 LocalFrame)
+	{
+		static_assert(TIsDerivedFrom<TActionType, FNetworkPhysicsActionPayload>::Value, "EnqueueScheduledActionAtFrame_External - Action needs to derive from FNetworkPhysicsActionPayload");
+		EnqueueScheduledActionInstanceAtFrame_External(TInstancedStruct<FNetworkPhysicsActionPayload>::Make<TActionType>(Action), SourceObject, LocalFrame);
+	}
+
+	/** Enqueue a TInstancedStruct with an action based on FNetworkPhysicsActionPayload, to be executed at @param LocalFrame.
+	* @param SourceId can be left at 0 for actions triggered by the owner of this implementation, else if triggered by an external actor pass in a network stable hash.
+	* Note: On the game thread a network stable hash can be retrieved from replicated UObjects or level baked UObjects (for example actors) via UE::NetworkPhysicsUtils::GetNetworkStableHash_External(UObject); */
+	template<typename TActionType>
+	void EnqueueScheduledActionAtFrame_External(TActionType& Action, const uint32 SourceId, const int32 LocalFrame)
+	{
+		static_assert(TIsDerivedFrom<TActionType, FNetworkPhysicsActionPayload>::Value, "EnqueueScheduledActionAtFrame_External - Action needs to derive from FNetworkPhysicsActionPayload");
+		EnqueueScheduledActionInstanceAtFrame_External(TInstancedStruct<FNetworkPhysicsActionPayload>::Make<TActionType>(Action), SourceId, LocalFrame);
+	}
+
 	/** Mark this as controlled through locally relayed inputs rather than controlled as a pawn through a player controller.
 	* Set if NetworkPhysicsComponent is implemented on an AActor instead of APawn and it's currently being fed inputs, or if this is controlled by the server. 
 	* NOTE: The actor for this NetworkPhysicsComponent also needs to be owned by the local client if this is used client-side. */
@@ -1705,6 +2522,10 @@ protected:
 	UFUNCTION(Server, unreliable)
 	ENGINE_API void ServerReceiveInputCollection(const FNetworkPhysicsDataCollection& ClientInputCollection);
 
+	// Server RPC to receive actions from client
+	UFUNCTION(Server, unreliable)
+	ENGINE_API void ServerReceiveActionCollection(const FNetworkPhysicsActionCollection& ClientActionCollection);
+
 	// repnotify for inputs on owner client
 	UFUNCTION()
 	ENGINE_API void OnRep_SetReplicatedInputCollection();
@@ -1717,6 +2538,10 @@ protected:
 	UFUNCTION()
 	ENGINE_API void OnRep_SetReplicatedStateCollection();
 
+	// repnotify for actions on the client
+	UFUNCTION()
+	ENGINE_API void OnRep_SetReplicatedActionsCollection();
+
 	// replicated physics inputs for owner client
 	UPROPERTY(Transient, ReplicatedUsing = OnRep_SetReplicatedInputCollection)
 	FNetworkPhysicsDataCollection ReplicatedInputCollection;
@@ -1729,17 +2554,20 @@ protected:
 	UPROPERTY(Transient, ReplicatedUsing = OnRep_SetReplicatedStateCollection)
 	FNetworkPhysicsDataCollection ReplicatedStateCollection;
 
+	// replicated physics actions 
+	UPROPERTY(Transient, ReplicatedUsing = OnRep_SetReplicatedActionsCollection)
+	FNetworkPhysicsActionCollection ReplicatedActionsCollection;
+
 #if UE_WITH_REMOTE_OBJECT_HANDLE
 	ENGINE_API virtual void PostMigrate(const struct FUObjectMigrationContext& MigrationContext) override;
 #endif
 
-
 private:
 	/** Send network data from marshaled data */
 	void NetworkMarshaledData();
 
 	/** Update the async component with GT properties */
-	void UpdateAsyncComponent(const bool bFullUpdate);
+	ENGINE_API void UpdateAsyncComponent(const bool bFullUpdate);
 
 	/** Setup data type and history in async input and trigger the creation of data history in the async component */
 	ENGINE_API void CreateAsyncDataHistory();
@@ -1828,6 +2656,16 @@ public:
 	template<typename Input, typename State>
 	void CreateDataHistory(TNetworkPhysicsInputState_Internal<Input, State>* Implementation_Internal, TNetworkPhysicsInputState_External<Input, State>* Implementation_External = nullptr);
 
+	/** Set implementation interface, only needed if using legacy CreateDataHistory() but want to use Actions (which only support interface implementation), else use CreateDataHistory() that takes the interface.
+	* Note: If input and state are not being used, implement with FNetworkPhysicsPayload. */
+	template<typename Input, typename State>
+	void SetImplementationInteraces(TNetworkPhysicsInputState_Internal<Input, State>* Implementation_Internal, TNetworkPhysicsInputState_External<Input, State>* Implementation_External = nullptr)
+	{
+		ImplementationInterface_Internal = Implementation_Internal;
+		ImplementationInterface_External = Implementation_External;
+		UpdateAsyncComponent(true);
+	}
+
 private:
 
 	friend FNetworkPhysicsCallback;
@@ -2041,6 +2879,9 @@ struct FAsyncNetworkPhysicsComponentInput : public Chaos::FSimCallbackInput
 	TArray<TUniquePtr<Chaos::FBaseRewindHistory>> InputDataImportant;
 	TArray<TUniquePtr<Chaos::FBaseRewindHistory>> StateDataImportant;
 
+	TArray<TInstancedStruct<FNetworkPhysicsActionPayload>> LocalActionData;
+	TArray<TInstancedStruct<FNetworkPhysicsActionPayload>> ReplicatedActionData;
+
 	void Reset()
 	{
 		bIsLocallyControlled.Reset();
@@ -2075,6 +2916,9 @@ struct FAsyncNetworkPhysicsComponentInput : public Chaos::FSimCallbackInput
 		// Todo, optimize
 		InputDataImportant.Reset();
 		StateDataImportant.Reset();
+
+		LocalActionData.Reset();
+		ReplicatedActionData.Reset();
 	}
 };
 
@@ -2089,6 +2933,8 @@ struct FAsyncNetworkPhysicsComponentOutput : public Chaos::FSimCallbackOutput
 	TArray<TUniquePtr<Chaos::FBaseRewindHistory>> InputDataImportant;
 	TArray<TUniquePtr<Chaos::FBaseRewindHistory>> StateDataImportant;
 
+	TArray<TInstancedStruct<FNetworkPhysicsActionPayload>> ActionData;
+
 	void Reset()
 	{
 		InputsToNetwork_Owner.Reset();
@@ -2106,6 +2952,8 @@ struct FAsyncNetworkPhysicsComponentOutput : public Chaos::FSimCallbackOutput
 		// Todo, optimize
 		InputDataImportant.Reset();
 		StateDataImportant.Reset();
+
+		ActionData.Reset();
 	}
 };
 
@@ -2172,7 +3020,50 @@ public:
 	/** Returns the current amount of input decay for sim proxies during simulation as a magnitude from 0.0 to 1.0. */
 	ENGINE_API const float GetCurrentSimProxyInputDecayAtRuntime() const;
 
-	ENGINE_API const int32 GetLatestReceivedTargetStateFrame() const;
+	/** Enqueue a TInstancedStruct with an action based on FNetworkPhysicsActionPayload, to be executed as soon as possible.
+	* @param SourceId can be left at 0 for actions triggered by the owner of this implementation, else if triggered by an external actor pass in a network stable hash.
+	* Note: On the game thread a network stable hash can be retrieved from replicated UObjects or level baked UObjects (for example actors) via UE::NetworkPhysicsUtils::GetNetworkStableHash_External(UObject); */
+	ENGINE_API void EnqueueImmediateActionInstance_Internal(TInstancedStruct<FNetworkPhysicsActionPayload>&& ActionInstance, const uint32 SourceId);
+
+	/** Enqueue an action based on FNetworkPhysicsActionPayload, to be executed as soon as possible.
+	* @param SourceId can be left at 0 for actions triggered by the owner of this implementation, else if triggered by an external actor pass in a network stable hash.
+	* Note: On the game thread a network stable hash can be retrieved from replicated UObjects or level baked UObjects (for example actors) via UE::NetworkPhysicsUtils::GetNetworkStableHash_External(UObject); */
+	template<typename TActionType>
+	void EnqueueImmediateAction_Internal(TActionType& Action, const uint32 SourceId)
+	{
+		static_assert(TIsDerivedFrom<TActionType, FNetworkPhysicsActionPayload>::Value, "EnqueueImmediateAction_Internal - Action needs to derive from FNetworkPhysicsActionPayload");
+		EnqueueImmediateActionInstance_Internal(TInstancedStruct<FNetworkPhysicsActionPayload>::Make<TActionType>(Action), SourceId);
+	}
+
+	/** Enqueue a TInstancedStruct with an action based on FNetworkPhysicsActionPayload, to be executed after a delay of @param DelaySeconds.
+	* @param SourceId can be left at 0 for actions triggered by the owner of this implementation, else if triggered by an external actor pass in a network stable hash.
+	* Note: On the game thread a network stable hash can be retrieved from replicated UObjects or level baked UObjects (for example actors) via UE::NetworkPhysicsUtils::GetNetworkStableHash_External(UObject); */
+	ENGINE_API void EnqueueScheduledActionInstance_Internal(TInstancedStruct<FNetworkPhysicsActionPayload>&& ActionInstance, const uint32 SourceId, const float DelaySeconds);
+
+	/** Enqueue an action based on FNetworkPhysicsActionPayload, to be executed after a delay of @param DelaySeconds.
+	* @param SourceId can be left at 0 for actions triggered by the owner of this implementation, else if triggered by an external actor pass in a network stable hash.
+	* Note: On the game thread a network stable hash can be retrieved from replicated UObjects or level baked UObjects (for example actors) via UE::NetworkPhysicsUtils::GetNetworkStableHash_External(UObject); */
+	template<typename TActionType>
+	void EnqueueScheduledAction_Internal(TActionType& Action, const uint32 SourceId, const float DelaySeconds)
+	{
+		static_assert(TIsDerivedFrom<TActionType, FNetworkPhysicsActionPayload>::Value, "EnqueueScheduledAction_Internal - Action needs to derive from FNetworkPhysicsActionPayload");
+		EnqueueScheduledActionInstance_Internal(TInstancedStruct<FNetworkPhysicsActionPayload>::Make<TActionType>(Action), SourceId, DelaySeconds);
+	}
+
+	/** Enqueue a TInstancedStruct with an action based on FNetworkPhysicsActionPayload, to be executed at @param LocalFrame.
+	* @param SourceId can be left at 0 for actions triggered by the owner of this implementation, else if triggered by an external actor pass in a network stable hash.
+	* Note: On the game thread a network stable hash can be retrieved from replicated UObjects or level baked UObjects (for example actors) via UE::NetworkPhysicsUtils::GetNetworkStableHash_External(UObject); */
+	ENGINE_API void EnqueueScheduledActionInstanceAtFrame_Internal(TInstancedStruct<FNetworkPhysicsActionPayload>&& ActionInstance, const uint32 SourceId, const int32 LocalFrame);
+
+	/** Enqueue an action based on FNetworkPhysicsActionPayload, to be executed at @param LocalFrame
+	* @param SourceId can be left at 0 for actions triggered by the owner of this implementation, else if triggered by an external actor pass in a network stable hash.
+	* Note:	On the game thread a network stable hash can be retrieved from replicated UObjects or level baked UObjects (for example actors) via UE::NetworkPhysicsUtils::GetNetworkStableHash_External(UObject); */
+	template<typename TActionType>
+	void EnqueueScheduledActionAtFrame_Internal(TActionType& Action, const uint32 SourceId, const int32 LocalFrame)
+	{
+		static_assert(TIsDerivedFrom<TActionType, FNetworkPhysicsActionPayload>::Value, "EnqueueScheduledActionAtFrame_Internal - Action needs to derive from FNetworkPhysicsActionPayload");
+		EnqueueScheduledActionInstanceAtFrame_Internal(TInstancedStruct<FNetworkPhysicsActionPayload>::Make<TActionType>(Action), SourceId, LocalFrame);
+	}
 
 private:
 	// Initialize, bind delegates etc.
@@ -2201,6 +3092,15 @@ private:
 	// Consume data from async input
 	void ConsumeAsyncInput_Internal(const int32 PhysicsStep, const bool bTriggerResim);
 
+	// Add new action and handle related proposed actions
+	void RegisterAction_Internal(const TInstancedStruct<FNetworkPhysicsActionPayload>& NewActionInstance, const bool bWasReceivedFromReplication);
+
+	// Apply actions
+	void ApplyActionsForFrame(const int32 LocalFrame);
+
+	// Get a new unique ActionId
+	uint16 GetUniqueActionId_Internal();
+
 	/** Get the rigid solver */
 	Chaos::FPBDRigidsSolver* GetRigidSolver();
 
@@ -2225,6 +3125,9 @@ private:
 
 	/** Populate state data in AsyncOutput to send over the network */
 	void SendStateData_Internal(FAsyncNetworkPhysicsComponentOutput& AsyncOutput, const int32 PhysicsStep);
+	
+	/** Populate action data in AsyncOutput to send over the network */
+	void SendActionData_Internal(FAsyncNetworkPhysicsComponentOutput& AsyncOutput, const int32 PhysicsStep);
 
 	/** Update function on the server to manage how many inputs the owning client should send with each RPC */
 	void UpdateDynamicInputReplicationScaling();
@@ -2246,6 +3149,11 @@ private:
 
 	int32 LastReceivedStateFrame = INDEX_NONE;
 
+	// Local iterator to produce a unique ID for an Action
+	uint16 ActionIdIterator = 0;
+
+	int32 LastAppliedActionLocalFrame = INDEX_NONE;
+
 	// Component settings
 	TWeakPtr<const FNetworkPhysicsSettingsData> SettingsComponent;
 	static const FNetworkPhysicsSettingsNetworkPhysicsComponent SettingsNetworkPhysicsComponent_Default;
@@ -2280,6 +3188,15 @@ private:
 	// Local temporary states data used by pre/post process inputs functions
 	TUniquePtr<FNetworkPhysicsPayload> StateData;
 
+	// Confirmed actions for past, present and future frames
+	FActionHistory ConfirmedActions;
+
+	// Predicted actions for present and future frames
+	FActionHistory PredictedActions;
+
+	// Proposed actions pending confirmation
+	FActionHistory PendingProposedActions;
+
 	// Send last N number of inputs each replication call
 	uint16 InputsToNetwork_OwnerDefault = 3; // Default value for owning client
 	uint16 InputsToNetwork_Owner = 3; // From owning client, i.e. autonomous proxy or client owning an actor with bIsRelayingLocalInputs enabled
