Chaos Remote Objects: Fixing the extraction of ItemIndex in ConvertQueryOverlap to not dereference the owner when we're running with remote object handles.
#Synced-CL 52100670

diff --git a/Engine/Source/Runtime/Engine/Private/Collision/CollisionConversions.cpp b/Engine/Source/Runtime/Engine/Private/Collision/CollisionConversions.cpp
--- a/Engine/Source/Runtime/Engine/Private/Collision/CollisionConversions.cpp
+++ b/Engine/Source/Runtime/Engine/Private/Collision/CollisionConversions.cpp
@@ -724,17 +724,18 @@ void ConvertQueryOverlap(const FPhysicsShape& Shape, const FPhysicsActor& Actor,
 		}
 
 		// Set ItemIndex
-		if (const IPhysicsBodyInstanceOwner* BodyInstanceOwner = IPhysicsBodyInstanceOwner::GetPhysicsBodyInstanceOwnerFromOverlapResult(OutOverlap))
-		{
-			bool bUseBodyIndex = false;
+		int32 ItemIndex = INDEX_NONE;
+		// To avoid migrations, we always set the body index and we defer bMultiBodyOverlap to the getter
 #if UE_WITH_REMOTE_OBJECT_HANDLE
-			// To avoid migrations, we always set the body index and we defer bMultiBodyOverlap to the getter
-			bUseBodyIndex = true;
+		ItemIndex  = BodyInst->InstanceBodyIndex;
 #else
-			bUseBodyIndex = BodyInstanceOwner->IsMultiBodyOverlap();
-#endif
-			OutOverlap.SetItemIndex(bUseBodyIndex ? BodyInst->InstanceBodyIndex : INDEX_NONE);
+		if (const IPhysicsBodyInstanceOwner* BodyInstanceOwner = IPhysicsBodyInstanceOwner::GetPhysicsBodyInstanceOwnerFromOverlapResult(OutOverlap))
+		{
+			const bool bUseBodyIndex = BodyInstanceOwner->IsMultiBodyOverlap();
+			ItemIndex = bUseBodyIndex ? BodyInst->InstanceBodyIndex : INDEX_NONE;
 		}
+#endif
+		OutOverlap.SetItemIndex(ItemIndex);
 	}
 	else
 	{
