Fix CVar names for SkinWeightProfile cvars

* a.SkinWeightProfile.LoadByDefaultMode -> r.SkinWeightProfile.LoadByDefaultMode
* a.SkinWeightProfile.DefaultLODOverride -> r.SkinWeightProfile.DefaultLODOverride

Shadow deprecation cvars already exist for above cvars

#rnx
#Synced-CL 52037470

diff --git a/Engine/Source/Runtime/Engine/Private/Rendering/SkinWeightProfile.cpp b/Engine/Source/Runtime/Engine/Private/Rendering/SkinWeightProfile.cpp
--- a/Engine/Source/Runtime/Engine/Private/Rendering/SkinWeightProfile.cpp
+++ b/Engine/Source/Runtime/Engine/Private/Rendering/SkinWeightProfile.cpp
@@ -107,7 +107,7 @@ static void OnDefaultProfileCVarsChanged(IConsoleVariable* Variable)
 
 int32 GSkinWeightProfilesLoadByDefaultMode = -1;
 FAutoConsoleVariableRef CVarSkinWeightsLoadByDefaultMode(
-	TEXT("a.SkinWeightProfile.LoadByDefaultMode"),
+	TEXT("r.SkinWeightProfile.LoadByDefaultMode"),
 	GSkinWeightProfilesLoadByDefaultMode,
 	TEXT("Enables/disables run-time optimization to override the original skin weights with a profile designated as the default to replace it. Can be used to optimize memory for specific platforms or devices")
 	TEXT("-1 = disabled")
@@ -121,7 +121,7 @@ FAutoConsoleVariableRef CVarSkinWeightsLoadByDefaultMode(
 
 int32 GSkinWeightProfilesDefaultLODOverride = -1;
 FAutoConsoleVariableRef CVarSkinWeightProfilesDefaultLODOverride(
-	TEXT("a.SkinWeightProfile.DefaultLODOverride"),
+	TEXT("r.SkinWeightProfile.DefaultLODOverride"),
 	GSkinWeightProfilesDefaultLODOverride,
 	TEXT("Override LOD index from which on the default Skin Weight Profile should override the Skeletal Mesh's default Skin Weights"),
 	FConsoleVariableDelegate::CreateStatic(&OnDefaultProfileCVarsChanged),
