Substrate - Toon fix for anisotropy that was not accounted for in some cases.

#rnx
#Synced-CL 52122838

diff --git a/Engine/Source/Runtime/Engine/Private/Materials/HLSLMaterialTranslator.cpp b/Engine/Source/Runtime/Engine/Private/Materials/HLSLMaterialTranslator.cpp
--- a/Engine/Source/Runtime/Engine/Private/Materials/HLSLMaterialTranslator.cpp
+++ b/Engine/Source/Runtime/Engine/Private/Materials/HLSLMaterialTranslator.cpp
@@ -1745,6 +1745,10 @@ void FHLSLMaterialTranslator::TranslateMaterial()
 		MaterialCompilationOutput.SubstrateMaterialCompilationOutput.SubstrateMaterialType = SUBSTRATE_MATERIAL_TYPE_COMPLEX;
 		MaterialCompilationOutput.SubstrateMaterialCompilationOutput.SubstrateMaterialBsdfFeatures |= ESubstrateBsdfFeature::Anisotropy;
 	}
+	if (uint16(MaterialCompilationOutput.SubstrateMaterialCompilationOutput.SubstrateMaterialBsdfFeatures & (ESubstrateBsdfFeature::Toon | ESubstrateBsdfFeature::Anisotropy)) > 0)
+	{
+		MaterialCompilationOutput.bUsesAnisotropy = true;
+	}
 
 	/**
 	 * A material is not compatible with the atlas if 
