Change a delegate to a thread safe version.

#rnx

- With recent changes in the engine it is now possible for UWorld::SelectedLevelsChangedEvent to accessed on both the Game and AsyncLoading threads in the editor which is unsafe.
- Since UWorld exposes the delegate by a getter we can't easy add our own thread protection so the easiest and safest fix it to change the delegate to a thread version.
- DECLARE_EVENT does not have a thread safe version but the code documentation for DECLARE_EVENT states that it is in a state of semi decprecation and should be changed to multicast delegate anyway. Soo we can change it to DECLARE_TS_MULTICAST_DELEGATE.

#Synced-CL 52058780

diff --git a/Engine/Source/Runtime/Engine/Classes/Engine/World.h b/Engine/Source/Runtime/Engine/Classes/Engine/World.h
--- a/Engine/Source/Runtime/Engine/Classes/Engine/World.h
+++ b/Engine/Source/Runtime/Engine/Classes/Engine/World.h
@@ -1850,14 +1850,14 @@ private:
 
 #if WITH_EDITORONLY_DATA
 	/** Called when selected level list changes. */
-	DECLARE_EVENT( UWorld, FOnSelectedLevelsChangedEvent);
+	DECLARE_TS_MULTICAST_DELEGATE(FOnSelectedLevelsChangedEvent);
 
 	/** Broadcasts whenever selected level list changes. */
 	FOnSelectedLevelsChangedEvent				SelectedLevelsChangedEvent;
 
 	/** Array of selected levels currently in this world. Not serialized to disk to avoid hard references.	*/
 	UPROPERTY(Transient)
-	TArray<TObjectPtr<class ULevel>>						SelectedLevels;
+	TArray<TObjectPtr<class ULevel>>			SelectedLevels;
 
 	/** Disables the broadcasting of level selection change. Internal use only. */
 	uint32 bBroadcastSelectionChange:1;
