Fix null pointer matching in FinishCompilation (UE-361151, UE-365823, UE-370055)

RefPoseSeq can be null for non-additive animations. When passed to
FinishCompilation, a null pointer matches any expired TWeakObjectPtr in
RegisteredAnimSequences (since both resolve to null), causing unrelated
sequences pending GC to be incorrectly "finished" and wrapped in
TStrongObjectPtr(nullptr) which asserts on operator->() during Reschedule.

Added null check in FinishCompilation before the Contains lookup to reject
null pointers from matching against the registered set.

#Synced-CL 52128385

diff --git a/Engine/Source/Runtime/Engine/Private/Animation/AnimationSequenceCompiler.cpp b/Engine/Source/Runtime/Engine/Private/Animation/AnimationSequenceCompiler.cpp
--- a/Engine/Source/Runtime/Engine/Private/Animation/AnimationSequenceCompiler.cpp
+++ b/Engine/Source/Runtime/Engine/Private/Animation/AnimationSequenceCompiler.cpp
@@ -153,7 +153,7 @@ namespace UE::Anim
 
 		for (UAnimSequence* AnimSequence : InAnimSequences)
 		{
-			if (RegisteredAnimSequences.Contains(AnimSequence))
+			if (AnimSequence && RegisteredAnimSequences.Contains(AnimSequence))
 			{
 				PendingAnimationSequences.Emplace(AnimSequence);
 			}
