Fix for use of LODScale during virtual texture compilation.
Previously we would clamp to tile size when compiling virtual textures, but didn't consider the effect of any LODScale when doing that clamp.
Now CalculateLODScale() adjusts its result to tile size into account.
This happens early in the compilation pipeline so that adjusted LODScale is fed into the DDCKey correctly.
#Synced-CL 52125649

diff --git a/Engine/Plugins/Mutable/Source/CustomizableObjectEditor/Private/MuCOE/GenerateMutableSource/GenerateMutableSource.cpp b/Engine/Plugins/Mutable/Source/CustomizableObjectEditor/Private/MuCOE/GenerateMutableSource/GenerateMutableSource.cpp
--- a/Engine/Plugins/Mutable/Source/CustomizableObjectEditor/Private/MuCOE/GenerateMutableSource/GenerateMutableSource.cpp
+++ b/Engine/Plugins/Mutable/Source/CustomizableObjectEditor/Private/MuCOE/GenerateMutableSource/GenerateMutableSource.cpp
@@ -1590,7 +1590,7 @@ uint32 ComputeLODBiasForTexture(const FMutableGraphGenerationContext& Generation
 		const int64 ReferenceTextureSize = MaxRefSourceSize >> PlatformLODBias;
 
 		// Additional LODBias of the Texture
-		const int32 ReferenceTextureLODBias = LODSettings.CalculateLODBias(ReferenceTextureSize, ReferenceTextureSize, 0,	ReferenceTexture->LODGroup,
+		const int32 ReferenceTextureLODBias = LODSettings.CalculateLODBias(ReferenceTextureSize, ReferenceTextureSize, 0, 0, ReferenceTexture->LODGroup,
 			ReferenceTexture->LODBias, 0, ReferenceTexture->MipGenSettings, ReferenceTexture->IsCurrentlyVirtualTextured(), ReferenceTexture->IsStreamable());
 
 		return FMath::Max(NumMipsSource - NumMipsRefSource + PlatformLODBias + ReferenceTextureLODBias, 0);
@@ -1603,7 +1603,7 @@ uint32 ComputeLODBiasForTexture(const FMutableGraphGenerationContext& Generation
 	const int64 TextureSize = SourceSize >> PlatformLODBias;
 
 	// Additional LODBias of the Texture
-	const int32 TextureLODBias = LODSettings.CalculateLODBias(TextureSize, TextureSize, 0, Texture.LODGroup, Texture.LODBias, 0, Texture.MipGenSettings, Texture.IsCurrentlyVirtualTextured(), Texture.IsStreamable());
+	const int32 TextureLODBias = LODSettings.CalculateLODBias(TextureSize, TextureSize, 0, 0, Texture.LODGroup, Texture.LODBias, 0, Texture.MipGenSettings, Texture.IsCurrentlyVirtualTextured(), Texture.IsStreamable());
 
 	return FMath::Max(PlatformLODBias + TextureLODBias, 0);
 }
@@ -1638,7 +1638,7 @@ int32 GetTextureSizeInGame(const UTexture2D& Texture, const UTextureLODSettings&
 	const int32 MaxTextureSizeAllowed = SourceSize >> PlatformLODBias;
 
 	// Calculate the LODBias specific for this texture 
-	const int32 TextureLODBias = LODSettings.CalculateLODBias(MaxTextureSizeAllowed, MaxTextureSizeAllowed, 0, Texture.LODGroup, Texture.LODBias, 0, Texture.MipGenSettings, Texture.IsCurrentlyVirtualTextured(), Texture.IsStreamable());
+	const int32 TextureLODBias = LODSettings.CalculateLODBias(MaxTextureSizeAllowed, MaxTextureSizeAllowed, 0, 0, Texture.LODGroup, Texture.LODBias, 0, Texture.MipGenSettings, Texture.IsCurrentlyVirtualTextured(), Texture.IsStreamable());
 
 	return MaxTextureSizeAllowed >> TextureLODBias;
 }
diff --git a/Engine/Source/Editor/TextureEditor/Private/TextureEditorToolkit.cpp b/Engine/Source/Editor/TextureEditor/Private/TextureEditorToolkit.cpp
--- a/Engine/Source/Editor/TextureEditor/Private/TextureEditorToolkit.cpp
+++ b/Engine/Source/Editor/TextureEditor/Private/TextureEditorToolkit.cpp
@@ -752,7 +752,7 @@ void FTextureEditorToolkit::PopulateQuickInfo( )
 											const UTextureLODSettings* LODSettings = 
 												UDeviceProfileManager::Get().FindProfile(TargetPlatform->CookingDeviceProfileName())->GetTextureLODSettings();
 											int32 FirstMipToSerialize = LODSettings->CalculateLODBias(
-												TextureDescription.TopMipSizeX, TextureDescription.TopMipSizeY, Texture->MaxTextureSize, 
+												TextureDescription.TopMipSizeX, TextureDescription.TopMipSizeY, Texture->GetMinTextureSize(), Texture->GetMaxTextureSize(),
 												Texture->LODGroup, Texture->LODBias, 0, Texture->MipGenSettings, false, Texture->IsStreamable());
 
 											int32 StrippedNumMips = TextureDescription.NumMips;
diff --git a/Engine/Source/Runtime/Engine/Classes/Engine/Texture.h b/Engine/Source/Runtime/Engine/Classes/Engine/Texture.h
--- a/Engine/Source/Runtime/Engine/Classes/Engine/Texture.h
+++ b/Engine/Source/Runtime/Engine/Classes/Engine/Texture.h
@@ -1917,7 +1917,12 @@ public:
 	ENGINE_API void GetDefaultFormatSettings(FTextureFormatSettings& OutSettings) const;
 	
 	ENGINE_API ETextureEncodeSpeed GetDesiredEncodeSpeed() const;
-		
+
+	/** Returns the minimum size that width/height/depth will be clamped to when generating this texture. */
+	ENGINE_API int32 GetMinTextureSize() const;
+	/** Returns the maximum size that width/height/depth will be clamped to when generating this texture. */
+	int32 GetMaxTextureSize() const { return MaxTextureSize; }
+
 	/** Ensure settings are valid after import or edit; this is called by PostEditChange. */
 	ENGINE_API virtual void ValidateSettingsAfterImportOrEdit(bool * pRequiresNotifyMaterials = nullptr);
 
diff --git a/Engine/Source/Runtime/Engine/Classes/Engine/TextureLODSettings.h b/Engine/Source/Runtime/Engine/Classes/Engine/TextureLODSettings.h
--- a/Engine/Source/Runtime/Engine/Classes/Engine/TextureLODSettings.h
+++ b/Engine/Source/Runtime/Engine/Classes/Engine/TextureLODSettings.h
@@ -158,9 +158,8 @@ class UTextureLODSettings : public UObject
 
 public:
 
-
 	/**
-	 * Calculates and returns the LOD bias based on texture LOD group, LOD bias and maximum size.
+	 * Calculates and returns the LOD bias based on texture LOD group, LOD bias and minimum/maximum size.
 	 *
 	 * @param	Texture				Texture object to calculate LOD bias for.
 	 * @param	bIncCinematicMips	If true, cinematic mips will also be included in consideration
@@ -173,6 +172,8 @@ public:
 	 *
 	 * @param	Width						Width of the texture
 	 * @param	Height						Height of the texture
+	 * @param	MinSize						Min size to clamp width or height of the texture
+	 * @param	MaxSize						Max size to clamp width or height of the texture
 	 * @param	LODGroup					Which LOD group the texture belongs to
 	 * @param	LODBias						LOD bias to include in the calculation
 	 * @param	NumCinematicMipLevels		The texture cinematic mip levels to include in the calculation
@@ -181,7 +182,13 @@ public:
 	* @param	bCurrentlyStreamableTexture	If texture is currently streamable (if it's not, then quality levels don't apply to it)
 	 * @return	LOD bias
 	 */
-	ENGINE_API int32 CalculateLODBias( int32 Width, int32 Height, int32 MaxSize, int32 LODGroup, int32 LODBias, int32 NumCinematicMipLevels, TextureMipGenSettings MipGenSetting, bool bVirtualTexture, bool bCurrentlyStreamableTexture = false) const;
+	ENGINE_API int32 CalculateLODBias(int32 Width, int32 Height, int32 MinSize, int32 MaxSize, int32 LODGroup, int32 LODBias, int32 NumCinematicMipLevels, TextureMipGenSettings MipGenSetting, bool bVirtualTexture, bool bCurrentlyStreamableTexture = false) const;
+
+	UE_DEPRECATED(5.8, "Call newer version of function that takes MinSize.")
+	int32 CalculateLODBias(int32 Width, int32 Height, int32 MaxSize, int32 LODGroup, int32 LODBias, int32 NumCinematicMipLevels, TextureMipGenSettings MipGenSetting, bool bVirtualTexture, bool bCurrentlyStreamableTexture = false) const
+	{
+		return CalculateLODBias(Width, Height, 0, MaxSize, LODGroup, LODBias, NumCinematicMipLevels, MipGenSetting, bVirtualTexture, bCurrentlyStreamableTexture);
+	}
 
 	/**
 	 * Calculate num optional mips
diff --git a/Engine/Source/Runtime/Engine/Private/Texture.cpp b/Engine/Source/Runtime/Engine/Private/Texture.cpp
--- a/Engine/Source/Runtime/Engine/Private/Texture.cpp
+++ b/Engine/Source/Runtime/Engine/Private/Texture.cpp
@@ -4088,7 +4088,7 @@ void UTexture::GetBuiltTextureSize(const ITargetPlatformSettings* TargetPlatform
 		const bool bVirtualTextureStreaming = VirtualTextureStreaming && IsVirtualTexturingEnabled(TargetPlatformSettings);
 
 		const UTextureLODSettings& LODSettings = TargetPlatformSettings->GetTextureLODSettings();
-		const uint32 LODBiasNoCinematics = FMath::Max<int32>(LODSettings.CalculateLODBias(SizeX, SizeY, MaxTextureSize, LODGroup, LODBias, 0, MipGenSettings, bVirtualTextureStreaming, IsStreamable()), 0);
+		const uint32 LODBiasNoCinematics = FMath::Max<int32>(LODSettings.CalculateLODBias(SizeX, SizeY, GetMinTextureSize(), GetMaxTextureSize(), LODGroup, LODBias, 0, MipGenSettings, bVirtualTextureStreaming, IsStreamable()), 0);
 		SizeX = FMath::Max<int32>(SizeX >> LODBiasNoCinematics, 1);
 		SizeY = FMath::Max<int32>(SizeY >> LODBiasNoCinematics, 1);
 		if ( bIsVolume )
diff --git a/Engine/Source/Runtime/Engine/Private/Texture2D.cpp b/Engine/Source/Runtime/Engine/Private/Texture2D.cpp
--- a/Engine/Source/Runtime/Engine/Private/Texture2D.cpp
+++ b/Engine/Source/Runtime/Engine/Private/Texture2D.cpp
@@ -1343,7 +1343,7 @@ int32 UTexture2D::Blueprint_GetSizeX() const
 		//	MaxTextureSize, PadForPow2, Downscale, etc. are not applied
 
 		const UTextureLODSettings* LODSettings = UDeviceProfileManager::Get().GetActiveProfile()->GetTextureLODSettings();
-		const int32 CookedLODBias = LODSettings->CalculateLODBias(Source.SizeX, Source.SizeY, MaxTextureSize, LODGroup, LODBias, 0, MipGenSettings, IsCurrentlyVirtualTextured(), IsStreamable());
+		const int32 CookedLODBias = LODSettings->CalculateLODBias(Source.SizeX, Source.SizeY, GetMinTextureSize(), GetMaxTextureSize(), LODGroup, LODBias, 0, MipGenSettings, IsCurrentlyVirtualTextured(), IsStreamable());
 		return FMath::Max<int32>(Source.SizeX >> CookedLODBias, 1);
 	}
 #endif
@@ -1406,7 +1406,7 @@ int32 UTexture2D::Blueprint_GetSizeY() const
 	if (!GetSizeY())
 	{
 		const UTextureLODSettings* LODSettings = UDeviceProfileManager::Get().GetActiveProfile()->GetTextureLODSettings();
-		const int32 CookedLODBias = LODSettings->CalculateLODBias(Source.SizeX, Source.SizeY, MaxTextureSize, LODGroup, LODBias, 0, MipGenSettings, IsCurrentlyVirtualTextured(), IsStreamable());
+		const int32 CookedLODBias = LODSettings->CalculateLODBias(Source.SizeX, Source.SizeY, GetMinTextureSize(), GetMaxTextureSize(), LODGroup, LODBias, 0, MipGenSettings, IsCurrentlyVirtualTextured(), IsStreamable());
 		return FMath::Max<int32>(Source.SizeY >> CookedLODBias, 1);
 	}
 #endif
diff --git a/Engine/Source/Runtime/Engine/Private/TextureDerivedData.cpp b/Engine/Source/Runtime/Engine/Private/TextureDerivedData.cpp
--- a/Engine/Source/Runtime/Engine/Private/TextureDerivedData.cpp
+++ b/Engine/Source/Runtime/Engine/Private/TextureDerivedData.cpp
@@ -841,7 +841,7 @@ static void ModifyMaxTextureResolutionBuildSettingsForPlatformLimit(
 		// after scaling to MaxTextureResolution, LODBiasNoCinematics will be applied
 			
 		const UTextureLODSettings& LODSettings = TargetPlatform->GetTextureLODSettings();
- 		const uint32 LODBiasNoCinematics = FMath::Max<int32>(LODSettings.CalculateLODBias(BuiltSizeX, BuiltSizeY, Texture.MaxTextureSize, Texture.LODGroup, Texture.LODBias, 0, Texture.MipGenSettings, OutSettings.bVirtualStreamable, Texture.IsStreamable()), 0);
+ 		const uint32 LODBiasNoCinematics = FMath::Max<int32>(LODSettings.CalculateLODBias(BuiltSizeX, BuiltSizeY, Texture.GetMinTextureSize(), Texture.GetMaxTextureSize(), Texture.LODGroup, Texture.LODBias, 0, Texture.MipGenSettings, OutSettings.bVirtualStreamable, Texture.IsStreamable()), 0);
 
 		int64 MaxTextureResolutionUp = ((int64)OutSettings.MaxTextureResolution)<<LODBiasNoCinematics;
 
@@ -1072,6 +1072,17 @@ ENGINE_API ETextureEncodeSpeed UTexture::GetDesiredEncodeSpeed() const
 	return FResolvedTextureEncodingSettings::Get().EncodeSpeed;
 }
 
+int32 UTexture::GetMinTextureSize() const
+{
+	if (VirtualTextureStreaming)
+	{
+		FVirtualTextureBuildSettings VirtualTextureBuildSettings;
+		GetVirtualTextureBuildSettings(VirtualTextureBuildSettings);
+		return FVirtualTextureBuildSettings::ClampAndAlignTileSize(VirtualTextureBuildSettings.TileSize);
+	}
+	return 0;
+}
+
 // from Texture.cpp
 extern FName GetLatestOodleTextureSdkVersion();
 
@@ -1328,8 +1339,8 @@ static void GetTextureBuildSettings(
 		OutBuildSettings.CompositeTextureMode = CTM_Disabled;
 	}
 
-	OutBuildSettings.LODBias = TextureLODSettings.CalculateLODBias(SourceSize.X, SourceSize.Y, Texture.MaxTextureSize, Texture.LODGroup, Texture.LODBias, Texture.NumCinematicMipLevels, Texture.MipGenSettings, bVirtualTextureStreaming, Texture.IsStreamable());
-	OutBuildSettings.LODBiasWithCinematicMips = TextureLODSettings.CalculateLODBias(SourceSize.X, SourceSize.Y, Texture.MaxTextureSize, Texture.LODGroup, Texture.LODBias, 0, Texture.MipGenSettings, bVirtualTextureStreaming, Texture.IsStreamable());
+	OutBuildSettings.LODBias = TextureLODSettings.CalculateLODBias(SourceSize.X, SourceSize.Y, Texture.GetMinTextureSize(), Texture.GetMaxTextureSize(), Texture.LODGroup, Texture.LODBias, Texture.NumCinematicMipLevels, Texture.MipGenSettings, bVirtualTextureStreaming, Texture.IsStreamable());
+	OutBuildSettings.LODBiasWithCinematicMips = TextureLODSettings.CalculateLODBias(SourceSize.X, SourceSize.Y, Texture.GetMinTextureSize(), Texture.GetMaxTextureSize(), Texture.LODGroup, Texture.LODBias, 0, Texture.MipGenSettings, bVirtualTextureStreaming, Texture.IsStreamable());
 	OutBuildSettings.PowerOfTwoMode = Texture.PowerOfTwoMode;
 	OutBuildSettings.PaddingColor = Texture.PaddingColor;
 	OutBuildSettings.bPadWithBorderColor = Texture.bPadWithBorderColor;
@@ -2947,7 +2958,7 @@ static void SerializePlatformData(
 		const int32 FirstMipTailMip = NumMips - (int32)PlatformData->GetNumMipsInTail();
 		check(FirstMipTailMip >= 0);
 
-		FirstMipToSerialize = Ar.CookingTarget()->GetTextureLODSettings().CalculateLODBias(Width, Height, Texture->MaxTextureSize, LODGroup, LODBias, 0, MipGenSetting, bIsVirtual, Texture->IsStreamable());
+		FirstMipToSerialize = Ar.CookingTarget()->GetTextureLODSettings().CalculateLODBias(Width, Height, Texture->GetMinTextureSize(), Texture->GetMaxTextureSize(), LODGroup, LODBias, 0, MipGenSetting, bIsVirtual, Texture->IsStreamable());
 		if (!bIsVirtual)
 		{
 			// Reassign NumMips as the number of mips starting from FirstMipToSerialize.
diff --git a/Engine/Source/Runtime/Engine/Private/TextureLODSettings.cpp b/Engine/Source/Runtime/Engine/Private/TextureLODSettings.cpp
--- a/Engine/Source/Runtime/Engine/Private/TextureLODSettings.cpp
+++ b/Engine/Source/Runtime/Engine/Private/TextureLODSettings.cpp
@@ -211,21 +211,23 @@ int32 UTextureLODSettings::CalculateLODBias(const UTexture* Texture, bool bIncCi
 {	
 	check( Texture );
 	TextureMipGenSettings MipGenSetting = TMGS_MAX;
+	int32 TextureMinSize = 0;
 	int32 TextureMaxSize = 0;
 
 #if WITH_EDITORONLY_DATA
 	MipGenSetting = Texture->MipGenSettings;
-	TextureMaxSize = Texture->MaxTextureSize;
+	TextureMinSize = Texture->GetMinTextureSize();
+	TextureMaxSize = Texture->GetMaxTextureSize();
 #endif // #if WITH_EDITORONLY_DATA
 
 	// note: these are the Platform texture size and can return 0 during build!
 	float Width = Texture->GetSurfaceWidth();
 	float Height = Texture->GetSurfaceHeight();
 
-	return CalculateLODBias( FMath::RoundToInt(Width), FMath::RoundToInt(Height), TextureMaxSize, Texture->LODGroup, Texture->LODBias, bIncCinematicMips ? Texture->NumCinematicMipLevels : 0, MipGenSetting, Texture->IsCurrentlyVirtualTextured(), Texture->IsStreamable());
+	return CalculateLODBias( FMath::RoundToInt(Width), FMath::RoundToInt(Height), TextureMinSize, TextureMaxSize, Texture->LODGroup, Texture->LODBias, bIncCinematicMips ? Texture->NumCinematicMipLevels : 0, MipGenSetting, Texture->IsCurrentlyVirtualTextured(), Texture->IsStreamable());
 }
 
-int32 UTextureLODSettings::CalculateLODBias(int32 Width, int32 Height, int32 MaxSize, int32 LODGroup, int32 LODBias, int32 NumCinematicMipLevels, TextureMipGenSettings InMipGenSetting, bool bVirtualTexture, bool bCurrentlyStreamableTexture) const
+int32 UTextureLODSettings::CalculateLODBias(int32 Width, int32 Height, int32 MinSize, int32 MaxSize, int32 LODGroup, int32 LODBias, int32 NumCinematicMipLevels, TextureMipGenSettings InMipGenSetting, bool bVirtualTexture, bool bCurrentlyStreamableTexture) const
 {	
 	checkf(LODGroup < TextureLODGroups.Num(), TEXT("A texture had passed a bad LODGroup to UTextureLODSettings::CalculateLODBias (%d, out of %d groups). This code does not have the texture name. The LODSettings object is '%s'"), LODGroup, TextureLODGroups.Num(), *GetPathName());
 
@@ -282,6 +284,13 @@ int32 UTextureLODSettings::CalculateLODBias(int32 Width, int32 Height, int32 Max
 		// -> this looks broken, probably should be inside the !RequiresCookedData
 	}
 
+	// Clamp the accumulated UsedLodBias to account for MinSize.
+	if (MinSize > 0)
+	{
+		const int32 MaxLODBias = FMath::Max((int32)FMath::CeilLogTwo(FMath::Min(Width, Height)) - (int32)FMath::CeilLogTwo(MinSize), 0);
+		UsedLODBias = FMath::Min(UsedLODBias, MaxLODBias);
+	}
+
 	// Min/Max LOD Size clamps the *LODBias*, not the actual texture LODs. Meaning:
 	//	MaxLODSize caps the largest mip size.
 	//	MinLODSize prevents LODBias from making the largest mip smaller than the given value. Does *not*
