Instead of flipping layers we now preview quad layers in the blueprint editor. This is due to the fact that fliping would also flip shapes like cylinder layers

#virtualized
#Synced-CL 52135110
#Synced-CL 52135110

diff --git a/Engine/Source/Runtime/Engine/Classes/Components/StereoLayerComponent.h b/Engine/Source/Runtime/Engine/Classes/Components/StereoLayerComponent.h
--- a/Engine/Source/Runtime/Engine/Classes/Components/StereoLayerComponent.h
+++ b/Engine/Source/Runtime/Engine/Classes/Components/StereoLayerComponent.h
@@ -6,6 +6,11 @@
 #include "UObject/ObjectMacros.h"
 #include "Components/SceneComponent.h"
 #include "IStereoLayers.h"
+
+#if WITH_EDITORONLY_DATA
+#include "Materials/Material.h"
+#endif
+
 #include "StereoLayerComponent.generated.h"
 
 class UTexture;
@@ -68,6 +73,11 @@ public:
 
 protected:
 	ENGINE_API void MarkStereoLayerDirty();
+	
+#if WITH_EDITORONLY_DATA
+	UPROPERTY(Transient)
+	TObjectPtr<UMaterial> DebugTextureMaterial = nullptr;
+#endif
 };
 
 UCLASS(meta = (DisplayName = "Quad Layer"), MinimalAPI)
diff --git a/Engine/Source/Runtime/Engine/Private/Components/StereoLayerComponent.cpp b/Engine/Source/Runtime/Engine/Private/Components/StereoLayerComponent.cpp
--- a/Engine/Source/Runtime/Engine/Private/Components/StereoLayerComponent.cpp
+++ b/Engine/Source/Runtime/Engine/Private/Components/StereoLayerComponent.cpp
@@ -7,23 +7,71 @@
 #include "StereoRendering.h"
 #include "StereoLayerAdditionalFlagsManager.h"
 
+
 #include UE_INLINE_GENERATED_CPP_BY_NAME(StereoLayerComponent)
 #if WITH_EDITOR
+#include "DynamicMeshBuilder.h"
+#include "Materials/MaterialInstanceDynamic.h"
+#include "Materials/MaterialRenderProxy.h"
 #include "PrimitiveDrawingUtils.h"
+#include "SceneView.h"
 #else
 #include "StereoRendering.h"
 #endif
 
-namespace UE::XRStereoLayerCVars
+#if WITH_EDITORONLY_DATA
+void DrawStereoLayerQuad(
+	FPrimitiveDrawInterface* PDI,
+	const FVector& Center,
+	const FVector& XAxis,
+	const FVector& YAxis, 
+	FColor Color,
+	float Width,
+	float Height,
+	const FMaterialRenderProxy* MaterialRenderProxy,
+	uint8 DepthPriority)
 {
-	static bool bUseLegacyFacing = false;
-	static FAutoConsoleVariableRef UseUECoordinateSystemCVar(
-		TEXT("XR.StereoLayers.UseLegacyFacing"),
-		bUseLegacyFacing,
-		TEXT("If true the old UE coordinate system will be used where layers facing away from the camera"
-			 "are visible instead of the ones seemingly facing it inside of UE"),
-		ECVF_Default);
+	FVector XOffset = XAxis * Width * 0.5f;
+	FVector YOffset = YAxis * Height * 0.5f;
+
+	// Calculate verts for a face lying in plane defined by the X and Y vectors
+	FVector Positions[4] =
+	{
+		Center - XOffset - YOffset,
+		Center + XOffset - YOffset,
+		Center + XOffset + YOffset,
+		Center - XOffset + YOffset
+	};
+
+	FVector2f UVs[4] =
+	{
+		FVector2f(0.0f,0.0f),
+		FVector2f(0.0f,1.0f),
+		FVector2f(1.0f,1.0f),
+		FVector2f(1.0f,0.0f),
+	};
+
+	FDynamicMeshBuilder MeshBuilder(PDI->View->GetFeatureLevel());
+
+	int32 VertexIndices[4];
+	for (int32 VertexIndex = 0; VertexIndex < 4; VertexIndex++)
+	{
+		VertexIndices[VertexIndex] = MeshBuilder.AddVertex(
+			(FVector3f)Positions[VertexIndex],
+			UVs[VertexIndex],
+			FVector3f(1.0f, 0.0f, 0.0f),
+			FVector3f(0.0f, 1.0f, 0.0f),
+			FVector3f(0.0f, 0.0f, 1.0f),
+			Color
+		);
+	}
+
+	MeshBuilder.AddTriangle(VertexIndices[0], VertexIndices[2], VertexIndices[1]);
+	MeshBuilder.AddTriangle(VertexIndices[0], VertexIndices[3], VertexIndices[2]);
+
+	MeshBuilder.Draw(PDI, FMatrix::Identity, MaterialRenderProxy, DepthPriority, false);
 }
+#endif
 
 UStereoLayerComponent::UStereoLayerComponent(const FObjectInitializer& ObjectInitializer)
 	: Super(ObjectInitializer)
@@ -120,12 +168,6 @@ void UStereoLayerComponent::TickComponent(float DeltaTime, enum ELevelTick TickT
 		LayerDesc.UVRect = UVRect;
 		LayerDesc.Transform = Transform;
 		
-		if (!UE::XRStereoLayerCVars::bUseLegacyFacing)
-		{
-			// The default Open Xr Stereo layer object faces the camera at no rotation where as the unreal object faces away from it at zero rotation
-			LayerDesc.Transform.SetRotation(Transform.GetRotation() * FQuat(0., 0., 1., 0.));
-		}
-
 		if (Texture)
 		{
 			Texture->SetForceMipLevelsToBeResident(30.0f);
@@ -355,6 +397,11 @@ void UStereoLayerShape::DrawShapeVisualization(const class FSceneView* View, cla
 
 void UStereoLayerShapeQuad::DrawShapeVisualization(const class FSceneView* View, class FPrimitiveDrawInterface* PDI)
 {
+	if (!ensure(GEngine) || !ensure(GEngine->EmissiveMeshMaterial))
+	{
+		return;
+	}
+	
 	FLinearColor YellowColor = FColor(231, 239, 0, 255);
 	check(GetOuter()->IsA<UStereoLayerComponent>());
 
@@ -363,6 +410,29 @@ void UStereoLayerShapeQuad::DrawShapeVisualization(const class FSceneView* View,
 	const FBox QuadBox(FVector(0.0f, -QuadSize.X, -QuadSize.Y), FVector(0.0f, QuadSize.X, QuadSize.Y));
 
 	DrawWireBox(PDI, StereoLayerComp->GetComponentTransform().ToMatrixWithScale(), QuadBox, YellowColor, 0);
+
+	if (!DebugTextureMaterial)
+	{
+		DebugTextureMaterial = DuplicateObject<UMaterial>(GEngine->EmissiveMeshMaterial, this);
+		DebugTextureMaterial->TwoSided = false;
+		DebugTextureMaterial->PostEditChange();
+		DebugTextureMaterial->SetTextureParameterValueEditorOnly(NAME_LinearColor, StereoLayerComp->GetTexture());
+		
+		TArray<FMaterialParameterInfo> OutParameterInfo;
+		TArray<FGuid> OutParameterIds;
+		DebugTextureMaterial->GetAllTextureParameterInfo(OutParameterInfo, OutParameterIds);
+	}
+	
+	DrawStereoLayerQuad(
+		PDI,
+		StereoLayerComp->GetComponentTransform().GetLocation(),
+		FVector(0, 0, 1),
+		FVector(0, 1, 0),
+		YellowColor.ToFColor(true),
+		StereoLayerComp->GetQuadSize().X,
+		StereoLayerComp->GetQuadSize().Y,
+		DebugTextureMaterial->GetRenderProxy(),
+		0);
 }
 
 void UStereoLayerShapeCylinder::DrawShapeVisualization(const class FSceneView* View, class FPrimitiveDrawInterface* PDI)
