[Bugfix] Fix dangling OwnerLastRenderTimePtr after component reparents
- When a primitive component is reparented, SceneData.OwnerLastRenderTimePtr still pointed at the old owner's FActorLastRenderTime.
- If the old actor is destroyed, the render thread crashes writing through the dangling pointer during shadow gathering (SetLastRenderTime).
- Added UPrimitiveComponent::PostRename() to update OwnerLastRenderTimePtr when the outer changes and transfer the NumAlwaysVisibleComponents count between old and new owners.
- Deprecated raw NumAlwaysVisibleComponents field in favor of HasAlwaysVisibleComponents/Increment/Decrement accessors.

#Synced-CL 52109088

diff --git a/Engine/Source/Runtime/Engine/Classes/Components/PrimitiveComponent.h b/Engine/Source/Runtime/Engine/Classes/Components/PrimitiveComponent.h
--- a/Engine/Source/Runtime/Engine/Classes/Components/PrimitiveComponent.h
+++ b/Engine/Source/Runtime/Engine/Classes/Components/PrimitiveComponent.h
@@ -2505,6 +2505,7 @@ public:
 	ENGINE_API virtual void SendRenderTransform_Concurrent() override;
 	ENGINE_API virtual void OnRegister()  override;
 	ENGINE_API virtual void OnUnregister()  override;
+	ENGINE_API virtual void PostRename(UObject* OldOuter, const FName OldName) override;
 	ENGINE_API virtual void DestroyRenderState_Concurrent() override;
 	ENGINE_API virtual void OnCreatePhysicsState() override;
 	ENGINE_API virtual void OnDestroyPhysicsState() override;
diff --git a/Engine/Source/Runtime/Engine/Classes/GameFramework/Actor.h b/Engine/Source/Runtime/Engine/Classes/GameFramework/Actor.h
--- a/Engine/Source/Runtime/Engine/Classes/GameFramework/Actor.h
+++ b/Engine/Source/Runtime/Engine/Classes/GameFramework/Actor.h
@@ -152,8 +152,31 @@ struct FActorLastRenderTime
 		PRAGMA_ENABLE_DEPRECATION_WARNINGS
 	}
 
+	UE_DEPRECATED(5.8, "Use HasAlwaysVisibleComponents/IncrementAlwaysVisibleComponents/DecrementAlwaysVisibleComponents instead.")
 	std::atomic_int32_t NumAlwaysVisibleComponents = 0;
 
+	bool HasAlwaysVisibleComponents() const
+	{
+		PRAGMA_DISABLE_DEPRECATION_WARNINGS
+		return NumAlwaysVisibleComponents.load(std::memory_order_relaxed) > 0;
+		PRAGMA_ENABLE_DEPRECATION_WARNINGS
+	}
+
+	void IncrementAlwaysVisibleComponents()
+	{
+		PRAGMA_DISABLE_DEPRECATION_WARNINGS
+		NumAlwaysVisibleComponents.fetch_add(1, std::memory_order_relaxed);
+		PRAGMA_ENABLE_DEPRECATION_WARNINGS
+	}
+
+	/** Returns the previous count (before decrement). */
+	int32 DecrementAlwaysVisibleComponents()
+	{
+		PRAGMA_DISABLE_DEPRECATION_WARNINGS
+		return NumAlwaysVisibleComponents.fetch_sub(1, std::memory_order_relaxed);
+		PRAGMA_ENABLE_DEPRECATION_WARNINGS
+	}
+
 private:
 	static void Set(AActor* InActor, float LastRenderTime);
 	static FActorLastRenderTime* GetPtr(AActor* InActor);
diff --git a/Engine/Source/Runtime/Engine/Private/Actor.cpp b/Engine/Source/Runtime/Engine/Private/Actor.cpp
--- a/Engine/Source/Runtime/Engine/Private/Actor.cpp
+++ b/Engine/Source/Runtime/Engine/Private/Actor.cpp
@@ -2592,7 +2592,7 @@ bool AActor::WasRecentlyRendered(float Tolerance) const
 
 float AActor::GetLastRenderTime() const
 {
-	if (LastRenderTime.NumAlwaysVisibleComponents.load(std::memory_order_relaxed) > 0)
+	if (LastRenderTime.HasAlwaysVisibleComponents())
 	{
 		if (const UWorld* World = GetWorld())
 		{
diff --git a/Engine/Source/Runtime/Engine/Private/Components/PrimitiveComponent.cpp b/Engine/Source/Runtime/Engine/Private/Components/PrimitiveComponent.cpp
--- a/Engine/Source/Runtime/Engine/Private/Components/PrimitiveComponent.cpp
+++ b/Engine/Source/Runtime/Engine/Private/Components/PrimitiveComponent.cpp
@@ -5411,7 +5411,7 @@ void UPrimitiveComponent::AssignSceneProxy(FPrimitiveSceneProxy* InSceneProxy)
 
 		if (SceneData.bAlwaysVisible && SceneData.OwnerLastRenderTimePtr)
 		{
-			SceneData.OwnerLastRenderTimePtr->NumAlwaysVisibleComponents.fetch_add(1, std::memory_order_relaxed);
+			SceneData.OwnerLastRenderTimePtr->IncrementAlwaysVisibleComponents();
 		}
 
 #if WITH_EDITOR
@@ -5423,12 +5423,37 @@ void UPrimitiveComponent::AssignSceneProxy(FPrimitiveSceneProxy* InSceneProxy)
 	}
 }
 
+void UPrimitiveComponent::PostRename(UObject* OldOuter, const FName OldName)
+{
+	Super::PostRename(OldOuter, OldName);
+
+	if (OldOuter == GetOuter())
+	{
+		return;
+	}
+
+	FActorLastRenderTime* OldLastRenderTimePtr = SceneData.OwnerLastRenderTimePtr;
+	FActorLastRenderTime* NewLastRenderTimePtr = FActorLastRenderTime::GetPtr(GetOwner());
+
+	if (SceneData.bAlwaysVisible && OldLastRenderTimePtr)
+	{
+		OldLastRenderTimePtr->DecrementAlwaysVisibleComponents();
+	}
+
+	SceneData.OwnerLastRenderTimePtr = NewLastRenderTimePtr;
+
+	if (SceneData.bAlwaysVisible && NewLastRenderTimePtr)
+	{
+		NewLastRenderTimePtr->IncrementAlwaysVisibleComponents();
+	}
+}
+
 void UPrimitiveComponent::ReleaseSceneProxy()
 {
 	check(SceneProxy == SceneData.SceneProxy);
 	if (SceneData.bAlwaysVisible && SceneData.OwnerLastRenderTimePtr)
 	{
-		const uint32 NumRefs = SceneData.OwnerLastRenderTimePtr->NumAlwaysVisibleComponents.fetch_sub(1, std::memory_order_relaxed);
+		const int32 NumRefs = SceneData.OwnerLastRenderTimePtr->DecrementAlwaysVisibleComponents();
 		check(NumRefs > 0);
 	}
 	// Leave SceneData.OwnerLastRenderTimePtr valid. It might still be referenced on the rendering thread.
