SkyAtmosphere - added missing member init to nullptr.

#tests Editor level with SkyATmsopehre
#rnx

#Synced-CL 52139128

diff --git a/Engine/Source/Runtime/Engine/Private/Components/SkyAtmosphereComponent.cpp b/Engine/Source/Runtime/Engine/Private/Components/SkyAtmosphereComponent.cpp
--- a/Engine/Source/Runtime/Engine/Private/Components/SkyAtmosphereComponent.cpp
+++ b/Engine/Source/Runtime/Engine/Private/Components/SkyAtmosphereComponent.cpp
@@ -598,6 +598,7 @@ ASkyAtmosphere::ASkyAtmosphere(const FObjectInitializer& ObjectInitializer)
 
 FSkyAtmosphereSceneProxy::FSkyAtmosphereSceneProxy(const USkyAtmosphereComponent* InComponent)
 	: bStaticLightingBuilt(false)
+	, RenderSceneInfo(nullptr)
 	, AtmosphereSetup(*InComponent)
 	, bHoldout(InComponent->bHoldout > 0)
 	, bRenderInMainPass(InComponent->bRenderInMainPass > 0)
@@ -614,6 +615,7 @@ FSkyAtmosphereSceneProxy::FSkyAtmosphereSceneProxy(const USkyAtmosphereComponent
 
 FSkyAtmosphereSceneProxy::FSkyAtmosphereSceneProxy(const FSkyAtmosphereDynamicState& Ds)
 	: bStaticLightingBuilt(false)
+	, RenderSceneInfo(nullptr)
 	, AtmosphereSetup(Ds)
 	, bHoldout(Ds.bHoldout > 0)
 	, bRenderInMainPass(Ds.bRenderInMainPass > 0)
diff --git a/Engine/Source/Runtime/Renderer/Private/SkyAtmosphereRendering.cpp b/Engine/Source/Runtime/Renderer/Private/SkyAtmosphereRendering.cpp
--- a/Engine/Source/Runtime/Renderer/Private/SkyAtmosphereRendering.cpp
+++ b/Engine/Source/Runtime/Renderer/Private/SkyAtmosphereRendering.cpp
@@ -709,6 +709,7 @@ void FScene::RemoveSkyAtmosphere(FSkyAtmosphereSceneProxy* SkyAtmosphereScenePro
 			}
 
 			delete SkyAtmosphereSceneProxy->RenderSceneInfo;
+			SkyAtmosphereSceneProxy->RenderSceneInfo = nullptr;
 			Scene->SkyAtmosphereStack.RemoveSingle(SkyAtmosphereSceneProxy);
 
 			if (Scene->SkyAtmosphereStack.Num() > 0)
