Add lod.ForceLODDitherSpeed CVar to debug LOD dithering

Enables forcing meshes into dithered LOD transitions with looping alpha animation. Allows visual debugging of dither
patterns without camera movement. Works with r.ForceLOD or defaults to MinLOD when -1.

#tests editor + console

diff --git a/Engine/Source/Runtime/Engine/Private/SceneManagement.cpp b/Engine/Source/Runtime/Engine/Private/SceneManagement.cpp
--- a/Engine/Source/Runtime/Engine/Private/SceneManagement.cpp
+++ b/Engine/Source/Runtime/Engine/Private/SceneManagement.cpp
@@ -31,6 +31,13 @@ static TAutoConsoleVariable<float> CVarLODTemporalLag(
 	TEXT("This controls the the time lag for temporal LOD, in seconds."),
 	ECVF_Scalability | ECVF_Default);
 
+static float GDebugLODDitherSpeed = -1.0f;
+static FAutoConsoleVariableRef CVarForceLODDitherSpeed(
+	TEXT("lod.DebugLODDitherSpeed"),
+	GDebugLODDitherSpeed,
+	TEXT("Force a LOD Dither transition with the next LOD, and have it looping. Useful to debug LOD dithering without moving the camera."),
+	ECVF_Cheat | ECVF_RenderThreadSafe);
+
 bool AreCompressedTransformsSupported()
 {
 	return FDataDrivenShaderPlatformInfo::GetSupportSceneDataCompressedTransforms(GMaxRHIShaderPlatform);
@@ -83,6 +90,20 @@ bool ShouldCompileDistanceFieldAOShaders(EShaderPlatform ShaderPlatform)
 		&& DoesPlatformSupportDistanceFieldAO(ShaderPlatform);
 }
 
+float FTemporalLODState::GetTemporalLODTransition(float LastRenderTime) const
+{
+	if (GDebugLODDitherSpeed > 0.0f)
+	{
+		return FMath::Frac(LastRenderTime * GDebugLODDitherSpeed);
+	}
+
+	if (TemporalLODLag == 0.0)
+	{
+		return 0.0f; // no fade
+	}
+
+	return FMath::Clamp((LastRenderTime - TemporalLODLag - TemporalLODTime[0]) / (TemporalLODTime[1] - TemporalLODTime[0]), 0.0f, 1.0f);
+}
 
 void FTemporalLODState::UpdateTemporalLODTransition(const FSceneView& View, float LastRenderTime)
 {
@@ -1043,7 +1064,19 @@ FLODMask ComputeLODForMeshes(const TArray<class FStaticMeshBatchRelevance>& Stat
 			}
 		}
 		MinLOD = FMath::Max(MinLOD, (int32)CurFirstLODIdx);
-		LODToRender.SetLOD(FMath::Clamp(ForcedLODLevel, MinLOD, MaxLOD));
+
+		// force LOD dither transition with the next LOD
+		if (GDebugLODDitherSpeed > 0.0f && NumMeshes > 0 && bDitheredLODTransition && StaticMeshRelevances[0].bDitheredLODTransition && MaxLOD != 0)
+		{
+			int32 LOD0 = FMath::Clamp(ForcedLODLevel, MinLOD, MaxLOD);
+			int32 LOD1 = FMath::Clamp(LOD0 + 1, MinLOD, MaxLOD);
+			LODToRender.SetLODSample(LOD0, 0);
+			LODToRender.SetLODSample(LOD1, 1);
+		}
+		else
+		{
+			LODToRender.SetLOD(FMath::Clamp(ForcedLODLevel, MinLOD, MaxLOD));
+		}
 	}
 	else if (LODView.Family->EngineShowFlags.LOD && NumMeshes)
 	{
@@ -1052,6 +1085,8 @@ FLODMask ComputeLODForMeshes(const TArray<class FStaticMeshBatchRelevance>& Stat
 			for (int32 SampleIndex = 0; SampleIndex < 2; SampleIndex++)
 			{
 				int32 MinLODFound = INT_MAX;
+				int32 MaxLOD = 0;
+
 				bool bFoundLOD = false;
 				OutScreenRadiusSquared = ComputeTemporalLODBoundsScreenRadiusSquared(Origin, SphereRadius, LODView, SampleIndex);
 
@@ -1071,6 +1106,7 @@ FLODMask ComputeLODForMeshes(const TArray<class FStaticMeshBatchRelevance>& Stat
 						}
 
 						MinLODFound = FMath::Min<int32>(MinLODFound, Mesh.GetLODIndex());
+						MaxLOD = FMath::Max<int32>(MaxLOD, Mesh.GetLODIndex());
 					}
 				}
 				// If no LOD was found matching the screen size, use the lowest in the array instead of LOD 0, to handle non-zero MinLOD
@@ -1078,6 +1114,15 @@ FLODMask ComputeLODForMeshes(const TArray<class FStaticMeshBatchRelevance>& Stat
 				{
 					LODToRender.SetLODSample(MinLODFound, SampleIndex);
 				}
+
+				// force LOD dither transition with the next LOD
+				if (GDebugLODDitherSpeed > 0.0f && LODToRender.IsValid() && LODToRender.LODIndex0 == LODToRender.LODIndex1)
+				{
+					int32 LOD0 = LODToRender.LODIndex0;
+					int32 LOD1 = FMath::Min(LOD0 + 1, MaxLOD);
+					LODToRender.SetLODSample(LOD0, 0);
+					LODToRender.SetLODSample(LOD1, 1);				
+				}
 			}
 		}
 		else
diff --git a/Engine/Source/Runtime/Engine/Public/SceneManagement.h b/Engine/Source/Runtime/Engine/Public/SceneManagement.h
--- a/Engine/Source/Runtime/Engine/Public/SceneManagement.h
+++ b/Engine/Source/Runtime/Engine/Public/SceneManagement.h
@@ -98,18 +98,10 @@ struct FTemporalLODState
 	{
 
 	}
-	/** 
-	 * Returns the blend factor between the last two LOD samples
-	 */
-	float GetTemporalLODTransition(float LastRenderTime) const
-	{
-		if (TemporalLODLag == 0.0)
-		{
-			return 0.0f; // no fade
-		}
-		return FMath::Clamp((LastRenderTime - TemporalLODLag - TemporalLODTime[0]) / (TemporalLODTime[1] - TemporalLODTime[0]), 0.0f, 1.0f);
-	}
 
+	/* Returns the blend factor between the last two LOD samples */
+	ENGINE_API float GetTemporalLODTransition(float LastRenderTime) const;
+	
 	ENGINE_API void UpdateTemporalLODTransition(const FSceneView& View, float LastRenderTime);
 };
 
