Fix for MIs that have a missing parent material in cooked games
If for some reason, a MI is cooked but its parent isn't, and the MI had it own shader permutation, we get a mismatch between the uniform buffers provided by the proxy, and what the shader expects, and end up with a crash.
Setting bHasStaticPermutationResource to false ensures that we will be using the default material's shaders and are compatible with the proxy we will create.

#Synced-CL 52163112

diff --git a/Engine/Source/Runtime/Engine/Private/Materials/MaterialInstance.cpp b/Engine/Source/Runtime/Engine/Private/Materials/MaterialInstance.cpp
--- a/Engine/Source/Runtime/Engine/Private/Materials/MaterialInstance.cpp
+++ b/Engine/Source/Runtime/Engine/Private/Materials/MaterialInstance.cpp
@@ -2430,6 +2430,18 @@ void UMaterialInstance::UpdateOverridableBaseProperties()
 		MaxWorldPositionOffsetDisplacement = 0.0f;
 		bCompatibleWithLumenCardSharing = false;
 		UsageFlags = 0;
+
+		// Discard the MI's cooked shadermap - it was compiled against a parent that is no longer available, so its uniform buffer layout is incompatible with WorldGridMaterial's shaders.
+		// Even though we already set BlendMode to BLEND_Opaque, GetMaterialResource() would still serve up the MIC's own StaticPermutationMaterialResources (compiled against the missing parent) 
+		// while the render proxy uses WorldGridMaterial's uniform buffers.
+		// This produces at draw-time a 'FOpaqueBasePassUniformParameters vs TranslucentBasePass' check failure because the two were compiled against different parents. 
+		// Clearing the flag makes GetMaterialResource() delegate entirely to Parent (WorldGridMaterial), so shadermap and uniform buffers come from the same source and stay consistent.
+		if (bHasStaticPermutationResource)
+		{
+			UE_LOGF(LogMaterial, Error, "MaterialInstance [%ls] clearing bHasStaticPermutationResource to avoid shadermap/blendmode mismatch with default material.", *GetName());
+			bHasStaticPermutationResource = false;
+		}
+
 		return;
 	}
 
