Fix invariant violation between FConservativeShaderParameters and FCompilationShaderParameters for legacy MaterialAttributes static switches.

Legacy materials can have a static switches with MaterialAttributes on one branch and a non-MaterialAttributes value (e.g. scalar constant) on the other. The UI has been fixed to prevent creating such connections, but they can still exist in serialized data.

When the non-MaterialAttributes branch is active, the translator broadcasts the scalar to all material properties. IsMaterialPropertyUsed() then reports those properties as used (since the scalar doesn't match expected defaults), causing FCompilationShaderParameters to create shader layouts with parameters that FConservativeShaderParameters did not predict from CachedExpressionData. This violates the condition that the layouts generated using FConservativeShaderParameters should always be a superset of the ones created using FCompilationShaderParameters.
It's also worth noting that even before FCompilationShaderParameters was introduced, this issue was still happening but it was hidden, with the renderer requesting shaders which were not available at runtime and just skipping the draws.

Fix: In UMaterialExpressionStaticSwitch::Compile and UMaterialExpressionStaticSwitchParameter::Compile, detect when the active branch of a MaterialAttributes-typed switch is not itself MaterialAttributes. In that case, compile the default expression for the current material attribute instead of the mismatched value. This makes all properties report as "not used" (matching their defaults), keeping FCompilationShaderParameters within the bounds of FConservativeShaderParameters.

[RN] minor | Rendering Architecture

#Synced-CL 52167452

diff --git a/Engine/Source/Runtime/Engine/Private/Materials/MaterialExpressions.cpp b/Engine/Source/Runtime/Engine/Private/Materials/MaterialExpressions.cpp
--- a/Engine/Source/Runtime/Engine/Private/Materials/MaterialExpressions.cpp
+++ b/Engine/Source/Runtime/Engine/Private/Materials/MaterialExpressions.cpp
@@ -8635,6 +8635,28 @@ EMaterialValueType UMaterialExpressionStaticSwitchParameter::GetOutputValueType(
 	return GetCompositeOutput(OutputIndex);
 }
 
+// Helper: Legacy materials may have a static switch with MaterialAttributes on one branch and a
+// non-MaterialAttributes value (e.g. scalar constant) on the other. The UI now prevents creating
+// such connections, but they can still exist in serialized data. The non-material attributes is 
+// evaluated to a MA instance with that value set to each attribute blindly, causing some attributes 
+// to mistakenly be considered connected, which leads to invalid shader layouts.
+// Instead, when we detect this mismatch, compile the default expression for the current material attribute.
+static int32 CompileStaticSwitchBranch(FExpressionInput& ActiveInput, UMaterialExpression* SwitchExpression, int32 OutputIndex, FMaterialCompiler* Compiler)
+{
+	const FExpressionInput TracedActive = ActiveInput.GetTracedInput();
+	const bool bIsResultMaterialAttributes = SwitchExpression->IsResultMaterialAttributes(OutputIndex);
+	if (bIsResultMaterialAttributes && TracedActive.IsConnected())
+	{
+		const bool bIsTracedNotMA = !TracedActive.Expression->IsResultMaterialAttributes(TracedActive.OutputIndex);
+		const FGuid MaterialAttribute = Compiler->GetMaterialAttribute();
+		if (bIsTracedNotMA && MaterialAttribute.IsValid())
+		{
+			return FMaterialAttributeDefinitionMap::CompileDefaultExpression(Compiler, MaterialAttribute);
+		}
+	}
+	return ActiveInput.Compile(Compiler);
+}
+
 FExpressionInput* UMaterialExpressionStaticSwitchParameter::GetEffectiveInput(class FMaterialCompiler* Compiler)
 {
 	bool bSucceeded;
@@ -8672,11 +8694,11 @@ int32 UMaterialExpressionStaticSwitchParameter::Compile(class FMaterialCompiler*
 		{
 			if(DefaultValue)
 			{
-				return A.Compile(Compiler);
+				return CompileStaticSwitchBranch(A, this, OutputIndex, Compiler);
 			}
 			else
 			{
-				return B.Compile(Compiler);
+				return CompileStaticSwitchBranch(B, this, OutputIndex, Compiler);
 			}
 		}
 	}
@@ -8684,7 +8706,7 @@ int32 UMaterialExpressionStaticSwitchParameter::Compile(class FMaterialCompiler*
 	FExpressionInput* EffectiveInput = GetEffectiveInput(Compiler);
 	if (EffectiveInput)
 	{
-		return EffectiveInput->Compile(Compiler);
+		return CompileStaticSwitchBranch(*EffectiveInput, this, OutputIndex, Compiler);
 	}
 	return INDEX_NONE;
 }
@@ -8902,7 +8924,8 @@ int32 UMaterialExpressionStaticSwitch::Compile(class FMaterialCompiler* Compiler
 		}
 	}
 
-	return bValue ? A.Compile(Compiler) : B.Compile(Compiler);
+	FExpressionInput& ActiveInput = bValue ? A : B;
+	return CompileStaticSwitchBranch(ActiveInput, this, OutputIndex, Compiler);
 }
 
 void UMaterialExpressionStaticSwitch::GetCaption(TArray<FString>& OutCaptions) const
