Better detection of masked materials to enable skipping of shadow shder permutations when opacity mask is 1.0
We already have UMaterial::bCanMaskedBeAssumedOpaque which is computed pre-translation and  misses cases where the opacity mask is constant 1.0 after folding, or where material attribute outputs are used.
The HLSL translator and MaterialIR path now compute bUsesOpacityMask and this is taken into account when r.Material.UseShaderCompilationParameters=1
#Synced-CL 52165013

diff --git a/Engine/Source/Runtime/Engine/Private/Materials/HLSLMaterialTranslator.cpp b/Engine/Source/Runtime/Engine/Private/Materials/HLSLMaterialTranslator.cpp
--- a/Engine/Source/Runtime/Engine/Private/Materials/HLSLMaterialTranslator.cpp
+++ b/Engine/Source/Runtime/Engine/Private/Materials/HLSLMaterialTranslator.cpp
@@ -1719,6 +1719,8 @@ void FHLSLMaterialTranslator::TranslateMaterial()
 		Material->IsTessellationEnabled() &&
 		IsMaterialPropertyUsed(MP_Displacement, Chunk[MP_Displacement], FLinearColor(-1, 0, 0, 0), 1);
 
+	const bool bUsesOpacityMask = IsMaterialPropertyUsed(MP_OpacityMask, Chunk[MP_OpacityMask], FLinearColor(1, 0, 0, 0), 1);
+
 	static const FGuid FirstPersonInterpolationAlphaGuid = FMaterialAttributeDefinitionMap::GetCustomAttributeID(TEXT("FirstPersonInterpolationAlpha"));
 	const bool bHasFirstPersonOutput = (CustomOutputsBitmask & FMaterialAttributeDefinitionMap::GetBitmask(FirstPersonInterpolationAlphaGuid)) != 0;
 	bUsesFirstPersonInterpolation = bHasFirstPersonOutput;
@@ -1727,6 +1729,7 @@ void FHLSLMaterialTranslator::TranslateMaterial()
 	MaterialCompilationOutput.bUsesPixelDepthOffset = bUsesPixelDepthOffset;
 	MaterialCompilationOutput.bUsesDisplacement = bUsesDisplacement;
 	MaterialCompilationOutput.bUsesCustomizedUVs = Material->GetNumCustomizedUVs() > 0;
+	MaterialCompilationOutput.bUsesOpacityMask = bUsesOpacityMask;
 
 	static const FGuid TemporalResponsivenessGuid = FMaterialAttributeDefinitionMap::GetCustomAttributeID(TEXT("TemporalResponsiveness"));
 	static const FGuid MotionVectorWorldOffsetGuid = FMaterialAttributeDefinitionMap::GetCustomAttributeID(TEXT("MotionVectorWorldOffset"));
diff --git a/Engine/Source/Runtime/Engine/Private/Materials/MaterialIRModuleBuilder.cpp b/Engine/Source/Runtime/Engine/Private/Materials/MaterialIRModuleBuilder.cpp
--- a/Engine/Source/Runtime/Engine/Private/Materials/MaterialIRModuleBuilder.cpp
+++ b/Engine/Source/Runtime/Engine/Private/Materials/MaterialIRModuleBuilder.cpp
@@ -1238,9 +1238,10 @@ struct FMaterialIRModuleBuilderImpl
 		Module->CompilationOutput.NumVirtualTextureFeedbackRequests = NumVirtualTextureFeedbackRequests;
 		
 		// Simple on/off flags depending on whether a non-default value is set to some material property
-		CompilationOutput.bUsesWorldPositionOffset 	= MaterialPropertyHasNonZeroValue(MP_WorldPositionOffset);
-		CompilationOutput.bUsesPixelDepthOffset 	= MaterialPropertyHasNonZeroValue(MP_PixelDepthOffset);
-		CompilationOutput.bUsesAnisotropy 			= MaterialPropertyHasNonZeroValue(MP_Anisotropy);
+		CompilationOutput.bUsesWorldPositionOffset  = MaterialPropertyHasNonZeroValue(MP_WorldPositionOffset);
+		CompilationOutput.bUsesPixelDepthOffset     = MaterialPropertyHasNonZeroValue(MP_PixelDepthOffset);
+		CompilationOutput.bUsesAnisotropy           = MaterialPropertyHasNonZeroValue(MP_Anisotropy);
+		CompilationOutput.bUsesOpacityMask          = Module->IsMaterialPropertyUsed(MP_OpacityMask);
 		
 		// Check if displacement is actually used (disable if it was left at the default, invalid constant value of -1)
 		CompilationOutput.bUsesDisplacement 	 |= DoesPlatformSupportNanite(Module->GetShaderPlatform()) && Builder->MaterialInterface->IsTessellationEnabled() && MaterialPropertyHasNonZeroValue(MP_Displacement);
diff --git a/Engine/Source/Runtime/Engine/Private/Materials/MaterialShared.cpp b/Engine/Source/Runtime/Engine/Private/Materials/MaterialShared.cpp
--- a/Engine/Source/Runtime/Engine/Private/Materials/MaterialShared.cpp
+++ b/Engine/Source/Runtime/Engine/Private/Materials/MaterialShared.cpp
@@ -4864,13 +4864,18 @@ void FMaterial::GetPreshaderStats(uint32& TotalParameters, uint32& TotalOps) con
 bool FMaterial::WritesEveryPixel(bool bShadowPass) const
 {
 	const bool bVFTypeSupportsNullPixelShader = !IsUsedWithInstancedStaticMeshes();
-	return WritesEveryPixel(bShadowPass, bVFTypeSupportsNullPixelShader);
+	return WritesEveryPixel(bShadowPass, bVFTypeSupportsNullPixelShader, false);
 }
 
 bool FMaterial::WritesEveryPixel(bool bShadowPass, bool bVFTypeSupportsNullPixelShader) const
+{
+	return WritesEveryPixel(bShadowPass, bVFTypeSupportsNullPixelShader, false);
+}
+
+bool FMaterial::WritesEveryPixel(bool bShadowPass, bool bVFTypeSupportsNullPixelShader, bool bOpacityMaskIsOpaque) const
 {
 	bool bLocalStencilDitheredLOD = IsStencilForLODDitherEnabled(GetShaderPlatform());
-	return !IsMasked()
+	return (!IsMasked() || bOpacityMaskIsOpaque)
 		// Render dithered material as masked if a stencil prepass is not used (UE-50064, UE-49537)
 		&& !((bShadowPass || !bLocalStencilDitheredLOD) && IsDitheredLODTransition())
 		&& !IsWireframe()
@@ -5962,8 +5967,6 @@ FMaterialShaderParametersBase::FMaterialShaderParametersBase(const FMaterial* In
 	bIsThinSurface = InMaterial->IsThinSurface();
 	bIsDistorted = InMaterial->IsDistorted();
 	bShouldCastDynamicShadows = InMaterial->ShouldCastDynamicShadows();
-	bWritesEveryPixel = InMaterial->WritesEveryPixel(false);
-	bWritesEveryPixelShadowPass = InMaterial->WritesEveryPixel(true);
 
 	if (Substrate::IsSubstrateEnabled())
 	{
@@ -6067,9 +6070,11 @@ void FMaterialShaderParameters::InitializeShaderLayoutParamsFromMaterial(const F
 	bHasVertexInterpolator = InMaterial->HasVertexInterpolator();
 	bUsesTemporalResponsiveness = InMaterial->UsesTemporalResponsiveness();
 	bUsesMotionVectorWorldOffset = InMaterial->UsesMotionVectorWorldOffset();
+	bWritesEveryPixel = InMaterial->WritesEveryPixel(/*bShadowPass*/false);
+	bWritesEveryPixelShadowPass = InMaterial->WritesEveryPixel(/*bShadowPass*/true);
 }
 
-void FMaterialShaderParameters::InitializeShaderLayoutParamsFromCompOutput(const FMaterialCompilationOutput& InCompilationOutput)
+void FMaterialShaderParameters::InitializeShaderLayoutParamsFromCompOutput(const FMaterial* InMaterial, const FMaterialCompilationOutput& InCompilationOutput)
 {
 	bHasAnisotropyConnected = InCompilationOutput.bUsesAnisotropy;
 	bHasVertexPositionOffsetConnected = InCompilationOutput.bUsesWorldPositionOffset;
@@ -6080,6 +6085,14 @@ void FMaterialShaderParameters::InitializeShaderLayoutParamsFromCompOutput(const
 	bHasVertexInterpolator = InCompilationOutput.bUsesVertexInterpolator;
 	bUsesTemporalResponsiveness = InCompilationOutput.bUsesTemporalResponsiveness;
 	bUsesMotionVectorWorldOffset = InCompilationOutput.bUsesMotionVectorWorldOffset;
+	
+	// Need to consider whether opacity mask is used when calculating WritesEveryPixel().
+	const bool bVFTypeSupportsNullPixelShader = !InMaterial->IsUsedWithInstancedStaticMeshes();
+	bWritesEveryPixel = InMaterial->WritesEveryPixel(/*bShadowPass*/false, bVFTypeSupportsNullPixelShader, !InCompilationOutput.bUsesOpacityMask);
+	bWritesEveryPixelShadowPass = InMaterial->WritesEveryPixel(/*bShadowPass*/true, bVFTypeSupportsNullPixelShader, !InCompilationOutput.bUsesOpacityMask);
+	
+	// Also override blend mode if opacity mask is not used.
+	BlendMode = BlendMode == BLEND_Masked && !InCompilationOutput.bUsesOpacityMask ? BLEND_Opaque : BlendMode;
 }
 
 FCompilationShaderParameters::FCompilationShaderParameters(const FMaterial* InMaterial, const FMaterialCompilationOutput& InCompilationOutput)
@@ -6090,7 +6103,7 @@ PRAGMA_DISABLE_DEPRECATION_WARNINGS
 	static const int32 bUseShaderCompilationParameters = CVarMaterialUseShaderCompilationParameters.GetValueOnAnyThread();
 	if (bUseShaderCompilationParameters)
 	{
-		InitializeShaderLayoutParamsFromCompOutput(InCompilationOutput);
+		InitializeShaderLayoutParamsFromCompOutput(InMaterial, InCompilationOutput);
 	}
 	else
 	{
diff --git a/Engine/Source/Runtime/Engine/Public/MaterialShared.h b/Engine/Source/Runtime/Engine/Public/MaterialShared.h
--- a/Engine/Source/Runtime/Engine/Public/MaterialShared.h
+++ b/Engine/Source/Runtime/Engine/Public/MaterialShared.h
@@ -909,7 +909,8 @@ public:
 		bIsLightFunctionAtlasCompatible(false),
 		bUsesDisplacement(false),
 		bUsedWithNeuralNetworks(false),
-		bUsesCustomizedUVs(false)
+		bUsesCustomizedUVs(false),
+		bUsesOpacityMask(false)
 	{
 #if WITH_EDITOR
 		FMemory::Memzero(EstimatedLWCFuncUsagesVS);
@@ -1080,6 +1081,9 @@ public:
 	
 	/** Whether the material uses customized UVs. */
 	LAYOUT_BITFIELD(uint8, bUsesCustomizedUVs, 1);
+
+	/** True if the opacity mask isn't constant 1.0 (after constant folding). */
+	LAYOUT_BITFIELD(uint8, bUsesOpacityMask, 1);
 };
 
 struct FDebugShaderPipelineInfo
@@ -1261,7 +1265,7 @@ struct FMaterialShaderParameters : public FMaterialShaderParametersBase
 
 protected:
 	void InitializeShaderLayoutParamsFromMaterial(const FMaterial* InMaterial);
-	void InitializeShaderLayoutParamsFromCompOutput(const FMaterialCompilationOutput& InCompilationOutput);
+	void InitializeShaderLayoutParamsFromCompOutput(const FMaterial* InMaterial, const FMaterialCompilationOutput& InCompilationOutput);
 };
 
 // For Shader Compilation. Built based on Translation output.
@@ -2842,6 +2846,7 @@ public:
 	/* Helper function to look at both IsMasked and IsDitheredLODTransition to determine if it writes every pixel */
 	ENGINE_API bool WritesEveryPixel(bool bShadowPass = false) const;
 	ENGINE_API bool WritesEveryPixel(bool bShadowPass, bool bVFTypeSupportsNullPixelShader) const;
+	ENGINE_API bool WritesEveryPixel(bool bShadowPass, bool bVFTypeSupportsNullPixelShader, bool bOpacityMaskIsOpaque) const;
 
 	/** call during shader compilation jobs setup to fill additional settings that may be required by classes who inherit from this */
 	UE_DEPRECATED(5.7, "Call SetupExtraCompilationSettings without EShaderPlatform as its now cached inside the FMaterial")
