[SKM render] removing a unnecessary check, which was crashing when opening a SKM asset that has morph targets but does not use skin cache. On the first frame, it is expected that previous position and current position should use the same morph buffer such that we have zero motion vector on the first frame

#Synced-CL 52172018

diff --git a/Engine/Source/Runtime/Engine/Private/SkeletalRenderGPUSkin.cpp b/Engine/Source/Runtime/Engine/Private/SkeletalRenderGPUSkin.cpp
--- a/Engine/Source/Runtime/Engine/Private/SkeletalRenderGPUSkin.cpp
+++ b/Engine/Source/Runtime/Engine/Private/SkeletalRenderGPUSkin.cpp
@@ -334,7 +334,6 @@ const FMorphVertexBuffer& FMorphVertexBufferPool::GetMorphVertexBufferForReading
 		{
 			// this should only happen the first time updating, in which case the previous buffer hasn't been written into yet.
 			check(Index == 1);
-			check(MorphVertexBuffers[0].bHasBeenUpdated);
 			Index = 0;
 		}
 	}
