Remove singleton access during function call, as we're already in the singleton instance

#rnx
#Synced-CL 52054134

diff --git a/Engine/Source/Runtime/Engine/Private/AssetManager.cpp b/Engine/Source/Runtime/Engine/Private/AssetManager.cpp
--- a/Engine/Source/Runtime/Engine/Private/AssetManager.cpp
+++ b/Engine/Source/Runtime/Engine/Private/AssetManager.cpp
@@ -5050,7 +5050,7 @@ static FString GetInstigatorChainString(UE::Cook::ICookInfo* CookInfo, FName Pac
 bool UAssetManager::VerifyCanCookPackage(UE::Cook::ICookInfo* CookInfo, FName PackageName, bool bLogError) const
 {
 	bool bRetVal = true;
-	EPrimaryAssetCookRule CookRule = UAssetManager::Get().GetPackageCookRule(PackageName);
+	EPrimaryAssetCookRule CookRule = GetPackageCookRule(PackageName);
 	if (CookRule == EPrimaryAssetCookRule::NeverCook)
 	{
 		if (bLogError)
