Update UBodySetup::AddCollisionFrom to support missing cloth-specific shape types.

#Synced-CL 52021766

diff --git a/Engine/Source/Runtime/Engine/Private/PhysicsEngine/BodySetup.cpp b/Engine/Source/Runtime/Engine/Private/PhysicsEngine/BodySetup.cpp
--- a/Engine/Source/Runtime/Engine/Private/PhysicsEngine/BodySetup.cpp
+++ b/Engine/Source/Runtime/Engine/Private/PhysicsEngine/BodySetup.cpp
@@ -176,7 +176,8 @@ void UBodySetup::CopyBodyPropertiesFrom(const UBodySetup* FromSetup)
 
 void UBodySetup::AddCollisionFrom(const FKAggregateGeom& FromAggGeom)
 {
-	// Add shapes from static mesh
+	static_assert((int)EAggCollisionShape::Unknown == 9, "Did you forget to add a new AggCollisionShape type to this method?");
+
 	AggGeom.SphereElems.Append(FromAggGeom.SphereElems);
 	AggGeom.BoxElems.Append(FromAggGeom.BoxElems);
 	AggGeom.SphylElems.Append(FromAggGeom.SphylElems);
@@ -192,8 +193,11 @@ void UBodySetup::AddCollisionFrom(const FKAggregateGeom& FromAggGeom)
 		ConvexElem.ResetChaosConvexMesh();
 	}
 
+	AggGeom.TaperedCapsuleElems.Append(FromAggGeom.TaperedCapsuleElems);
 	AggGeom.LevelSetElems.Append(FromAggGeom.LevelSetElems);
+	AggGeom.SkinnedLevelSetElems.Append(FromAggGeom.SkinnedLevelSetElems);
 	AggGeom.MLLevelSetElems.Append(FromAggGeom.MLLevelSetElems);
+	AggGeom.SkinnedTriangleMeshElems.Append(FromAggGeom.SkinnedTriangleMeshElems);
 }
 
 namespace
@@ -212,6 +216,7 @@ namespace
 
 bool UBodySetup::AddCollisionElemFrom(const FKAggregateGeom& FromAggGeom, const EAggCollisionShape::Type ShapeType, const int32 ElemIndex)
 {
+	static_assert((int)EAggCollisionShape::Unknown == 9, "Did you forget to add a new AggCollisionShape type to this method?");
 	switch (ShapeType)
 	{
 	case EAggCollisionShape::Sphere:
@@ -231,6 +236,14 @@ bool UBodySetup::AddCollisionElemFrom(const FKAggregateGeom& FromAggGeom, const
 		return AddCollisionElemFrom_Helper(FromAggGeom.TaperedCapsuleElems, ElemIndex, AggGeom.TaperedCapsuleElems);
 	case EAggCollisionShape::LevelSet:
 		return AddCollisionElemFrom_Helper(FromAggGeom.LevelSetElems, ElemIndex, AggGeom.LevelSetElems);
+	case EAggCollisionShape::SkinnedLevelSet:
+		return AddCollisionElemFrom_Helper(FromAggGeom.SkinnedLevelSetElems, ElemIndex, AggGeom.SkinnedLevelSetElems);
+	case EAggCollisionShape::MLLevelSet:
+		return AddCollisionElemFrom_Helper(FromAggGeom.MLLevelSetElems, ElemIndex, AggGeom.MLLevelSetElems);
+	case EAggCollisionShape::SkinnedTriangleMesh:
+		return AddCollisionElemFrom_Helper(FromAggGeom.SkinnedTriangleMeshElems, ElemIndex, AggGeom.SkinnedTriangleMeshElems);
+	default:
+		checkNoEntry();
 	}
 	return false;
 }
