re-submit of 52169919
If for some reason, a MI is cooked but its parent isn't, and the MI had it own shader permutation, we get a mismatch between the uniform buffers provided by the proxy, and what the shader expects, and end up with a crash.
Setting bHasStaticPermutationResource to false ensures that we will be using the default material's shaders and are compatible with the proxy we will create.
Compared to 52169919,  change where we clear bHasStaticPermutationResource: it's actually possible to create a MaterialInstance without a parent. We now do it in InitStaticPermutation where we're already setting up bHasStaticPermutationResource for WITH_EDITORONLY_DATA
If we had no parent at cook time, bHasStaticPermutationResource should be set to false anyway, so the new branch shouldn't do anything.
But If bHasStaticPermutationResource is set to true, and we don't have a Parent in cooked, this means our parent is gone and we will fallback to the default material's shaders instead

#Synced-CL 52163112

#Synced-CL 52163112

diff --git a/Engine/Source/Runtime/Engine/Private/Materials/MaterialInstance.cpp b/Engine/Source/Runtime/Engine/Private/Materials/MaterialInstance.cpp
--- a/Engine/Source/Runtime/Engine/Private/Materials/MaterialInstance.cpp
+++ b/Engine/Source/Runtime/Engine/Private/Materials/MaterialInstance.cpp
@@ -2383,6 +2383,17 @@ void UMaterialInstance::InitStaticPermutation(EMaterialShaderPrecompileMode Prec
 		bHasStaticPermutationResource = Parent && (HasStaticParameters() || HasOverridenBaseProperties());
 		ValidateStaticPermutationAllowed();
 	}
+#else
+	// Discard the MI's cooked shadermap - it was compiled against a parent that is no longer available, so its uniform buffer layout is incompatible with WorldGridMaterial's shaders.
+	// Even though we already set BlendMode to BLEND_Opaque, GetMaterialResource() would still serve up the MIC's own StaticPermutationMaterialResources (compiled against the missing parent) 
+	// while the render proxy uses WorldGridMaterial's uniform buffers.
+	// This produces at draw-time a 'FOpaqueBasePassUniformParameters vs TranslucentBasePass' check failure because the two were compiled against different parents. 
+	// Clearing the flag makes GetMaterialResource() delegate entirely to Parent (WorldGridMaterial), so shadermap and uniform buffers come from the same source and stay consistent.
+	if (!Parent && bHasStaticPermutationResource)
+	{
+		UE_LOGF(LogMaterial, Error, "MaterialInstance [%ls] clearing bHasStaticPermutationResource to avoid shadermap/blendmode mismatch with default material.", *GetName());
+		bHasStaticPermutationResource = false;
+	}
 #endif // WITH_EDITORONLY_DATA
 
 	FMaterialResourceDeferredDeletionArray ResourcesToFree;
