Temporary revert of removal of shaders for masked materials with opacity mask 1.
We need to take into account the that materials can also use clip in custom hlsl, and mark themselves as masked as a way to communicate that to the shader map.
#rnx
#Synced-CL 52204248

diff --git a/Engine/Source/Runtime/Engine/Private/Materials/MaterialShared.cpp b/Engine/Source/Runtime/Engine/Private/Materials/MaterialShared.cpp
--- a/Engine/Source/Runtime/Engine/Private/Materials/MaterialShared.cpp
+++ b/Engine/Source/Runtime/Engine/Private/Materials/MaterialShared.cpp
@@ -6088,11 +6088,13 @@ void FMaterialShaderParameters::InitializeShaderLayoutParamsFromCompOutput(const
 	
 	// Need to consider whether opacity mask is used when calculating WritesEveryPixel().
 	const bool bVFTypeSupportsNullPixelShader = !InMaterial->IsUsedWithInstancedStaticMeshes();
-	bWritesEveryPixel = InMaterial->WritesEveryPixel(/*bShadowPass*/false, bVFTypeSupportsNullPixelShader, !InCompilationOutput.bUsesOpacityMask);
-	bWritesEveryPixelShadowPass = InMaterial->WritesEveryPixel(/*bShadowPass*/true, bVFTypeSupportsNullPixelShader, !InCompilationOutput.bUsesOpacityMask);
+	// Temporarily assume bUsesOpacityMask is true while we fix issues that come from not detecting the use of clip() in the shader.
+	const bool bUsesOpacityMask = true; // InCompilationOutput.bUsesOpacityMask
+	bWritesEveryPixel = InMaterial->WritesEveryPixel(/*bShadowPass*/false, bVFTypeSupportsNullPixelShader, !bUsesOpacityMask);
+	bWritesEveryPixelShadowPass = InMaterial->WritesEveryPixel(/*bShadowPass*/true, bVFTypeSupportsNullPixelShader, !bUsesOpacityMask);
 	
 	// Also override blend mode if opacity mask is not used.
-	BlendMode = BlendMode == BLEND_Masked && !InCompilationOutput.bUsesOpacityMask ? BLEND_Opaque : BlendMode;
+	BlendMode = BlendMode == BLEND_Masked && !bUsesOpacityMask ? BLEND_Opaque : BlendMode;
 }
 
 FCompilationShaderParameters::FCompilationShaderParameters(const FMaterial* InMaterial, const FMaterialCompilationOutput& InCompilationOutput)
