Fix compilation error -Wbitfield-enum-conversion under Clang 21. enum EDerivativeAxis is used as a single bit bitfield, however, 1 bit is insufficient to hold a 0-1 enum as a signed value. Specify an unsigned type for the EDerivativeAxis enum to resolve.

#Synced-CL 52149918

diff --git a/Engine/Source/Runtime/Engine/Public/Materials/MaterialIRCommon.h b/Engine/Source/Runtime/Engine/Public/Materials/MaterialIRCommon.h
--- a/Engine/Source/Runtime/Engine/Public/Materials/MaterialIRCommon.h
+++ b/Engine/Source/Runtime/Engine/Public/Materials/MaterialIRCommon.h
@@ -167,7 +167,7 @@ ENUM_CLASS_FLAGS(EGraphProperties);
 // This enum is used to indicate whether a partial derivative should be taken
 // along the X or Y screen-space direction, corresponding to HLSL's ddx and ddy
 // functions.
-enum class EDerivativeAxis
+enum class EDerivativeAxis : unsigned char
 {
 	X, Y,
 };
