Fix crash in touch input: initialize SourceAction to nullptr in FDelegateDispatchDetails touch and gesture constructors

#rnx

#Synced-CL 0

diff --git a/Engine/Source/Runtime/Engine/Private/UserInterface/PlayerInput.cpp b/Engine/Source/Runtime/Engine/Private/UserInterface/PlayerInput.cpp
--- a/Engine/Source/Runtime/Engine/Private/UserInterface/PlayerInput.cpp
+++ b/Engine/Source/Runtime/Engine/Private/UserInterface/PlayerInput.cpp
@@ -110,6 +110,7 @@ struct FDelegateDispatchDetails
 	FDelegateDispatchDetails(const uint32 InEventIndex, const uint32 InFoundIndex, FInputTouchUnifiedDelegate InDelegate, const FVector InLocation, const ETouchIndex::Type InFingerIndex)
 		: EventIndex(InEventIndex)
 		, FoundIndex(InFoundIndex)
+		, SourceAction(nullptr)
 		, TouchDelegate(MoveTemp(InDelegate))
 		, TouchLocation(InLocation)
 		, FingerIndex(InFingerIndex)
@@ -118,6 +119,7 @@ struct FDelegateDispatchDetails
 	FDelegateDispatchDetails(const uint32 InEventIndex, const uint32 InFoundIndex, FInputGestureUnifiedDelegate InDelegate, const float InValue)
 		: EventIndex(InEventIndex)
 		, FoundIndex(InFoundIndex)
+		, SourceAction(nullptr)
 		, GestureDelegate(MoveTemp(InDelegate))
 		, GestureValue(InValue)
 	{}
