Fix GPU page fault realted to Residency: clear Cascade particle state textures at creation to establish residency

#Synced-CL 52149918

diff --git a/Engine/Source/Runtime/Engine/Private/Particles/ParticleGpuSimulation.cpp b/Engine/Source/Runtime/Engine/Private/Particles/ParticleGpuSimulation.cpp
--- a/Engine/Source/Runtime/Engine/Private/Particles/ParticleGpuSimulation.cpp
+++ b/Engine/Source/Runtime/Engine/Private/Particles/ParticleGpuSimulation.cpp
@@ -380,7 +380,14 @@ public:
 		RHICmdList.BindDebugLabelName(PositionTextureRHI, ParticlePositionTextureName);
 		RHICmdList.BindDebugLabelName(VelocityTextureRHI, ParticleVelocityTextureName);
 
-		bTexturesCleared = false;
+		// Clear textures at creation time to ensure GPU residency is established immediately.
+		ClearRenderTarget(RHICmdList.GetAsImmediate(), PositionTextureRHI);
+		ClearRenderTarget(RHICmdList.GetAsImmediate(), VelocityTextureRHI);
+		RHICmdList.GetAsImmediate().Transition({
+			FRHITransitionInfo(PositionTextureRHI, ERHIAccess::RTV, ERHIAccess::SRVMask),
+			FRHITransitionInfo(VelocityTextureRHI, ERHIAccess::RTV, ERHIAccess::SRVMask)
+		});
+		bTexturesCleared = true;
 	}
 
 	/**
@@ -671,7 +678,7 @@ public:
 	void Update(FRHICommandListImmediate& RHICmdList)
 	{
 		if(TileAllocator.SupportResizingTiles())
-		{ 
+		{
 			check(IsInRenderingThread());
 			if (TileAllocator.bDirtyAlloc)
 			{
