Promote SpringMath out of experimental

Move SpringMath.h from Engine to Core (zero Engine dependencies - header-only, uses only Core types).
Remove all UE_EXPERIMENTAL annotations and PRAGMA_DISABLE/ENABLE_EXPERIMENTAL_WARNINGS pairs.
Remove Experimental specifiers from BlueprintSpringMathLibrary and RigVM SpringMath structs.
Update all callers to new include path (Math/SpringMath.h), add UE_DEPRECATED_HEADER forwarding header at old location.
Fix coding standard issues (Allman braces, fabs->FMath::Abs, indentation).
Fix incomplete [REDACTED_HANDLE] DeltaTime documentation.
Add explicit SmoothingTime=0 snap handling to all CriticalSpringDamper and CriticalDecay variants.
Add comprehensive LLT test coverage (103 tests, 3841 assertions) in MathCoreTests.
Migrate and existing SpringMath tests found in SpringDamperTests.cpp from EngineTest/AnimationTests
Parametric sweep tests for character stopping/starting/max acceleration are now in the MathCoreTests LLT.

#Synced-CL 52208314

diff --git a/Engine/Plugins/Animation/AnimationWarping/Source/Runtime/Private/BoneControllers/AnimNode_Steering.cpp b/Engine/Plugins/Animation/AnimationWarping/Source/Runtime/Private/BoneControllers/AnimNode_Steering.cpp
--- a/Engine/Plugins/Animation/AnimationWarping/Source/Runtime/Private/BoneControllers/AnimNode_Steering.cpp
+++ b/Engine/Plugins/Animation/AnimationWarping/Source/Runtime/Private/BoneControllers/AnimNode_Steering.cpp
@@ -5,7 +5,7 @@
 #include "Animation/AnimNodeFunctionRef.h"
 #include "Animation/AnimRootMotionProvider.h"
 #include "Animation/AnimSequenceHelpers.h"
-#include "Animation/SpringMath.h"
+#include "Math/SpringMath.h"
 #include "HAL/IConsoleManager.h"
 #include "Animation/AnimTrace.h"
 #include "BoneControllers/AnimNode_OffsetRootBone.h"
diff --git a/Engine/Plugins/Experimental/Animation/AnimGen/Source/AnimGen/Private/AnimGenBehavior.cpp b/Engine/Plugins/Experimental/Animation/AnimGen/Source/AnimGen/Private/AnimGenBehavior.cpp
--- a/Engine/Plugins/Experimental/Animation/AnimGen/Source/AnimGen/Private/AnimGenBehavior.cpp
+++ b/Engine/Plugins/Experimental/Animation/AnimGen/Source/AnimGen/Private/AnimGenBehavior.cpp
@@ -8,7 +8,7 @@
 
 #include "AnimGenLog.h"
 #include "InputCoreTypes.h"
-#include "Animation/SpringMath.h"
+#include "Math/SpringMath.h"
 #include "DrawDebugLibrary.h"
 
 #include "GameFramework/Actor.h"
diff --git a/Engine/Plugins/Runtime/RigVM/Source/RigVM/Public/RigVMFunctions/Math/RigVMFunction_SpringMath.h b/Engine/Plugins/Runtime/RigVM/Source/RigVM/Public/RigVMFunctions/Math/RigVMFunction_SpringMath.h
--- a/Engine/Plugins/Runtime/RigVM/Source/RigVM/Public/RigVMFunctions/Math/RigVMFunction_SpringMath.h
+++ b/Engine/Plugins/Runtime/RigVM/Source/RigVM/Public/RigVMFunctions/Math/RigVMFunction_SpringMath.h
@@ -5,11 +5,10 @@
 #include "RigVMFunction_MathBase.h"
 #include "RigVMFunction_SpringMath.generated.h"
 
-UE_EXPERIMENTAL(5.7, "SpringMath API is experimental")
 /**
  * Damps a float value using exponential decay damping
  */
-USTRUCT(meta=(DisplayName="Damp (Float)", Category="Math|Damp|Experimental", TemplateName="Damp"))
+USTRUCT(meta=(DisplayName="Damp (Float)", Category="Math|Damp", TemplateName="Damp"))
 struct FRigVMFunction_DampFloat : public FRigVMFunction_MathBase
 {
 	GENERATED_BODY()
@@ -34,11 +33,10 @@ struct FRigVMFunction_DampFloat : public FRigVMFunction_MathBase
 	float Result = 0.0f;
 };
 
-UE_EXPERIMENTAL(5.7, "SpringMath API is experimental")
 /**
  * Damps a vector value using exponential decay damping
  */
-USTRUCT(meta=(DisplayName="Damp (Vector)", Category="Math|Damp|Experimental", TemplateName="Damp"))
+USTRUCT(meta=(DisplayName="Damp (Vector)", Category="Math|Damp", TemplateName="Damp"))
 struct FRigVMFunction_DampVector : public FRigVMFunction_MathBase
 {
 	GENERATED_BODY()
@@ -63,11 +61,10 @@ struct FRigVMFunction_DampVector : public FRigVMFunction_MathBase
 	FVector Result = FVector::ZeroVector;
 };
 
-UE_EXPERIMENTAL(5.7, "SpringMath API is experimental") 
 /**
  * Damps a quaternion value using exponential decay damping
  */
-USTRUCT(meta=(DisplayName="Damp (Quaternion)", Category="Math|Damp|Experimental", TemplateName="Damp"))
+USTRUCT(meta=(DisplayName="Damp (Quaternion)", Category="Math|Damp", TemplateName="Damp"))
 struct FRigVMFunction_DampQuaternion : public FRigVMFunction_MathBase
 {
 	GENERATED_BODY()
@@ -92,11 +89,10 @@ struct FRigVMFunction_DampQuaternion : public FRigVMFunction_MathBase
 	FQuat Result = FQuat::Identity;
 };
 
-UE_EXPERIMENTAL(5.7, "SpringMath API is experimental")
 /**
  * Damps a float using a spring damper
  */
-USTRUCT(meta=(DisplayName="Critical Spring Damp (Float)", Category="Math|Damp|Experimental", TemplateName="Critical Spring Damp"))
+USTRUCT(meta=(DisplayName="Critical Spring Damp (Float)", Category="Math|Damp", TemplateName="Critical Spring Damp"))
 struct FRigVMFunction_SpringDampFloat : public FRigVMFunction_MathMutableBase
 {
 	GENERATED_BODY()
@@ -121,11 +117,10 @@ struct FRigVMFunction_SpringDampFloat : public FRigVMFunction_MathMutableBase
 	float SmoothingTime = 0.2f;
 };
 
-UE_EXPERIMENTAL(5.7, "SpringMath API is experimental")
 /**
  * Damps a vector using a spring damper
  */
-USTRUCT(meta=(DisplayName="Critical Spring Damp (Vector)", Category="Math|Damp|Experimental", TemplateName="Critical Spring Damp"))
+USTRUCT(meta=(DisplayName="Critical Spring Damp (Vector)", Category="Math|Damp", TemplateName="Critical Spring Damp"))
 struct FRigVMFunction_SpringDampVector : public FRigVMFunction_MathMutableBase
 {
 	GENERATED_BODY()
@@ -150,11 +145,10 @@ struct FRigVMFunction_SpringDampVector : public FRigVMFunction_MathMutableBase
 	float SmoothingTime = 0.2f;
 };
 
-UE_EXPERIMENTAL(5.7, "SpringMath API is experimental")
 /**
  * Damps a quaternion using a spring damper
  */
-USTRUCT(meta=(DisplayName="Critical Spring Damp (Quat)", Category="Math|Damp|Experimental", TemplateName="Critical Spring Damp"))
+USTRUCT(meta=(DisplayName="Critical Spring Damp (Quat)", Category="Math|Damp", TemplateName="Critical Spring Damp"))
 struct FRigVMFunction_SpringDampQuat : public FRigVMFunction_MathMutableBase
 {
 	GENERATED_BODY()
diff --git a/Engine/Source/Runtime/Engine/Classes/Kismet/BlueprintSpringMathLibrary.h b/Engine/Source/Runtime/Engine/Classes/Kismet/BlueprintSpringMathLibrary.h
--- a/Engine/Source/Runtime/Engine/Classes/Kismet/BlueprintSpringMathLibrary.h
+++ b/Engine/Source/Runtime/Engine/Classes/Kismet/BlueprintSpringMathLibrary.h
@@ -7,7 +7,7 @@
 #include "BlueprintSpringMathLibrary.generated.h"
 
 /** A simple struct that stores velocities for an FTransform. */
-USTRUCT(BlueprintType, Experimental)
+USTRUCT(BlueprintType)
 struct FSpringTransformVelocity
 {
 	GENERATED_BODY()
@@ -25,7 +25,7 @@ struct FSpringTransformVelocity
 	FVector ScalarVelocity = FVector::ZeroVector;
 };
 
-UCLASS(MinimalAPI, Experimental, meta=(BlueprintThreadSafe))
+UCLASS(MinimalAPI, meta=(BlueprintThreadSafe))
 class UBlueprintSpringMathLibrary : public UBlueprintFunctionLibrary
 {
 	GENERATED_BODY()
@@ -33,7 +33,6 @@ class UBlueprintSpringMathLibrary : public UBlueprintFunctionLibrary
 public:
 	// Simple Spring //
 
-	UE_EXPERIMENTAL(5.7, "SpringMath is experimental")
 	/** Interpolates the value InOutX towards TargetX with the motion of a critically damped spring. The velocity of X is stored in InOutV.
 	 * 
 	 * @param InOutX The value to be damped
@@ -42,10 +41,9 @@ public:
 	 * @param SmoothingTime The smoothing time to use for the spring. Longer times lead to more damped behaviour. A time of 0 will snap InOutX to TargetX
 	 * @param DeltaTime Timestep in seconds
 	 */
-	UFUNCTION(BlueprintCallable, meta=(Category="Math|Spring|Experimental", AutoCreateRefTerm="TargetX"))
+	UFUNCTION(BlueprintCallable, meta=(Category="Math|Spring", AutoCreateRefTerm="TargetX"))
 	static ENGINE_API void CriticalSpringDampFloat(UPARAM(ref) float& InOutX, UPARAM(ref) float& InOutV, const float& TargetX, float DeltaTime, float SmoothingTime = 0.2f);
 
-	UE_EXPERIMENTAL(5.7, "SpringMath is experimental")
 	/** Interpolates the angle InOutAngle towards TargetAngle with the motion of a critically damped spring. The velocity of InOutAngle is stored in InOutAngularVelocity in deg/s.
 	 *
 	 * @param InOutAngle The value to be damped in degrees
@@ -54,10 +52,9 @@ public:
 	 * @param SmoothingTime The smoothing time to use for the spring. Longer times lead to more damped behaviour. A time of 0 will snap InOutAngle to TargetAngle
 	 * @param DeltaTime Timestep in seconds
 	 */
-	UFUNCTION(BlueprintCallable, meta = (Category = "Math|Spring|Experimental", AutoCreateRefTerm = "TargetAngle"))
+	UFUNCTION(BlueprintCallable, meta = (Category = "Math|Spring", AutoCreateRefTerm = "TargetAngle"))
 	static ENGINE_API void CriticalSpringDampAngle(UPARAM(ref) float& InOutAngle, UPARAM(ref) float& InOutAngularVelocity, const float& TargetAngle, float DeltaTime, float SmoothingTime = 0.2f);
 
-	UE_EXPERIMENTAL(5.7, "SpringMath is experimental")
 	/** Interpolates the value InOutX towards TargetX with the motion of a critically damped spring. The velocity of X is stored in InOutV.
 	 * 
 	 * @param InOutX The value to be damped
@@ -66,10 +63,9 @@ public:
 	 * @param SmoothingTime The smoothing time to use for the spring. Longer times lead to more damped behaviour. A time of 0 will snap InOutX to TargetX
 	 * @param DeltaTime Timestep in seconds
 	 */
-	UFUNCTION(BlueprintCallable, meta=(Category="Math|Spring|Experimental", AutoCreateRefTerm = "TargetX"))
+	UFUNCTION(BlueprintCallable, meta=(Category="Math|Spring", AutoCreateRefTerm = "TargetX"))
 	static ENGINE_API void CriticalSpringDampVector(UPARAM(ref) FVector& InOutX, UPARAM(ref) FVector& InOutV, const FVector& TargetX, float DeltaTime, float SmoothingTime = 0.2f);
 
-	UE_EXPERIMENTAL(5.7, "SpringMath is experimental")
 	/** Interpolates the value InOutX towards TargetX with the motion of a critically damped spring. The velocity of X is stored in InOutV.
 	 * 
 	 * @param InOutX The value to be damped
@@ -78,10 +74,9 @@ public:
 	 * @param SmoothingTime The smoothing time to use for the spring. Longer times lead to more damped behaviour. A time of 0 will snap InOutX to TargetX
 	 * @param DeltaTime Timestep in seconds
 	 */
-	UFUNCTION(BlueprintCallable, meta=(Category="Math|Spring|Experimental", AutoCreateRefTerm = "TargetX"))
+	UFUNCTION(BlueprintCallable, meta=(Category="Math|Spring", AutoCreateRefTerm = "TargetX"))
 	static ENGINE_API void CriticalSpringDampVector2D(UPARAM(ref) FVector2D& InOutX, UPARAM(ref) FVector2D& InOutV, const FVector2D& TargetX, float DeltaTime, float SmoothingTime = 0.2f);
 
-	UE_EXPERIMENTAL(5.7, "SpringMath is experimental")
 	/** Interpolates the value InOutRotation towards TargetRotation with the motion of a critically damped spring. The velocity of InOutRotation is stored in InOutAngularVelocity in deg/s.
 	 *
 	 * @param InOutRotation The value to be damped
@@ -90,10 +85,9 @@ public:
 	 * @param SmoothingTime The smoothing time to use for the spring. Longer times lead to more damped behaviour. A time of 0 will snap InOutRotation to TargetRotation
 	 * @param DeltaTime Timestep in seconds
 	 */
-	UFUNCTION(BlueprintCallable, meta=(Category="Math|Spring|Experimental", AutoCreateRefTerm = "TargetRotation"))
+	UFUNCTION(BlueprintCallable, meta=(Category="Math|Spring", AutoCreateRefTerm = "TargetRotation"))
 	static ENGINE_API void CriticalSpringDampQuat(UPARAM(ref) FQuat& InOutRotation, UPARAM(ref) FVector& InOutAngularVelocity, const FQuat& TargetRotation, float DeltaTime, float SmoothingTime = 0.2f);
 
-	UE_EXPERIMENTAL(5.7, "SpringMath is experimental")
 	/** Interpolates the value InOutRotation towards TargetRotation with the motion of a critically damped spring. The velocity of InOutRotation is stored in InOutAngularVelocity in deg/s.
 	 *
 	 * @param InOutRotation The value to be damped
@@ -102,10 +96,9 @@ public:
 	 * @param SmoothingTime The smoothing time to use for the spring. Longer times lead to more damped behaviour. A time of 0 will snap InOutRotation to TargetRotation
 	 * @param DeltaTime Timestep in seconds
 	 */
-	UFUNCTION(BlueprintCallable, meta=(Category="Math|Spring|Experimental", AutoCreateRefTerm = "TargetRotation"))
+	UFUNCTION(BlueprintCallable, meta=(Category="Math|Spring", AutoCreateRefTerm = "TargetRotation"))
 	static ENGINE_API void CriticalSpringDampRotator(UPARAM(ref) FRotator& InOutRotation, UPARAM(ref) FVector& InOutAngularVelocity, const FRotator& TargetRotation, float DeltaTime, float SmoothingTime = 0.2f);
 
-	UE_EXPERIMENTAL(5.8, "SpringMath is experimental")
 	/** Interpolates the value InOutScale towards TargetScale with the motion of a critically damped spring.
 	 *
 	 * @param InOutScale The value to be damped
@@ -114,10 +107,9 @@ public:
 	 * @param SmoothingTime The smoothing time to use for the spring. Longer times lead to more damped behaviour. A time of 0 will snap InOutScale to TargetScale
 	 * @param DeltaTime Timestep in seconds
 	 */
-	UFUNCTION(BlueprintCallable, meta = (Category = "Math|Spring|Experimental", AutoCreateRefTerm = "TargetScale"))
+	UFUNCTION(BlueprintCallable, meta = (Category = "Math|Spring", AutoCreateRefTerm = "TargetScale"))
 	static ENGINE_API void CriticalSpringDampScale(UPARAM(ref) FVector& InOutScale, UPARAM(ref) FVector& InOutScalarVelocity, const FVector& TargetScale, float DeltaTime, float SmoothingTime = 0.2f);
 
-	UE_EXPERIMENTAL(5.8, "SpringMath is experimental")
 	/** Interpolates the value InOutTransform towards TargetTransform with the motion of a critically damped spring.
 	 *
 	 * @param InOutTransform The value to be damped
@@ -126,13 +118,12 @@ public:
 	 * @param SmoothingTime The smoothing time to use for the spring. Longer times lead to more damped behaviour. A time of 0 will snap InOutTransform to TargetTransform
 	 * @param DeltaTime Timestep in seconds
 	 */
-	UFUNCTION(BlueprintCallable, meta = (Category = "Math|Spring|Experimental", AutoCreateRefTerm = "TargetTransform"))
+	UFUNCTION(BlueprintCallable, meta = (Category = "Math|Spring", AutoCreateRefTerm = "TargetTransform"))
 	static ENGINE_API void CriticalSpringDampTransform(UPARAM(ref) FTransform& InOutTransform, UPARAM(ref) FSpringTransformVelocity& InOutVelocity, const FTransform& TargetTransform, float DeltaTime, float SmoothingTime = 0.2f);
 
 
 	// Double Spring //
 
-	UE_EXPERIMENTAL(5.7, "SpringMath is experimental")
 	/** Interpolates the value InOutX towards TargetX using a critically damped spring via an intermediate state - producing a more S-shaped curve. The velocity of X is stored in InOutV.
 	 *
 	 * @param InOutX The value to be damped
@@ -143,10 +134,9 @@ public:
 	 * @param SmoothingTime The smoothing time to use for the spring. Longer times lead to more damped behaviour. A time of 0 will snap InOutX to TargetX
 	 * @param DeltaTime Timestep in seconds
 	 */
-	UFUNCTION(BlueprintCallable, meta = (Category = "Math|Spring|Experimental", AutoCreateRefTerm = "TargetX"))
+	UFUNCTION(BlueprintCallable, meta = (Category = "Math|Spring", AutoCreateRefTerm = "TargetX"))
 	static ENGINE_API void CriticalDoubleSpringDampFloat(UPARAM(ref) float& InOutX, UPARAM(ref) float& InOutV, UPARAM(ref) float& InOutXi, UPARAM(ref) float& InOutVi, const float& TargetX, float DeltaTime, float SmoothingTime = 0.2f);
 
-	UE_EXPERIMENTAL(5.7, "SpringMath is experimental")
 	/** Interpolates the angle InOutAngle towards TargetAngle using a critically damped spring via an intermediate state - producing a more S-shaped curve. The velocity of InOutAngle is stored in InOutAngularVelocity in deg/s.
 	 *
 	 * @param InOutAngle The value to be damped in degrees
@@ -157,10 +147,9 @@ public:
 	 * @param SmoothingTime The smoothing time to use for the spring. Longer times lead to more damped behaviour. A time of 0 will snap InOutAngle to TargetAngle
 	 * @param DeltaTime Timestep in seconds
 	 */
-	UFUNCTION(BlueprintCallable, meta = (Category = "Math|Spring|Experimental", AutoCreateRefTerm = "TargetAngle"))
+	UFUNCTION(BlueprintCallable, meta = (Category = "Math|Spring", AutoCreateRefTerm = "TargetAngle"))
 	static ENGINE_API void CriticalDoubleSpringDampAngle(UPARAM(ref) float& InOutAngle,UPARAM(ref) float& InOutAngularVelocity, UPARAM(ref) float& InOutIntermediateAngle, UPARAM(ref) float& InOutIntermediateAngularVelocity, const float& TargetAngle, float DeltaTime, float SmoothingTime = 0.2f);
 
-	UE_EXPERIMENTAL(5.7, "SpringMath is experimental")
 	/** Interpolates the value InOutX towards TargetX using a critically damped spring via an intermediate state - producing a more S-shaped curve. The velocity of X is stored in InOutV.
 	 *
 	 * @param InOutX The value to be damped
@@ -171,10 +160,9 @@ public:
 	 * @param SmoothingTime The smoothing time to use for the spring. Longer times lead to more damped behaviour. A time of 0 will snap InOutX to TargetX
 	 * @param DeltaTime Timestep in seconds
 	 */
-	UFUNCTION(BlueprintCallable, meta = (Category = "Math|Spring|Experimental", AutoCreateRefTerm = "TargetAngle"))
+	UFUNCTION(BlueprintCallable, meta = (Category = "Math|Spring", AutoCreateRefTerm = "TargetAngle"))
 	static ENGINE_API void CriticalDoubleSpringDampVector(UPARAM(ref) FVector& InOutX, UPARAM(ref) FVector& InOutV, UPARAM(ref) FVector& InOutXi, UPARAM(ref) FVector& InOutVi, const FVector& TargetX, float DeltaTime, float SmoothingTime = 0.2f);
 	
-	UE_EXPERIMENTAL(5.7, "SpringMath is experimental")
 	/** Interpolates the value InOutX towards TargetX using a critically damped spring via an intermediate state - producing a more S-shaped curve. The velocity of X is stored in InOutV.
 	 *
 	 * @param InOutX The value to be damped
@@ -185,10 +173,9 @@ public:
 	 * @param SmoothingTime The smoothing time to use for the spring. Longer times lead to more damped behaviour. A time of 0 will snap InOutX to TargetX
 	 * @param DeltaTime Timestep in seconds
 	 */
-	UFUNCTION(BlueprintCallable, meta = (Category = "Math|Spring|Experimental", AutoCreateRefTerm = "TargetX"))
+	UFUNCTION(BlueprintCallable, meta = (Category = "Math|Spring", AutoCreateRefTerm = "TargetX"))
 	static ENGINE_API void CriticalDoubleSpringDampVector2D(UPARAM(ref) FVector2D& InOutX, UPARAM(ref) FVector2D& InOutV, UPARAM(ref) FVector2D& InOutXi, UPARAM(ref) FVector2D& InOutVi, const FVector2D& TargetX, float DeltaTime, float SmoothingTime = 0.2f);
 
-	UE_EXPERIMENTAL(5.7, "SpringMath is experimental")
 	/** Interpolates the value InOutRotation towards TargetRotation using a critically damped spring via an intermediate state - producing a more S-shaped curve. The velocity of InOutRotation is stored in InOutAngularVelocity in deg/s.
 	 *
 	 * @param InOutRotation The value to be damped
@@ -199,10 +186,9 @@ public:
 	 * @param SmoothingTime The smoothing time to use for the spring. Longer times lead to more damped behaviour. A time of 0 will snap InOutRotation to TargetRotation
 	 * @param DeltaTime Timestep in seconds
 	 */
-	UFUNCTION(BlueprintCallable, meta = (Category = "Math|Spring|Experimental", AutoCreateRefTerm = "TargetRotation"))
+	UFUNCTION(BlueprintCallable, meta = (Category = "Math|Spring", AutoCreateRefTerm = "TargetRotation"))
 	static ENGINE_API void CriticalDoubleSpringDampQuat(UPARAM(ref) FQuat& InOutRotation, UPARAM(ref) FVector& InOutAngularVelocity, UPARAM(ref) FQuat& InOutIntermediateRotation, UPARAM(ref) FVector& InOutIntermediateAngularVelocity, const FQuat& TargetRotation, float DeltaTime, float SmoothingTime = 0.2f);
 
-	UE_EXPERIMENTAL(5.7, "SpringMath is experimental")
 	/** Interpolates the value InOutRotation towards TargetRotation using a critically damped spring via an intermediate value - producing a more S-shaped curve. The velocity of InOutRotation is stored in InOutAngularVelocity in deg/s.
 	 *
 	 * @param InOutRotation The value to be damped
@@ -213,10 +199,9 @@ public:
 	 * @param SmoothingTime The smoothing time to use for the spring. Longer times lead to more damped behaviour. A time of 0 will snap InOutRotation to TargetRotation
 	 * @param DeltaTime Timestep in seconds
 	 */
-	UFUNCTION(BlueprintCallable, meta = (Category = "Math|Spring|Experimental", AutoCreateRefTerm = "TargetRotation"))
+	UFUNCTION(BlueprintCallable, meta = (Category = "Math|Spring", AutoCreateRefTerm = "TargetRotation"))
 	static ENGINE_API void CriticalDoubleSpringDampRotator(UPARAM(ref) FRotator& InOutRotation, UPARAM(ref) FVector& InOutAngularVelocity, UPARAM(ref) FRotator& InOutIntermediateRotation, UPARAM(ref) FVector& InOutIntermediateAngularVelocity, const FRotator& TargetRotation, float DeltaTime, float SmoothingTime = 0.2f);
 
-	UE_EXPERIMENTAL(5.8, "SpringMath is experimental")
 	/** Interpolates the value InOutScale towards TargetScale using a critically damped spring via an intermediate state - producing a more S-shaped curve.
 	 *
 	 * @param InOutScale The value to be damped
@@ -227,10 +212,9 @@ public:
 	 * @param SmoothingTime The smoothing time to use for the spring. Longer times lead to more damped behaviour. A time of 0 will snap InOutScale to TargetScale
 	 * @param DeltaTime Timestep in seconds
 	 */
-	UFUNCTION(BlueprintCallable, meta = (Category = "Math|Spring|Experimental", AutoCreateRefTerm = "TargetScale"))
+	UFUNCTION(BlueprintCallable, meta = (Category = "Math|Spring", AutoCreateRefTerm = "TargetScale"))
 	static ENGINE_API void CriticalDoubleSpringDampScale(UPARAM(ref) FVector& InOutScale, UPARAM(ref) FVector& InOutScalarVelocity, UPARAM(ref) FVector& InOutIntermediateScale, UPARAM(ref) FVector& InOutIntermediateScalarVelocity, const FVector& TargetScale, float DeltaTime, float SmoothingTime = 0.2f);
 
-	UE_EXPERIMENTAL(5.8, "SpringMath is experimental")
 	/** Interpolates the value InOutTransform towards TargetTransform using a critically damped spring via an intermediate state - producing a more S-shaped curve.
 	 *
 	 * @param InOutTransform The value to be damped
@@ -241,13 +225,12 @@ public:
 	 * @param SmoothingTime The smoothing time to use for the spring. Longer times lead to more damped behaviour. A time of 0 will snap InOutTransform to TargetTransform
 	 * @param DeltaTime Timestep in seconds
 	 */
-	UFUNCTION(BlueprintCallable, meta = (Category = "Math|Spring|Experimental", AutoCreateRefTerm = "TargetTransform"))
+	UFUNCTION(BlueprintCallable, meta = (Category = "Math|Spring", AutoCreateRefTerm = "TargetTransform"))
 	static ENGINE_API void CriticalDoubleSpringDampTransform(UPARAM(ref) FTransform& InOutTransform, UPARAM(ref) FSpringTransformVelocity& InOutVelocity, UPARAM(ref) FTransform& InOutIntermediateTransform, UPARAM(ref) FSpringTransformVelocity& InOutIntermediateTransformVelocity, const FTransform& TargetTransform, float DeltaTime, float SmoothingTime = 0.2f);
 
 
 	// Velocity Spring //
 
-	UE_EXPERIMENTAL(5.7, "SpringMath is experimental")
 	/** A velocity spring will damp towards a target that follows a fixed linear target velocity, allowing control of the interpolation speed
 	 * while still giving a smoothed behavior. A SmoothingTime of 0 will give a linear interpolation between X and TargetX
 	 * 
@@ -259,10 +242,9 @@ public:
 	 * @param SmoothingTime The smoothing time to use while damping towards X. Higher values will give more smoothed behaviour. A value of 0 will give a linear interpolation of X to Target
 	 * @param DeltaTime The timestep in seconds
 	 */
-	UFUNCTION(BlueprintCallable, meta=(Category="Math|Spring|Experimental"))
+	UFUNCTION(BlueprintCallable, meta=(Category="Math|Spring"))
 	static ENGINE_API void VelocitySpringDampFloat(UPARAM(ref) float& InOutX, UPARAM(ref) float& InOutV, UPARAM(ref) float& InOutVi, float TargetX, float MaxSpeed, float DeltaTime, float SmoothingTime = 0.2f);
 
-	UE_EXPERIMENTAL(5.7, "SpringMath is experimental")
 	/** A velocity spring will damp towards a target that follows a fixed linear target velocity, allowing control of the interpolation speed
 	 * while still giving a smoothed behavior. A SmoothingTime of 0 will give a linear interpolation between X and TargetX
 	 * 
@@ -274,10 +256,9 @@ public:
 	 * @param SmoothingTime The smoothing time to use while damping towards X. Higher values will give more smoothed behaviour. A value of 0 will give a linear interpolation of X to Target
 	 * @param DeltaTime The timestep in seconds
 	 */
-	UFUNCTION(BlueprintCallable, meta=(Category="Math|Spring|Experimental"))
+	UFUNCTION(BlueprintCallable, meta=(Category="Math|Spring"))
 	static ENGINE_API void VelocitySpringDampVector(UPARAM(ref) FVector& InOutX, UPARAM(ref) FVector& InOutV, UPARAM(ref) FVector& InOutVi, const FVector& TargetX, float MaxSpeed, float DeltaTime,  float SmoothingTime = 0.2f);
 
-	UE_EXPERIMENTAL(5.7, "SpringMath is experimental")
 	/** A velocity spring will damp towards a target that follows a fixed linear target velocity, allowing control of the interpolation speed
 	 * while still giving a smoothed behavior. A SmoothingTime of 0 will give a linear interpolation between X and TargetX
 	 * 
@@ -289,13 +270,12 @@ public:
 	 * @param SmoothingTime The smoothing time to use while damping towards X. Higher values will give more smoothed behaviour. A value of 0 will give a linear interpolation of X to Target
 	 * @param DeltaTime The timestep in seconds
 	 */
-	UFUNCTION(BlueprintCallable, meta=(Category="Math|Spring|Experimental"))
+	UFUNCTION(BlueprintCallable, meta=(Category="Math|Spring"))
 	static ENGINE_API void VelocitySpringDampVector2D(UPARAM(ref) FVector2D& InOutX, UPARAM(ref) FVector2D& InOutV, UPARAM(ref) FVector2D& InOutVi, const FVector2D& TargetX, float MaxSpeed, float DeltaTime,  float SmoothingTime = 0.2f);
 
 
 	// Damper //
 
-	UE_EXPERIMENTAL(5.7, "SpringMath is experimental")
 	/**
 	 * Smooths a value using exponential damping towards a target.
 	 * 
@@ -305,10 +285,9 @@ public:
 	 * @param  DeltaTime Time interval
 	 * @param  SmoothingTime Timescale over which to smooth. Larger values result in more smoothed behaviour. Can be zero.
 	 */
-	UFUNCTION(BlueprintPure, meta=(Category="Math|Spring|Experimental"))
+	UFUNCTION(BlueprintPure, meta=(Category="Math|Spring"))
 	static ENGINE_API float DampFloat(float Value, float Target, float DeltaTime, float SmoothingTime = 0.2f);
 
-	UE_EXPERIMENTAL(5.7, "SpringMath is experimental")
 	/**
 	 * Smooths an angle in degrees using exponential damping towards a target.
 	 *
@@ -318,10 +297,9 @@ public:
 	 * @param  DeltaTime Time interval
 	 * @param  SmoothingTime Timescale over which to smooth. Larger values result in more smoothed behaviour. Can be zero.
 	 */
-	UFUNCTION(BlueprintPure, meta = (Category = "Math|Spring|Experimental"))
+	UFUNCTION(BlueprintPure, meta = (Category = "Math|Spring"))
 	static ENGINE_API float DampAngle(float Angle, float TargetAngle, float DeltaTime, float SmoothingTime = 0.2f);
 
-	UE_EXPERIMENTAL(5.7, "SpringMath is experimental")
 	/**
 	 * Smooths a value using exponential damping towards a target.
 	 * 
@@ -331,10 +309,9 @@ public:
 	 * @param  DeltaTime Time interval
 	 * @param  SmoothingTime Timescale over which to smooth. Larger values result in more smoothed behaviour. Can be zero.
 	 */
-	UFUNCTION(BlueprintPure, meta=(Category="Math|Spring|Experimental", AutoCreateRefTerm = "Value, Target"))
+	UFUNCTION(BlueprintPure, meta=(Category="Math|Spring", AutoCreateRefTerm = "Value, Target"))
 	static ENGINE_API FVector DampVector(const FVector& Value, const FVector& Target, float DeltaTime, float SmoothingTime = 0.2f);
 
-	UE_EXPERIMENTAL(5.7, "SpringMath is experimental")
 	/**
 	 * Smooths a value using exponential damping towards a target.
 	 * 
@@ -344,10 +321,9 @@ public:
 	 * @param  DeltaTime Time interval
 	 * @param  SmoothingTime Timescale over which to smooth. Larger values result in more smoothed behaviour. Can be zero.
 	 */
-	UFUNCTION(BlueprintPure, meta=(Category="Math|Spring|Experimental", AutoCreateRefTerm = "Value, Target"))
+	UFUNCTION(BlueprintPure, meta=(Category="Math|Spring", AutoCreateRefTerm = "Value, Target"))
 	static ENGINE_API FVector2D DampVector2D(const FVector2D& Value, const FVector2D& Target, float DeltaTime, float SmoothingTime = 0.2f);
 
-	UE_EXPERIMENTAL(5.7, "SpringMath is experimental")
 	/**
 	 * Smooths a rotation using exponential damping towards a target.
 	 * 
@@ -357,10 +333,9 @@ public:
 	 * @param  DeltaTime Time interval
 	 * @param  SmoothingTime Timescale over which to smooth. Larger values result in more smoothed behaviour. Can be zero.
 	 */
-	UFUNCTION(BlueprintPure, meta=(Category="Math|Spring|Experimental", AutoCreateRefTerm = "Rotation, TargetRotation"))
+	UFUNCTION(BlueprintPure, meta=(Category="Math|Spring", AutoCreateRefTerm = "Rotation, TargetRotation"))
 	static ENGINE_API FQuat DampQuat(const FQuat& Rotation, const FQuat& TargetRotation, float DeltaTime, float SmoothingTime = 0.2f);
 
-	UE_EXPERIMENTAL(5.7, "SpringMath is experimental")
 	/**
 	 * Smooths a value using exponential damping towards a target.
 	 * 
@@ -370,10 +345,9 @@ public:
 	 * @param  DeltaTime Time interval
 	 * @param  SmoothingTime Timescale over which to smooth. Larger values result in more smoothed behaviour. Can be zero.
 	 */
-	UFUNCTION(BlueprintPure, meta=(Category="Math|Spring|Experimental", AutoCreateRefTerm = "Rotation, TargetRotation"))
+	UFUNCTION(BlueprintPure, meta=(Category="Math|Spring", AutoCreateRefTerm = "Rotation, TargetRotation"))
 	static ENGINE_API FRotator DampRotator(const FRotator& Rotation, const FRotator& TargetRotation, float DeltaTime, float SmoothingTime = 0.2f);
 
-	UE_EXPERIMENTAL(5.8, "SpringMath is experimental")
 	/**
 	 * Smooths a scale using exponential damping towards a target.
 	 *
@@ -383,10 +357,9 @@ public:
 	 * @param  DeltaTime Time interval
 	 * @param  SmoothingTime Timescale over which to smooth. Larger values result in more smoothed behaviour. Can be zero.
 	 */
-	UFUNCTION(BlueprintPure, meta = (Category = "Math|Spring|Experimental", AutoCreateRefTerm = "Scale, TargetScale"))
+	UFUNCTION(BlueprintPure, meta = (Category = "Math|Spring", AutoCreateRefTerm = "Scale, TargetScale"))
 	static ENGINE_API FVector DampScale(const FVector& Scale, const FVector& TargetScale, float DeltaTime, float SmoothingTime = 0.2f);
 
-	UE_EXPERIMENTAL(5.8, "SpringMath is experimental")
 	/**
 	 * Smooths a transform using exponential damping towards a target.
 	 *
@@ -396,13 +369,12 @@ public:
 	 * @param  DeltaTime Time interval
 	 * @param  SmoothingTime Timescale over which to smooth. Larger values result in more smoothed behaviour. Can be zero.
 	 */
-	UFUNCTION(BlueprintPure, meta = (Category = "Math|Spring|Experimental", AutoCreateRefTerm = "Transform, TargetTransform"))
+	UFUNCTION(BlueprintPure, meta = (Category = "Math|Spring", AutoCreateRefTerm = "Transform, TargetTransform"))
 	static ENGINE_API FTransform DampTransform(const FTransform& Transform, const FTransform& TargetTransform, float DeltaTime, float SmoothingTime = 0.2f);
 
 
 	// Character //
 
-	UE_EXPERIMENTAL(5.7, "SpringMath is experimental")
 	/** Update the position of a character given a target velocity using a simple damped spring
 	 * 
 	 * @param InOutPosition The position of the character
@@ -412,10 +384,9 @@ public:
 	 * @param DeltaTime The delta time to tick the character
 	 * @param SmoothingTime The time over which to smooth velocity. It takes roughly the smoothing time in order for the character to reach the target velocity.
 	 */
-	UFUNCTION(BlueprintCallable, meta=(Category="Math|Spring|Experimental", AutoCreateRefTerm = "TargetVelocity"))
+	UFUNCTION(BlueprintCallable, meta=(Category="Math|Spring", AutoCreateRefTerm = "TargetVelocity"))
 	static ENGINE_API void SpringCharacterUpdate(UPARAM(ref) FVector& InOutPosition, UPARAM(ref) FVector& InOutVelocity, UPARAM(ref) FVector& InOutAcceleration, const FVector& TargetVelocity, float DeltaTime, float SmoothingTime = 0.2f);
 
-	UE_EXPERIMENTAL(5.8, "SpringMath is experimental")
 	/** Estimates an approximate stopping distance from a smoothing time for a character using a simple damped spring
 	 *
 	 * @param InitialVelocity The velocity of the character at the start of the stop
@@ -423,10 +394,9 @@ public:
 	 * @param SmoothingTime The time over which to smooth velocity. It takes roughly the smoothing time in order for the character to reach the target velocity.
 	 * @return An estimation of the total distance traveled during the stop
 	 */
-	UFUNCTION(BlueprintPure, meta = (Category = "Math|Spring|Experimental", ReturnDisplayName = "Stopping Distance"))
+	UFUNCTION(BlueprintPure, meta = (Category = "Math|Spring", ReturnDisplayName = "Stopping Distance"))
 	static ENGINE_API float SpringCharacterStoppingDistance(const float InitialVelocity, const float InitialAcceleration, const float SmoothingTime = 0.2f);
 
-	UE_EXPERIMENTAL(5.8, "SpringMath is experimental")
 	/** Estimates an approximate stopping time (how long until the velocity goes below the VelocityThreshold) from a smoothing time for a character using a simple damped spring. Assumes the initial acceleration is zero.
 	 *
 	 * @param InitialVelocity The velocity of the character at the start of the stop
@@ -434,20 +404,18 @@ public:
 	 * @param VelocityThreshold The velocity threshold at which the character is considered to be stopped in cm/s
 	 * @return An estimation of the time taken to come to a stop
 	 */
-	UFUNCTION(BlueprintPure, meta = (Category = "Math|Spring|Experimental", ReturnDisplayName = "Stopping Time"))
+	UFUNCTION(BlueprintPure, meta = (Category = "Math|Spring", ReturnDisplayName = "Stopping Time"))
 	static ENGINE_API float SpringCharacterStoppingTime(const float InitialVelocity, const float SmoothingTime = 0.2f, const float VelocityThreshold = 1.0f);
 
-	UE_EXPERIMENTAL(5.8, "SpringMath is experimental")
 	/** Estimates an approximate smoothing time from a stopping distance for a character using by a simple damped spring. Assumes the initial acceleration is zero.
 	 *
 	 * @param InitialVelocity The velocity of the character at the start of the stop
 	 * @param StoppingDistance The total distance traveled during the stop
 	 * @return An estimation of the SmoothingTime of the spring
 	 */
-	UFUNCTION(BlueprintPure, meta = (Category = "Math|Spring|Experimental", ReturnDisplayName = "Smoothing Time"))
+	UFUNCTION(BlueprintPure, meta = (Category = "Math|Spring", ReturnDisplayName = "Smoothing Time"))
 	static ENGINE_API float SpringCharacterSmoothingTimeFromStoppingDistance(const float InitialVelocity, const float StoppingDistance = 100.0f);
 
-	UE_EXPERIMENTAL(5.8, "SpringMath is experimental")
 	/** Estimates an approximate smoothing time from a stopping time for a character using by a simple damped spring. Assumes the initial acceleration is zero.
 	 *
 	 * @param InitialVelocity The velocity of the character at the start of the stop
@@ -455,10 +423,9 @@ public:
 	 * @param VelocityThreshold The velocity threshold at which the character is considered to be stopped in cm/s
 	 * @return An estimation of the SmoothingTime of the spring
 	 */
-	UFUNCTION(BlueprintPure, meta = (Category = "Math|Spring|Experimental", ReturnDisplayName = "Smoothing Time"))
+	UFUNCTION(BlueprintPure, meta = (Category = "Math|Spring", ReturnDisplayName = "Smoothing Time"))
 	static ENGINE_API float SpringCharacterSmoothingTimeFromStoppingTime(const float InitialVelocity, const float StoppingTime = 1.0f, const float VelocityThreshold = 10.0f);
 
-	UE_EXPERIMENTAL(5.8, "SpringMath is experimental")
 	/** Estimates an approximate starting distance (distance traveled until the velocity is within the VelocityThreshold of the target) from a smoothing time for a character using a simple damped spring
 	 *
 	 * @param TargetVelocity The target velocity of the character
@@ -466,10 +433,9 @@ public:
 	 * @param VelocityThreshold The velocity threshold from the target at which the character is considered to have started in cm/s
 	 * @return An estimation of the total distance traveled during the start
 	 */
-	UFUNCTION(BlueprintPure, meta = (Category = "Math|Spring|Experimental", ReturnDisplayName = "Starting Distance"))
+	UFUNCTION(BlueprintPure, meta = (Category = "Math|Spring", ReturnDisplayName = "Starting Distance"))
 	static ENGINE_API float SpringCharacterStartingDistance(const float TargetVelocity, const float SmoothingTime = 0.2f, const float VelocityThreshold = 10.0f);
 
-	UE_EXPERIMENTAL(5.8, "SpringMath is experimental")
 	/** Estimates an approximate starting time (time until the velocity is within the VelocityThreshold of the target) from a smoothing time for a character using a simple damped spring.
 	 *
 	 * @param TargetVelocity The target velocity of the character
@@ -477,10 +443,9 @@ public:
 	 * @param VelocityThreshold The velocity threshold from the target at which the character is considered to have started in cm/s
 	 * @return An estimation of the time taken to come to a start
 	 */
-	UFUNCTION(BlueprintPure, meta = (Category = "Math|Spring|Experimental", ReturnDisplayName = "Starting Time"))
+	UFUNCTION(BlueprintPure, meta = (Category = "Math|Spring", ReturnDisplayName = "Starting Time"))
 	static ENGINE_API float SpringCharacterStartingTime(const float TargetVelocity, const float SmoothingTime = 0.2f, const float VelocityThreshold = 10.0f);
 
-	UE_EXPERIMENTAL(5.8, "SpringMath is experimental")
 	/** Estimates an approximate smoothing time from a starting distance for a character using by a simple damped spring.
 	 *
 	 * @param TargetVelocity The target velocity of the character
@@ -488,10 +453,9 @@ public:
 	 * @param VelocityThreshold The velocity threshold from the target at which the character is considered to have started in cm/s
 	 * @return An estimation of the SmoothingTime of the spring
 	 */
-	UFUNCTION(BlueprintPure, meta = (Category = "Math|Spring|Experimental", ReturnDisplayName = "Smoothing Time"))
+	UFUNCTION(BlueprintPure, meta = (Category = "Math|Spring", ReturnDisplayName = "Smoothing Time"))
 	static ENGINE_API float SpringCharacterSmoothingTimeFromStartingDistance(const float TargetVelocity, const float StartingDistance = 100.0f, const float VelocityThreshold = 10.0f);
 
-	UE_EXPERIMENTAL(5.8, "SpringMath is experimental")
 	/** Estimates an approximate smoothing time from a starting time for a character using by a simple damped spring.
 	 *
 	 * @param TargetVelocity The target velocity of the character
@@ -499,10 +463,9 @@ public:
 	 * @param VelocityThreshold The velocity threshold from the target at which the character is considered to have started in cm/s
 	 * @return An estimation of the SmoothingTime of the spring
 	 */
-	UFUNCTION(BlueprintPure, meta = (Category = "Math|Spring|Experimental", ReturnDisplayName = "Smoothing Time"))
+	UFUNCTION(BlueprintPure, meta = (Category = "Math|Spring", ReturnDisplayName = "Smoothing Time"))
 	static ENGINE_API float SpringCharacterSmoothingTimeFromStartingTime(const float TargetVelocity, const float StartingTime = 1.0f, const float VelocityThreshold = 10.0f);
 
-	UE_EXPERIMENTAL(5.8, "SpringMath is experimental")
 	/** Estimates an approximate turning time (time until the angle is within the AngleThreshold of the target) from a smoothing time for a character using a simple damped spring. Assumes the initial angular velocity is zero.
 	 *
 	 * @param TargetAngle The target angle of the character in degrees
@@ -510,10 +473,9 @@ public:
 	 * @param AngleThreshold The angle threshold from the target at which the character is considered to have turned in degrees
 	 * @return An estimation of the time taken to turn
 	 */
-	UFUNCTION(BlueprintPure, meta = (Category = "Math|Spring|Experimental", ReturnDisplayName = "Turning Time"))
+	UFUNCTION(BlueprintPure, meta = (Category = "Math|Spring", ReturnDisplayName = "Turning Time"))
 	static ENGINE_API float SpringCharacterTurningTime(const float TargetAngle, const float SmoothingTime = 0.2f, const float AngleThreshold = 5.0f);
 
-	UE_EXPERIMENTAL(5.8, "SpringMath is experimental")
 	/** Estimates an approximate smoothing time from a turning time for a character using by a simple damped spring.
 	 *
 	 * @param TargetAngle The target angle of the character in degrees
@@ -521,10 +483,9 @@ public:
 	 * @param AngleThreshold The angle threshold from the target at which the character is considered to have turned in degrees
 	 * @return An estimation of the SmoothingTime of the spring
 	 */
-	UFUNCTION(BlueprintPure, meta = (Category = "Math|Spring|Experimental", ReturnDisplayName = "Smoothing Time"))
+	UFUNCTION(BlueprintPure, meta = (Category = "Math|Spring", ReturnDisplayName = "Smoothing Time"))
 	static ENGINE_API float SpringCharacterSmoothingTimeFromTurningTime(const float TargetAngle, const float TurningTime = 1.0f, const float AngleThreshold = 5.0f);
 
-	UE_EXPERIMENTAL(5.8, "SpringMath is experimental")
 	/** Estimates an approximate maximum acceleration from a smoothing time for a character using a simple damped spring.
 	 *
 	 * @param InitialVelocity The initial velocity of the character
@@ -533,10 +494,9 @@ public:
 	 * @param SmoothingTime The time over which to smooth velocity. It takes roughly the smoothing time in order for the character to reach the target velocity.
 	 * @return An estimation of the maximum acceleration
 	 */
-	UFUNCTION(BlueprintPure, meta = (Category = "Math|Spring|Experimental", ReturnDisplayName = "Maximum Acceleration"))
+	UFUNCTION(BlueprintPure, meta = (Category = "Math|Spring", ReturnDisplayName = "Maximum Acceleration"))
 	static ENGINE_API float SpringCharacterMaximumAcceleration(const float InitialVelocity, const float InitialAcceleration, const float TargetVelocity, const float SmoothingTime = 0.2f);
 
-	UE_EXPERIMENTAL(5.8, "SpringMath is experimental")
 	/** Estimates an approximate smoothing time from a maximum acceleration for a character using by a simple damped spring.
 	 *
 	 * @param InitialVelocity The initial velocity of the character
@@ -544,10 +504,9 @@ public:
 	 * @param MaximumAcceleration The maximum acceleration of the character
 	 * @return An estimation of the SmoothingTime of the spring
 	 */
-	UFUNCTION(BlueprintPure, meta = (Category = "Math|Spring|Experimental", ReturnDisplayName = "Smoothing Time"))
+	UFUNCTION(BlueprintPure, meta = (Category = "Math|Spring", ReturnDisplayName = "Smoothing Time"))
 	static ENGINE_API float SpringCharacterSmoothingTimeFromMaximumAcceleration(const float InitialVelocity, const float TargetVelocity, const float MaximumAcceleration);
 
-	UE_EXPERIMENTAL(5.8, "SpringMath is experimental")
 	/** Estimates an approximate maximum angular velocity from a smoothing time for a character using a simple damped spring.
 	 *
 	 * @param InitialAngle The initial angle of the character in degrees
@@ -556,10 +515,9 @@ public:
 	 * @param SmoothingTime The time over which to smooth angle. It takes roughly the smoothing time in order for the character to reach the target angle.
 	 * @return An estimation of the maximum angular velocity
 	 */
-	UFUNCTION(BlueprintPure, meta = (Category = "Math|Spring|Experimental", ReturnDisplayName = "Maximum Angular Velocity"))
+	UFUNCTION(BlueprintPure, meta = (Category = "Math|Spring", ReturnDisplayName = "Maximum Angular Velocity"))
 	static ENGINE_API float SpringCharacterMaximumAngularVelocity(const float InitialAngle, const float InitialAngularVelocity, const float TargetAngle, const float SmoothingTime = 0.2f);
 
-	UE_EXPERIMENTAL(5.8, "SpringMath is experimental")
 	/** Estimates an approximate smoothing time from a maximum angular velocity for a character using by a simple damped spring.
 	 *
 	 * @param InitialAngle The initial angle of the character in degrees
@@ -567,10 +525,9 @@ public:
 	 * @param MaximumAngularVelocity The maximum angular velocity of the character
 	 * @return An estimation of the SmoothingTime of the spring
 	 */
-	UFUNCTION(BlueprintPure, meta = (Category = "Math|Spring|Experimental", ReturnDisplayName = "Smoothing Time"))
+	UFUNCTION(BlueprintPure, meta = (Category = "Math|Spring", ReturnDisplayName = "Smoothing Time"))
 	static ENGINE_API float SpringCharacterSmoothingTimeFromMaximumAngularVelocity(const float InitialAngle, const float TargetAngle, const float MaximumAngularVelocity);
 
-	UE_EXPERIMENTAL(5.7, "SpringMath is experimental")
 	/** Update a position representing a character given a target velocity using a velocity spring.
 	 * A velocity spring tracks an intermediate velocity which moves at a maximum acceleration linearly towards a target.
 	 * This means unlike the "SpringCharacterUpdate", it will take longer to reach a target velocity that is further away from the current velocity.
@@ -584,95 +541,82 @@ public:
 	 * @param SmoothingTime The time over which to smooth velocity. It takes roughly the smoothing time in order for the character to reach the target velocity.
 	 * @param MaxAcceleration Puts a limit on the maximum acceleration that the intermediate velocity can do each frame. If MaxAccel is very large, the behaviour will be the same as SpringCharacterUpdate
 	 */
-	UFUNCTION(BlueprintCallable, meta=(Category="Math|Spring|Experimental", AutoCreateRefTerm = "TargetVelocity"))
+	UFUNCTION(BlueprintCallable, meta=(Category="Math|Spring", AutoCreateRefTerm = "TargetVelocity"))
 	static ENGINE_API void VelocitySpringCharacterUpdate(UPARAM(ref) FVector& InOutPosition, UPARAM(ref) FVector& InOutVelocity, UPARAM(ref) FVector& InOutVelocityIntermediate, UPARAM(ref) FVector& InOutAcceleration, const FVector& TargetVelocity, float DeltaTime, float SmoothingTime = 0.2f, float MaxAcceleration = 1.0f);
 
 
 
 	// Conversion Methods //
 
-	UE_EXPERIMENTAL(5.7, "SpringMath is experimental")
 	/** Convert from smoothing time to spring strength.
 	* 
 	* @param SmoothingTime The smoothing time of the spring in seconds. It is the time by which the output lags the input when critically damped, when the input is changing at a constant rate.
 	* @return The spring strength. This corresponds to the undamped frequency of the spring in hz.
 	*/
-	UFUNCTION(BlueprintPure, meta=(Category="Math|Spring|Experimental"))
+	UFUNCTION(BlueprintPure, meta=(Category="Math|Spring"))
 	static ENGINE_API float ConvertSmoothingTimeToStrength(float SmoothingTime);
 
-	UE_EXPERIMENTAL(5.7, "SpringMath is experimental")
 	/** Convert from spring strength to smoothing time.
  	* 
  	* @param Strength The spring strength. This corresponds to the undamped frequency of the spring in hz.
  	* @return The smoothing time of the spring in seconds. It is the time by which the output lags the input when critically damped, when the input is changing at a constant rate.
  	*/
-	UFUNCTION(BlueprintPure, meta=(Category="Math|Spring|Experimental"))
+	UFUNCTION(BlueprintPure, meta=(Category="Math|Spring"))
 	static ENGINE_API float ConvertStrengthToSmoothingTime(float Strength);
 
-	UE_EXPERIMENTAL(5.7, "SpringMath is experimental")
 	/** Convert a halflife to a smoothing time
  	* 
  	* @param HalfLife The half life of the spring. How long it takes the value to get halfway towards the target.
  	* @return The smoothing time of the spring in seconds. It is the time by which the output lags the input when critically damped, when the input is changing at a constant rate.
  	*/
-	UFUNCTION(BlueprintPure, meta=(Category="Math|Spring|Experimental"))
+	UFUNCTION(BlueprintPure, meta=(Category="Math|Spring"))
 	static ENGINE_API float ConvertHalfLifeToSmoothingTime(float HalfLife);
 
-	UE_EXPERIMENTAL(5.7, "SpringMath is experimental")
 	/** Convert a smoothing time to a half life
  	* 
  	* @param SmoothingTime The smoothing time of the spring in seconds. It is the time by which the output lags the input when critically damped, when the input is changing at a constant rate. 
  	* @return The half life of the spring. How long it takes the value to get halfway towards the target. 
  	*/
-	UFUNCTION(BlueprintPure, meta=(Category="Math|Spring|Experimental"))
+	UFUNCTION(BlueprintPure, meta=(Category="Math|Spring"))
 	static ENGINE_API float ConvertSmoothingTimeToHalfLife(float SmoothingTime);
 
 
 	// Velocity Tracking //
 
-	UE_EXPERIMENTAL(5.8, "SpringMath is experimental")
 	/** Update the value and velocity using the finite difference  */
-	UFUNCTION(BlueprintCallable, meta = (Category = "Math|Spring|Experimental"))
+	UFUNCTION(BlueprintCallable, meta = (Category = "Math|Spring"))
 	static ENGINE_API void TrackVelocityFloat(UPARAM(ref) float& InOutValue, UPARAM(ref) float& InOutVelocity, float InValue, float DeltaTime);
 
-	UE_EXPERIMENTAL(5.8, "SpringMath is experimental")
 	/** Update the value and velocity using the finite difference  */
-	UFUNCTION(BlueprintCallable, meta = (Category = "Math|Spring|Experimental"))
+	UFUNCTION(BlueprintCallable, meta = (Category = "Math|Spring"))
 	static ENGINE_API void TrackVelocityVector(UPARAM(ref) FVector& InOutValue, UPARAM(ref) FVector& InOutVelocity, FVector InValue, float DeltaTime);
 
-	UE_EXPERIMENTAL(5.8, "SpringMath is experimental")
 	/** Update the value and velocity using the finite difference  */
-	UFUNCTION(BlueprintCallable, meta = (Category = "Math|Spring|Experimental"))
+	UFUNCTION(BlueprintCallable, meta = (Category = "Math|Spring"))
 	static ENGINE_API void TrackVelocityVector2D(UPARAM(ref) FVector2D& InOutValue, UPARAM(ref) FVector2D& InOutVelocity, FVector2D InValue, float DeltaTime);
 
-	UE_EXPERIMENTAL(5.8, "SpringMath is experimental")
 	/** Update the value and velocity using the finite difference  */
-	UFUNCTION(BlueprintCallable, meta = (Category = "Math|Spring|Experimental"))
+	UFUNCTION(BlueprintCallable, meta = (Category = "Math|Spring"))
 	static ENGINE_API void TrackVelocityQuat(UPARAM(ref) FQuat& InOutValue, UPARAM(ref) FVector& InOutVelocity, FQuat InValue, float DeltaTime);
 
-	UE_EXPERIMENTAL(5.8, "SpringMath is experimental")
 	/** Update the value and velocity using the finite difference  */
-	UFUNCTION(BlueprintCallable, meta = (Category = "Math|Spring|Experimental"))
+	UFUNCTION(BlueprintCallable, meta = (Category = "Math|Spring"))
 	static ENGINE_API void TrackVelocityRotator(UPARAM(ref) FRotator& InOutValue, UPARAM(ref) FVector& InOutVelocity, FRotator InValue, float DeltaTime);
 
-	UE_EXPERIMENTAL(5.8, "SpringMath is experimental")
 	/** Update the value and velocity using the finite difference  */
-	UFUNCTION(BlueprintCallable, meta = (Category = "Math|Spring|Experimental"))
+	UFUNCTION(BlueprintCallable, meta = (Category = "Math|Spring"))
 	static ENGINE_API void TrackVelocityAngle(UPARAM(ref) float& InOutValue, UPARAM(ref) float& InOutVelocity, float InValue, float DeltaTime);
 
-	UE_EXPERIMENTAL(5.8, "SpringMath is experimental")
 	/** Update the value and velocity using the finite difference  */
-	UFUNCTION(BlueprintCallable, meta = (Category = "Math|Spring|Experimental"))
+	UFUNCTION(BlueprintCallable, meta = (Category = "Math|Spring"))
 	static ENGINE_API void TrackVelocityScale(UPARAM(ref) FVector& InOutValue, UPARAM(ref) FVector& InOutVelocity, FVector InValue, float DeltaTime);
 
-	UE_EXPERIMENTAL(5.8, "SpringMath is experimental")
 	/** Update the value and velocity using the finite difference  */
-	UFUNCTION(BlueprintCallable, meta = (Category = "Math|Spring|Experimental"))
+	UFUNCTION(BlueprintCallable, meta = (Category = "Math|Spring"))
 	static ENGINE_API void TrackVelocityTransform(UPARAM(ref) FTransform& InOutValue, UPARAM(ref) FSpringTransformVelocity& InOutVelocity, FTransform InValue, float DeltaTime);
 
 	// Spring Inertialization //
 
-	UE_EXPERIMENTAL(5.8, "SpringMath is experimental")
 	/** Applies a spring inertialization to a float value, adding the appropriate offset given the time since transition.
 	*
 	* @param InValue				The value to apply the inertialization to
@@ -682,10 +626,9 @@ public:
 	* @param SmoothingTime			Timescale over which to apply the offset. Larger values result in longer blend times. Can be zero.
 	* @return						The value with appropriate offset applied
 	*/
-	UFUNCTION(BlueprintPure, meta = (Category = "Math|Spring|Experimental"))
+	UFUNCTION(BlueprintPure, meta = (Category = "Math|Spring"))
 	static ENGINE_API float SpringInertializeApplyToFloat(float InValue, float InValueOffset, float InVelocityOffset, float TimeSinceTransition, float SmoothingTime = 0.2f);
 
-	UE_EXPERIMENTAL(5.8, "SpringMath is experimental")
 	/** Transitions a spring inertialization, updating the appropriate offsets.
 	*
 	* @param InOutValueOffset		Offset value to update
@@ -697,10 +640,9 @@ public:
 	* @param TimeSinceTransition	Time since the last transition
 	* @param SmoothingTime			Timescale over which to apply the offset. Larger values result in longer blend times. Can be zero.
 	*/
-	UFUNCTION(BlueprintCallable, meta = (Category = "Math|Spring|Experimental"))
+	UFUNCTION(BlueprintCallable, meta = (Category = "Math|Spring"))
 	static ENGINE_API void SpringInertializeTransitionFloat(UPARAM(ref) float& InOutValueOffset, UPARAM(ref) float& InOutVelocityOffset, float InSrcValue, float InSrcVelocity, float InDstValue, float InDstVelocity, float TimeSinceTransition, float SmoothingTime = 0.2f);
 
-	UE_EXPERIMENTAL(5.8, "SpringMath is experimental")
 	/** Applies a spring inertialization to a vector value, adding the appropriate offset given the time since transition.
 	*
 	* @param InValue				The value to apply the inertialization to
@@ -710,10 +652,9 @@ public:
 	* @param SmoothingTime			Timescale over which to apply the offset. Larger values result in longer blend times. Can be zero.
 	* @return						The value with appropriate offset applied
 	*/
-	UFUNCTION(BlueprintPure, meta = (Category = "Math|Spring|Experimental"))
+	UFUNCTION(BlueprintPure, meta = (Category = "Math|Spring"))
 	static ENGINE_API FVector SpringInertializeApplyToVector(FVector InValue, FVector InValueOffset, FVector InVelocityOffset, float TimeSinceTransition, float SmoothingTime = 0.2f);
 
-	UE_EXPERIMENTAL(5.8, "SpringMath is experimental")
 	/** Transitions a spring inertialization, updating the appropriate offsets.
 	*
 	* @param InOutValueOffset		Offset value to update
@@ -725,10 +666,9 @@ public:
 	* @param TimeSinceTransition	Time since the last transition
 	* @param SmoothingTime			Timescale over which to apply the offset. Larger values result in longer blend times. Can be zero.
 	*/
-	UFUNCTION(BlueprintCallable, meta = (Category = "Math|Spring|Experimental"))
+	UFUNCTION(BlueprintCallable, meta = (Category = "Math|Spring"))
 	static ENGINE_API void SpringInertializeTransitionVector(UPARAM(ref) FVector& InOutValueOffset, UPARAM(ref) FVector& InOutVelocityOffset, FVector InSrcValue, FVector InSrcVelocity, FVector InDstValue, FVector InDstVelocity, float TimeSinceTransition, float SmoothingTime = 0.2f);
 
-	UE_EXPERIMENTAL(5.8, "SpringMath is experimental")
 	/** Applies a spring inertialization to a Vector2D value, adding the appropriate offset given the time since transition.
 	*
 	* @param InValue				The value to apply the inertialization to
@@ -738,10 +678,9 @@ public:
 	* @param SmoothingTime			Timescale over which to apply the offset. Larger values result in longer blend times. Can be zero.
 	* @return						The value with appropriate offset applied
 	*/
-	UFUNCTION(BlueprintPure, meta = (Category = "Math|Spring|Experimental"))
+	UFUNCTION(BlueprintPure, meta = (Category = "Math|Spring"))
 	static ENGINE_API FVector2D SpringInertializeApplyToVector2D(FVector2D InValue, FVector2D InValueOffset, FVector2D InVelocityOffset, float TimeSinceTransition, float SmoothingTime = 0.2f);
 
-	UE_EXPERIMENTAL(5.8, "SpringMath is experimental")
 	/** Transitions a spring inertialization, updating the appropriate offsets.
 	*
 	* @param InOutValueOffset		Offset value to update
@@ -753,10 +692,9 @@ public:
 	* @param TimeSinceTransition	Time since the last transition
 	* @param SmoothingTime			Timescale over which to apply the offset. Larger values result in longer blend times. Can be zero.
 	*/
-	UFUNCTION(BlueprintCallable, meta = (Category = "Math|Spring|Experimental"))
+	UFUNCTION(BlueprintCallable, meta = (Category = "Math|Spring"))
 	static ENGINE_API void SpringInertializeTransitionVector2D(UPARAM(ref) FVector2D& InOutValueOffset, UPARAM(ref) FVector2D& InOutVelocityOffset, FVector2D InSrcValue, FVector2D InSrcVelocity, FVector2D InDstValue, FVector2D InDstVelocity, float TimeSinceTransition, float SmoothingTime = 0.2f);
 
-	UE_EXPERIMENTAL(5.8, "SpringMath is experimental")
 	/** Applies a spring inertialization to a quaternion, applying the appropriate offset given the time since transition.
 	*
 	* @param InRotation					The rotation to apply the inertialization to
@@ -766,10 +704,9 @@ public:
 	* @param SmoothingTime				Timescale over which to apply the offset. Larger values result in longer blend times. Can be zero.
 	* @return							The rotation with appropriate offset applied
 	*/
-	UFUNCTION(BlueprintPure, meta = (Category = "Math|Spring|Experimental"))
+	UFUNCTION(BlueprintPure, meta = (Category = "Math|Spring"))
 	static ENGINE_API FQuat SpringInertializeApplyToQuat(FQuat InRotation, FQuat InRotationOffset, FVector InAngularVelocityOffset, float TimeSinceTransition, float SmoothingTime = 0.2f);
 
-	UE_EXPERIMENTAL(5.8, "SpringMath is experimental")
 	/** Transitions a spring inertialization, updating the appropriate offsets.
 	*
 	* @param InOutRotationOffset		Offset rotation to update
@@ -781,10 +718,9 @@ public:
 	* @param TimeSinceTransition		Time since the last transition
 	* @param SmoothingTime				Timescale over which to apply the offset. Larger values result in longer blend times. Can be zero.
 	*/
-	UFUNCTION(BlueprintCallable, meta = (Category = "Math|Spring|Experimental"))
+	UFUNCTION(BlueprintCallable, meta = (Category = "Math|Spring"))
 	static ENGINE_API void SpringInertializeTransitionQuat(UPARAM(ref) FQuat& InOutRotationOffset, UPARAM(ref) FVector& InOutAngularVelocityOffset, FQuat InSrcRotation, FVector InSrcAngularVelocity, FQuat InDstRotation, FVector InDstAngularVelocity, float TimeSinceTransition, float SmoothingTime = 0.2f);
 
-	UE_EXPERIMENTAL(5.8, "SpringMath is experimental")
 	/** Applies a spring inertialization to a rotator, applying the appropriate offset given the time since transition.
 	*
 	* @param InRotation					The rotation to apply the inertialization to
@@ -794,10 +730,9 @@ public:
 	* @param SmoothingTime				Timescale over which to apply the offset. Larger values result in longer blend times. Can be zero.
 	* @return							The rotation with appropriate offset applied
 	*/
-	UFUNCTION(BlueprintPure, meta = (Category = "Math|Spring|Experimental"))
+	UFUNCTION(BlueprintPure, meta = (Category = "Math|Spring"))
 	static ENGINE_API FRotator SpringInertializeApplyToRotator(FRotator InRotation, FRotator InRotationOffset, FVector InAngularVelocityOffset, float TimeSinceTransition, float SmoothingTime = 0.2f);
 
-	UE_EXPERIMENTAL(5.8, "SpringMath is experimental")
 	/** Transitions a spring inertialization, updating the appropriate offsets.
 	*
 	* @param InOutRotationOffset		Offset rotation to update
@@ -809,10 +744,9 @@ public:
 	* @param TimeSinceTransition		Time since the last transition
 	* @param SmoothingTime				Timescale over which to apply the offset. Larger values result in longer blend times. Can be zero.
 	*/
-	UFUNCTION(BlueprintCallable, meta = (Category = "Math|Spring|Experimental"))
+	UFUNCTION(BlueprintCallable, meta = (Category = "Math|Spring"))
 	static ENGINE_API void SpringInertializeTransitionRotator(UPARAM(ref) FRotator& InOutRotationOffset, UPARAM(ref) FVector& InOutAngularVelocityOffset, FRotator InSrcRotation, FVector InSrcAngularVelocity, FRotator InDstRotation, FVector InDstAngularVelocity, float TimeSinceTransition, float SmoothingTime = 0.2f);
 
-	UE_EXPERIMENTAL(5.8, "SpringMath is experimental")
 	/** Applies a spring inertialization to a angle, adding the appropriate offset given the time since transition.
 	*
 	* @param InAngle					The angle to apply the inertialization to
@@ -822,10 +756,9 @@ public:
 	* @param SmoothingTime				Timescale over which to apply the offset. Larger values result in longer blend times. Can be zero.
 	* @return							The angle with appropriate offset applied
 	*/
-	UFUNCTION(BlueprintPure, meta = (Category = "Math|Spring|Experimental"))
+	UFUNCTION(BlueprintPure, meta = (Category = "Math|Spring"))
 	static ENGINE_API float SpringInertializeApplyToAngle(float InAngle, float InAngleOffset, float InAngularVelocityOffset, float TimeSinceTransition, float SmoothingTime = 0.2f);
 
-	UE_EXPERIMENTAL(5.8, "SpringMath is experimental")
 	/** Transitions a spring inertialization, updating the appropriate offsets.
 	*
 	* @param InOutAngleOffset			Offset angle to update
@@ -837,10 +770,9 @@ public:
 	* @param TimeSinceTransition		Time since the last transition
 	* @param SmoothingTime				Timescale over which to apply the offset. Larger values result in longer blend times. Can be zero.
 	*/
-	UFUNCTION(BlueprintCallable, meta = (Category = "Math|Spring|Experimental"))
+	UFUNCTION(BlueprintCallable, meta = (Category = "Math|Spring"))
 	static ENGINE_API void SpringInertializeTransitionAngle(UPARAM(ref) float& InOutAngleOffset, UPARAM(ref) float& InOutAngularVelocityOffset, float InSrcAngle, float InSrcAngularVelocity, float InDstAngle, float InDstAngularVelocity, float TimeSinceTransition, float SmoothingTime = 0.2f);
 
-	UE_EXPERIMENTAL(5.8, "SpringMath is experimental")
 	/** Applies a spring inertialization to a scale, applying the appropriate offset given the time since transition.
 	*
 	* @param InScale					The scale to apply the inertialization to
@@ -850,10 +782,9 @@ public:
 	* @param SmoothingTime				Timescale over which to apply the offset. Larger values result in longer blend times. Can be zero.
 	* @return							The scale with appropriate offset applied
 	*/
-	UFUNCTION(BlueprintPure, meta = (Category = "Math|Spring|Experimental"))
+	UFUNCTION(BlueprintPure, meta = (Category = "Math|Spring"))
 	static ENGINE_API FVector SpringInertializeApplyToScale(FVector InScale, FVector InScaleOffset, FVector InScalarVelocityOffset, float TimeSinceTransition, float SmoothingTime = 0.2f);
 
-	UE_EXPERIMENTAL(5.8, "SpringMath is experimental")
 	/** Transitions a spring inertialization, updating the appropriate offsets.
 	*
 	* @param InOutScaleOffset			Offset scale to update
@@ -865,10 +796,9 @@ public:
 	* @param TimeSinceTransition		Time since the last transition
 	* @param SmoothingTime				Timescale over which to apply the offset. Larger values result in longer blend times. Can be zero.
 	*/
-	UFUNCTION(BlueprintCallable, meta = (Category = "Math|Spring|Experimental"))
+	UFUNCTION(BlueprintCallable, meta = (Category = "Math|Spring"))
 	static ENGINE_API void SpringInertializeTransitionScale(UPARAM(ref) FVector& InOutScaleOffset, UPARAM(ref) FVector& InOutScalarVelocityOffset, FVector InSrcScale, FVector InSrcScalarVelocity, FVector InDstScale, FVector InDstScalarVelocity, float TimeSinceTransition, float SmoothingTime = 0.2f);
 
-	UE_EXPERIMENTAL(5.8, "SpringMath is experimental")
 	/** Applies a spring inertialization to a transform, applying the appropriate offset given the time since transition.
 	*
 	* @param InTransform				The transform to apply the inertialization to
@@ -878,10 +808,9 @@ public:
 	* @param SmoothingTime				Timescale over which to apply the offset. Larger values result in longer blend times. Can be zero.
 	* @return							The transform with appropriate offset applied
 	*/
-	UFUNCTION(BlueprintPure, meta = (Category = "Math|Spring|Experimental"))
+	UFUNCTION(BlueprintPure, meta = (Category = "Math|Spring"))
 	static ENGINE_API FTransform SpringInertializeApplyToTransform(FTransform InTransform, FTransform InTransformOffset, FSpringTransformVelocity InVelocityOffset, float TimeSinceTransition, float SmoothingTime = 0.2f);
 
-	UE_EXPERIMENTAL(5.8, "SpringMath is experimental")
 	/** Transitions a spring inertialization, updating the appropriate offsets.
 	*
 	* @param InOutTransformOffset		Offset transform to update
@@ -893,12 +822,11 @@ public:
 	* @param TimeSinceTransition		Time since the last transition
 	* @param SmoothingTime				Timescale over which to apply the offset. Larger values result in longer blend times. Can be zero.
 	*/
-	UFUNCTION(BlueprintCallable, meta = (Category = "Math|Spring|Experimental"))
+	UFUNCTION(BlueprintCallable, meta = (Category = "Math|Spring"))
 	static ENGINE_API void SpringInertializeTransitionTransform(UPARAM(ref) FTransform& InOutTransformOffset, UPARAM(ref) FSpringTransformVelocity& InOutVelocityOffset, FTransform InSrcTransform, FSpringTransformVelocity InSrcVelocity, FTransform InDstTransform, FSpringTransformVelocity InDstVelocity, float TimeSinceTransition, float SmoothingTime = 0.2f);
 
 	// Cubic Inertialization //
 
-	UE_EXPERIMENTAL(5.8, "SpringMath is experimental")
 	/** Applies a cubic inertialization to a float value, adding the appropriate offset given the time since transition.
 	*
 	* @param InValue				The value to apply the inertialization to
@@ -908,10 +836,9 @@ public:
 	* @param BlendTime				Timescale over which to apply the offset. Larger values result in longer blend times. Can be zero.
 	* @return						The value with appropriate offset applied
 	*/
-	UFUNCTION(BlueprintPure, meta = (Category = "Math|Spring|Experimental"))
+	UFUNCTION(BlueprintPure, meta = (Category = "Math|Spring"))
 	static ENGINE_API float CubicInertializeApplyToFloat(float InValue, float InValueOffset, float InVelocityOffset, float TimeSinceTransition, float BlendTime = 0.2f);
 
-	UE_EXPERIMENTAL(5.8, "SpringMath is experimental")
 	/** Transitions a cubic inertialization, updating the appropriate offsets.
 	*
 	* @param InOutValueOffset		Offset value to update
@@ -923,10 +850,9 @@ public:
 	* @param TimeSinceTransition	Time since the last transition
 	* @param BlendTime				Timescale over which to apply the offset. Larger values result in longer blend times. Can be zero.
 	*/
-	UFUNCTION(BlueprintCallable, meta = (Category = "Math|Spring|Experimental"))
+	UFUNCTION(BlueprintCallable, meta = (Category = "Math|Spring"))
 	static ENGINE_API void CubicInertializeTransitionFloat(UPARAM(ref) float& InOutValueOffset, UPARAM(ref) float& InOutVelocityOffset, float InSrcValue, float InSrcVelocity, float InDstValue, float InDstVelocity, float TimeSinceTransition, float BlendTime = 0.2f);
 
-	UE_EXPERIMENTAL(5.8, "SpringMath is experimental")
 	/** Applies a cubic inertialization to a vector value, adding the appropriate offset given the time since transition.
 	*
 	* @param InValue				The value to apply the inertialization to
@@ -936,10 +862,9 @@ public:
 	* @param BlendTime				Timescale over which to apply the offset. Larger values result in longer blend times. Can be zero.
 	* @return						The value with appropriate offset applied
 	*/
-	UFUNCTION(BlueprintPure, meta = (Category = "Math|Spring|Experimental"))
+	UFUNCTION(BlueprintPure, meta = (Category = "Math|Spring"))
 	static ENGINE_API FVector CubicInertializeApplyToVector(FVector InValue, FVector InValueOffset, FVector InVelocityOffset, float TimeSinceTransition, float BlendTime = 0.2f);
 
-	UE_EXPERIMENTAL(5.8, "SpringMath is experimental")
 	/** Transitions a cubic inertialization, updating the appropriate offsets.
 	*
 	* @param InOutValueOffset		Offset value to update
@@ -951,10 +876,9 @@ public:
 	* @param TimeSinceTransition	Time since the last transition
 	* @param BlendTime				Timescale over which to apply the offset. Larger values result in longer blend times. Can be zero.
 	*/
-	UFUNCTION(BlueprintCallable, meta = (Category = "Math|Spring|Experimental"))
+	UFUNCTION(BlueprintCallable, meta = (Category = "Math|Spring"))
 	static ENGINE_API void CubicInertializeTransitionVector(UPARAM(ref) FVector& InOutValueOffset, UPARAM(ref) FVector& InOutVelocityOffset, FVector InSrcValue, FVector InSrcVelocity, FVector InDstValue, FVector InDstVelocity, float TimeSinceTransition, float BlendTime = 0.2f);
 
-	UE_EXPERIMENTAL(5.8, "SpringMath is experimental")
 	/** Applies a cubic inertialization to a Vector2D value, adding the appropriate offset given the time since transition.
 	*
 	* @param InValue				The value to apply the inertialization to
@@ -964,10 +888,9 @@ public:
 	* @param BlendTime				Timescale over which to apply the offset. Larger values result in longer blend times. Can be zero.
 	* @return						The value with appropriate offset applied
 	*/
-	UFUNCTION(BlueprintPure, meta = (Category = "Math|Spring|Experimental"))
+	UFUNCTION(BlueprintPure, meta = (Category = "Math|Spring"))
 	static ENGINE_API FVector2D CubicInertializeApplyToVector2D(FVector2D InValue, FVector2D InValueOffset, FVector2D InVelocityOffset, float TimeSinceTransition, float BlendTime = 0.2f);
 
-	UE_EXPERIMENTAL(5.8, "SpringMath is experimental")
 	/** Transitions a cubic inertialization, updating the appropriate offsets.
 	*
 	* @param InOutValueOffset		Offset value to update
@@ -979,10 +902,9 @@ public:
 	* @param TimeSinceTransition	Time since the last transition
 	* @param BlendTime				Timescale over which to apply the offset. Larger values result in longer blend times. Can be zero.
 	*/
-	UFUNCTION(BlueprintCallable, meta = (Category = "Math|Spring|Experimental"))
+	UFUNCTION(BlueprintCallable, meta = (Category = "Math|Spring"))
 	static ENGINE_API void CubicInertializeTransitionVector2D(UPARAM(ref) FVector2D& InOutValueOffset, UPARAM(ref) FVector2D& InOutVelocityOffset, FVector2D InSrcValue, FVector2D InSrcVelocity, FVector2D InDstValue, FVector2D InDstVelocity, float TimeSinceTransition, float BlendTime = 0.2f);
 
-	UE_EXPERIMENTAL(5.8, "SpringMath is experimental")
 	/** Applies a cubic inertialization to a quaternion, applying the appropriate offset given the time since transition.
 	*
 	* @param InRotation					The rotation to apply the inertialization to
@@ -992,10 +914,9 @@ public:
 	* @param BlendTime					Timescale over which to apply the offset. Larger values result in longer blend times. Can be zero.
 	* @return							The rotation with appropriate offset applied
 	*/
-	UFUNCTION(BlueprintPure, meta = (Category = "Math|Spring|Experimental"))
+	UFUNCTION(BlueprintPure, meta = (Category = "Math|Spring"))
 	static ENGINE_API FQuat CubicInertializeApplyToQuat(FQuat InRotation, FQuat InRotationOffset, FVector InAngularVelocityOffset, float TimeSinceTransition, float BlendTime = 0.2f);
 
-	UE_EXPERIMENTAL(5.8, "SpringMath is experimental")
 	/** Transitions a cubic inertialization, updating the appropriate offsets.
 	*
 	* @param InOutRotationOffset		Offset rotation to update
@@ -1007,10 +928,9 @@ public:
 	* @param TimeSinceTransition		Time since the last transition
 	* @param BlendTime					Timescale over which to apply the offset. Larger values result in longer blend times. Can be zero.
 	*/
-	UFUNCTION(BlueprintCallable, meta = (Category = "Math|Spring|Experimental"))
+	UFUNCTION(BlueprintCallable, meta = (Category = "Math|Spring"))
 	static ENGINE_API void CubicInertializeTransitionQuat(UPARAM(ref) FQuat& InOutRotationOffset, UPARAM(ref) FVector& InOutAngularVelocityOffset, FQuat InSrcRotation, FVector InSrcAngularVelocity, FQuat InDstRotation, FVector InDstAngularVelocity, float TimeSinceTransition, float BlendTime = 0.2f);
 
-	UE_EXPERIMENTAL(5.8, "SpringMath is experimental")
 	/** Applies a cubic inertialization to a rotator, applying the appropriate offset given the time since transition.
 	*
 	* @param InRotation					The rotation to apply the inertialization to
@@ -1020,10 +940,9 @@ public:
 	* @param BlendTime					Timescale over which to apply the offset. Larger values result in longer blend times. Can be zero.
 	* @return							The rotation with appropriate offset applied
 	*/
-	UFUNCTION(BlueprintPure, meta = (Category = "Math|Spring|Experimental"))
+	UFUNCTION(BlueprintPure, meta = (Category = "Math|Spring"))
 	static ENGINE_API FRotator CubicInertializeApplyToRotator(FRotator InRotation, FRotator InRotationOffset, FVector InAngularVelocityOffset, float TimeSinceTransition, float BlendTime = 0.2f);
 
-	UE_EXPERIMENTAL(5.8, "SpringMath is experimental")
 	/** Transitions a cubic inertialization, updating the appropriate offsets.
 	*
 	* @param InOutRotationOffset		Offset rotation to update
@@ -1035,10 +954,9 @@ public:
 	* @param TimeSinceTransition		Time since the last transition
 	* @param BlendTime					Timescale over which to apply the offset. Larger values result in longer blend times. Can be zero.
 	*/
-	UFUNCTION(BlueprintCallable, meta = (Category = "Math|Spring|Experimental"))
+	UFUNCTION(BlueprintCallable, meta = (Category = "Math|Spring"))
 	static ENGINE_API void CubicInertializeTransitionRotator(UPARAM(ref) FRotator& InOutRotationOffset, UPARAM(ref) FVector& InOutAngularVelocityOffset, FRotator InSrcRotation, FVector InSrcAngularVelocity, FRotator InDstRotation, FVector InDstAngularVelocity, float TimeSinceTransition, float BlendTime = 0.2f);
 
-	UE_EXPERIMENTAL(5.8, "SpringMath is experimental")
 	/** Applies a cubic inertialization to an angle, applying the appropriate offset given the time since transition.
 	*
 	* @param InAngle					The angle to apply the inertialization to
@@ -1048,10 +966,9 @@ public:
 	* @param BlendTime					Timescale over which to apply the offset. Larger values result in longer blend times. Can be zero.
 	* @return							The angle with appropriate offset applied
 	*/
-	UFUNCTION(BlueprintPure, meta = (Category = "Math|Spring|Experimental"))
+	UFUNCTION(BlueprintPure, meta = (Category = "Math|Spring"))
 	static ENGINE_API float CubicInertializeApplyToAngle(float InAngle, float InAngleOffset, float InAngularVelocityOffset, float TimeSinceTransition, float BlendTime = 0.2f);
 
-	UE_EXPERIMENTAL(5.8, "SpringMath is experimental")
 	/** Transitions a cubic inertialization, updating the appropriate offsets.
 	*
 	* @param InOutAngleOffset			Offset angle to update
@@ -1063,10 +980,9 @@ public:
 	* @param TimeSinceTransition		Time since the last transition
 	* @param BlendTime					Timescale over which to apply the offset. Larger values result in longer blend times. Can be zero.
 	*/
-	UFUNCTION(BlueprintCallable, meta = (Category = "Math|Spring|Experimental"))
+	UFUNCTION(BlueprintCallable, meta = (Category = "Math|Spring"))
 	static ENGINE_API void CubicInertializeTransitionAngle(UPARAM(ref) float& InOutAngleOffset, UPARAM(ref) float& InOutAngularVelocityOffset, float InSrcAngle, float InSrcAngularVelocity, float InDstAngle, float InDstAngularVelocity, float TimeSinceTransition, float BlendTime = 0.2f);
 
-	UE_EXPERIMENTAL(5.8, "SpringMath is experimental")
 	/** Applies a cubic inertialization to a scale, applying the appropriate offset given the time since transition.
 	*
 	* @param InScale					The scale to apply the inertialization to
@@ -1076,10 +992,9 @@ public:
 	* @param BlendTime					Timescale over which to apply the offset. Larger values result in longer blend times. Can be zero.
 	* @return							The scale with appropriate offset applied
 	*/
-	UFUNCTION(BlueprintPure, meta = (Category = "Math|Spring|Experimental"))
+	UFUNCTION(BlueprintPure, meta = (Category = "Math|Spring"))
 	static ENGINE_API FVector CubicInertializeApplyToScale(FVector InScale, FVector InScaleOffset, FVector InScalarVelocityOffset, float TimeSinceTransition, float BlendTime = 0.2f);
 
-	UE_EXPERIMENTAL(5.8, "SpringMath is experimental")
 	/** Transitions a cubic inertialization, updating the appropriate offsets.
 	*
 	* @param InOutScaleOffset			Offset scale to update
@@ -1091,10 +1006,9 @@ public:
 	* @param TimeSinceTransition		Time since the last transition
 	* @param BlendTime					Timescale over which to apply the offset. Larger values result in longer blend times. Can be zero.
 	*/
-	UFUNCTION(BlueprintCallable, meta = (Category = "Math|Spring|Experimental"))
+	UFUNCTION(BlueprintCallable, meta = (Category = "Math|Spring"))
 	static ENGINE_API void CubicInertializeTransitionScale(UPARAM(ref) FVector& InOutScaleOffset, UPARAM(ref) FVector& InOutScalarVelocityOffset, FVector InSrcScale, FVector InSrcScalarVelocity, FVector InDstScale, FVector InDstScalarVelocity, float TimeSinceTransition, float BlendTime = 0.2f);
 
-	UE_EXPERIMENTAL(5.8, "SpringMath is experimental")
 	/** Applies a cubic inertialization to a transform, applying the appropriate offset given the time since transition.
 	*
 	* @param InTransform				The transform to apply the inertialization to
@@ -1104,10 +1018,9 @@ public:
 	* @param BlendTime					Timescale over which to apply the offset. Larger values result in longer blend times. Can be zero.
 	* @return							The transform with appropriate offset applied
 	*/
-	UFUNCTION(BlueprintPure, meta = (Category = "Math|Spring|Experimental"))
+	UFUNCTION(BlueprintPure, meta = (Category = "Math|Spring"))
 	static ENGINE_API FTransform CubicInertializeApplyToTransform(FTransform InTransform, FTransform InTransformOffset, FSpringTransformVelocity InVelocityOffset, float TimeSinceTransition, float BlendTime = 0.2f);
 
-	UE_EXPERIMENTAL(5.8, "SpringMath is experimental")
 	/** Transitions a cubic inertialization, updating the appropriate offsets.
 	*
 	* @param InOutTransformOffset		Offset transform to update
@@ -1119,12 +1032,11 @@ public:
 	* @param TimeSinceTransition		Time since the last transition
 	* @param BlendTime					Timescale over which to apply the offset. Larger values result in longer blend times. Can be zero.
 	*/
-	UFUNCTION(BlueprintCallable, meta = (Category = "Math|Spring|Experimental"))
+	UFUNCTION(BlueprintCallable, meta = (Category = "Math|Spring"))
 	static ENGINE_API void CubicInertializeTransitionTransform(UPARAM(ref) FTransform& InOutTransformOffset, UPARAM(ref) FSpringTransformVelocity& InOutVelocityOffset, FTransform InSrcTransform, FSpringTransformVelocity InSrcVelocity, FTransform InDstTransform, FSpringTransformVelocity InDstVelocity, float TimeSinceTransition, float BlendTime = 0.2f);
 
 	// Dead Blending //
 
-	UE_EXPERIMENTAL(5.8, "SpringMath is experimental")
 	/** Applies a dead blend to a float value, blending it with an extrapolation of the value and velocity at the point of the last transition.
 	*
 	* @param InValue					The value to apply the dead blend to
@@ -1135,10 +1047,9 @@ public:
 	* @param SmoothingTime				Decay rate of the extrapolated velocity. Smaller values make the velocity decay faster.
 	* @return							The blended value
 	*/
-	UFUNCTION(BlueprintPure, meta = (Category = "Math|Spring|Experimental"))
+	UFUNCTION(BlueprintPure, meta = (Category = "Math|Spring"))
 	static ENGINE_API float DeadBlendApplyToFloat(float InValue, float InValueTransition, float InVelocityTransition, float TimeSinceTransition, float BlendTime = 0.2f, float SmoothingTime = 0.2f);
 
-	UE_EXPERIMENTAL(5.8, "SpringMath is experimental")
 	/** Transitions a dead blend, computing the new value and velocity at the point of transition.
 	*
 	* @param InOutValueTransition		Value at the point of transition
@@ -1149,10 +1060,9 @@ public:
 	* @param BlendTime					Timescale over which to apply the offset. Larger values result in longer blend times. Can be zero.
 	* @param SmoothingTime				Decay rate of the extrapolated velocity. Smaller values make the velocity decay faster.
 	*/
-	UFUNCTION(BlueprintCallable, meta = (Category = "Math|Spring|Experimental"))
+	UFUNCTION(BlueprintCallable, meta = (Category = "Math|Spring"))
 	static ENGINE_API void DeadBlendTransitionFloat(UPARAM(ref) float& InOutValueTransition, UPARAM(ref) float& InOutVelocityTransition, float InSrcValue, float InSrcVelocity, float TimeSinceTransition, float BlendTime = 0.2f, float SmoothingTime = 0.2f);
 
-	UE_EXPERIMENTAL(5.8, "SpringMath is experimental")
 	/** Applies a dead blend to a vector value, blending it with an extrapolation of the value and velocity at the point of the last transition.
 	*
 	* @param InValue					The value to apply the dead blend to
@@ -1163,10 +1073,9 @@ public:
 	* @param SmoothingTime				Decay rate of the extrapolated velocity. Smaller values make the velocity decay faster.
 	* @return							The blended value
 	*/
-	UFUNCTION(BlueprintPure, meta = (Category = "Math|Spring|Experimental"))
+	UFUNCTION(BlueprintPure, meta = (Category = "Math|Spring"))
 	static ENGINE_API FVector DeadBlendApplyToVector(FVector InValue, FVector InValueTransition, FVector InVelocityTransition, float TimeSinceTransition, float BlendTime = 0.2f, float SmoothingTime = 0.2f);
 
-	UE_EXPERIMENTAL(5.8, "SpringMath is experimental")
 	/** Transitions a dead blend, computing the new value and velocity at the point of transition.
 	*
 	* @param InOutValueTransition		Value at the point of transition
@@ -1177,10 +1086,9 @@ public:
 	* @param BlendTime					Timescale over which to apply the offset. Larger values result in longer blend times. Can be zero.
 	* @param SmoothingTime				Decay rate of the extrapolated velocity. Smaller values make the velocity decay faster.
 	*/
-	UFUNCTION(BlueprintCallable, meta = (Category = "Math|Spring|Experimental"))
+	UFUNCTION(BlueprintCallable, meta = (Category = "Math|Spring"))
 	static ENGINE_API void DeadBlendTransitionVector(UPARAM(ref) FVector& InOutValueTransition, UPARAM(ref) FVector& InOutVelocityTransition, FVector InSrcValue, FVector InSrcVelocity, float TimeSinceTransition, float BlendTime = 0.2f, float SmoothingTime = 0.2f);
 
-	UE_EXPERIMENTAL(5.8, "SpringMath is experimental")
 	/** Applies a dead blend to a Vector2D value, blending it with an extrapolation of the value and velocity at the point of the last transition.
 	*
 	* @param InValue					The value to apply the dead blend to
@@ -1191,10 +1099,9 @@ public:
 	* @param SmoothingTime				Decay rate of the extrapolated velocity. Smaller values make the velocity decay faster.
 	* @return							The blended value
 	*/
-	UFUNCTION(BlueprintPure, meta = (Category = "Math|Spring|Experimental"))
+	UFUNCTION(BlueprintPure, meta = (Category = "Math|Spring"))
 	static ENGINE_API FVector2D DeadBlendApplyToVector2D(FVector2D InValue, FVector2D InValueTransition, FVector2D InVelocityTransition, float TimeSinceTransition, float BlendTime = 0.2f, float SmoothingTime = 0.2f);
 
-	UE_EXPERIMENTAL(5.8, "SpringMath is experimental")
 	/** Transitions a dead blend, computing the new value and velocity at the point of transition.
 	*
 	* @param InOutValueTransition		Value at the point of transition
@@ -1205,10 +1112,9 @@ public:
 	* @param BlendTime					Timescale over which to apply the offset. Larger values result in longer blend times. Can be zero.
 	* @param SmoothingTime				Decay rate of the extrapolated velocity. Smaller values make the velocity decay faster.
 	*/
-	UFUNCTION(BlueprintCallable, meta = (Category = "Math|Spring|Experimental"))
+	UFUNCTION(BlueprintCallable, meta = (Category = "Math|Spring"))
 	static ENGINE_API void DeadBlendTransitionVector2D(UPARAM(ref) FVector2D& InOutValueTransition, UPARAM(ref) FVector2D& InOutVelocityTransition, FVector2D InSrcValue, FVector2D InSrcVelocity, float TimeSinceTransition, float BlendTime = 0.2f, float SmoothingTime = 0.2f);
 
-	UE_EXPERIMENTAL(5.8, "SpringMath is experimental")
 	/** Applies a dead blend to a quaternion, blending it with an extrapolation of the rotation and velocity at the point of the last transition.
 	*
 	* @param InRotation						The rotation to apply the dead blend to
@@ -1219,10 +1125,9 @@ public:
 	* @param SmoothingTime					Decay rate of the extrapolated velocity. Smaller values make the velocity decay faster.
 	* @return								The blended rotation
 	*/
-	UFUNCTION(BlueprintPure, meta = (Category = "Math|Spring|Experimental"))
+	UFUNCTION(BlueprintPure, meta = (Category = "Math|Spring"))
 	static ENGINE_API FQuat DeadBlendApplyToQuat(FQuat InRotation, FQuat InRotationTransition, FVector InAngularVelocityTransition, float TimeSinceTransition, float BlendTime = 0.2f, float SmoothingTime = 0.2f);
 
-	UE_EXPERIMENTAL(5.8, "SpringMath is experimental")
 	/** Transitions a dead blend, computing the new value and velocity at the point of transition.
 	*
 	* @param InOutRotationTransition		Rotation at the point of transition
@@ -1233,10 +1138,9 @@ public:
 	* @param BlendTime						Timescale over which to apply the offset. Larger values result in longer blend times. Can be zero.
 	* @param SmoothingTime					Decay rate of the extrapolated velocity. Smaller values make the velocity decay faster.
 	*/
-	UFUNCTION(BlueprintCallable, meta = (Category = "Math|Spring|Experimental"))
+	UFUNCTION(BlueprintCallable, meta = (Category = "Math|Spring"))
 	static ENGINE_API void DeadBlendTransitionQuat(UPARAM(ref) FQuat& InOutRotationTransition, UPARAM(ref) FVector& InOutAngularVelocityTransition, FQuat InSrcRotation, FVector InSrcAngularVelocity, float TimeSinceTransition, float BlendTime = 0.2f, float SmoothingTime = 0.2f);
 
-	UE_EXPERIMENTAL(5.8, "SpringMath is experimental")
 	/** Applies a dead blend to a rotator, blending it with an extrapolation of the rotation and velocity at the point of the last transition.
 	*
 	* @param InRotation						The rotation to apply the dead blend to
@@ -1247,10 +1151,9 @@ public:
 	* @param SmoothingTime					Decay rate of the extrapolated velocity. Smaller values make the velocity decay faster.
 	* @return								The blended rotation
 	*/
-	UFUNCTION(BlueprintPure, meta = (Category = "Math|Spring|Experimental"))
+	UFUNCTION(BlueprintPure, meta = (Category = "Math|Spring"))
 	static ENGINE_API FRotator DeadBlendApplyToRotator(FRotator InRotation, FRotator InRotationTransition, FVector InAngularVelocityTransition, float TimeSinceTransition, float BlendTime = 0.2f, float SmoothingTime = 0.2f);
 
-	UE_EXPERIMENTAL(5.8, "SpringMath is experimental")
 	/** Transitions a dead blend, computing the new value and velocity at the point of transition.
 	*
 	* @param InOutRotationTransition		Rotation at the point of transition
@@ -1261,10 +1164,9 @@ public:
 	* @param BlendTime						Timescale over which to apply the offset. Larger values result in longer blend times. Can be zero.
 	* @param SmoothingTime					Decay rate of the extrapolated velocity. Smaller values make the velocity decay faster.
 	*/
-	UFUNCTION(BlueprintCallable, meta = (Category = "Math|Spring|Experimental"))
+	UFUNCTION(BlueprintCallable, meta = (Category = "Math|Spring"))
 	static ENGINE_API void DeadBlendTransitionRotator(UPARAM(ref) FRotator& InOutRotationTransition, UPARAM(ref) FVector& InOutAngularVelocityTransition, FRotator InSrcRotation, FVector InSrcAngularVelocity, float TimeSinceTransition, float BlendTime = 0.2f, float SmoothingTime = 0.2f);
 
-	UE_EXPERIMENTAL(5.8, "SpringMath is experimental")
 	/** Applies a dead blend to an angle, blending it with an extrapolation of the angle and velocity at the point of the last transition.
 	*
 	* @param InAngle						The angle to apply the dead blend to
@@ -1275,10 +1177,9 @@ public:
 	* @param SmoothingTime					Decay rate of the extrapolated velocity. Smaller values make the velocity decay faster.
 	* @return								The blended angle
 	*/
-	UFUNCTION(BlueprintPure, meta = (Category = "Math|Spring|Experimental"))
+	UFUNCTION(BlueprintPure, meta = (Category = "Math|Spring"))
 	static ENGINE_API float DeadBlendApplyToAngle(float InAngle, float InAngleTransition, float InAngularVelocityTransition, float TimeSinceTransition, float BlendTime = 0.2f, float SmoothingTime = 0.2f);
 
-	UE_EXPERIMENTAL(5.8, "SpringMath is experimental")
 	/** Transitions a dead blend, computing the new value and velocity at the point of transition.
 	*
 	* @param InOutAngleTransition			Angle at the point of transition
@@ -1289,10 +1190,9 @@ public:
 	* @param BlendTime						Timescale over which to apply the offset. Larger values result in longer blend times. Can be zero.
 	* @param SmoothingTime					Decay rate of the extrapolated velocity. Smaller values make the velocity decay faster.
 	*/
-	UFUNCTION(BlueprintCallable, meta = (Category = "Math|Spring|Experimental"))
+	UFUNCTION(BlueprintCallable, meta = (Category = "Math|Spring"))
 	static ENGINE_API void DeadBlendTransitionAngle(UPARAM(ref) float& InOutAngleTransition, UPARAM(ref) float& InOutAngularVelocityTransition, float InSrcAngle, float InSrcAngularVelocity, float TimeSinceTransition, float BlendTime = 0.2f, float SmoothingTime = 0.2f);
 
-	UE_EXPERIMENTAL(5.8, "SpringMath is experimental")
 	/** Applies a dead blend to a scale, blending it with an extrapolation of the scale and velocity at the point of the last transition.
 	*
 	* @param InScale						The scale to apply the dead blend to
@@ -1303,10 +1203,9 @@ public:
 	* @param SmoothingTime					Decay rate of the extrapolated velocity. Smaller values make the velocity decay faster.
 	* @return								The blended scale
 	*/
-	UFUNCTION(BlueprintPure, meta = (Category = "Math|Spring|Experimental"))
+	UFUNCTION(BlueprintPure, meta = (Category = "Math|Spring"))
 	static ENGINE_API FVector DeadBlendApplyToScale(FVector InScale, FVector InScaleTransition, FVector InScalarVelocityTransition, float TimeSinceTransition, float BlendTime = 0.2f, float SmoothingTime = 0.2f);
 
-	UE_EXPERIMENTAL(5.8, "SpringMath is experimental")
 	/** Transitions a dead blend, computing the new value and velocity at the point of transition.
 	*
 	* @param InOutScaleTransition			Scale at the point of transition
@@ -1317,10 +1216,9 @@ public:
 	* @param BlendTime						Timescale over which to apply the offset. Larger values result in longer blend times. Can be zero.
 	* @param SmoothingTime					Decay rate of the extrapolated velocity. Smaller values make the velocity decay faster.
 	*/
-	UFUNCTION(BlueprintCallable, meta = (Category = "Math|Spring|Experimental"))
+	UFUNCTION(BlueprintCallable, meta = (Category = "Math|Spring"))
 	static ENGINE_API void DeadBlendTransitionScale(UPARAM(ref) FVector& InOutScaleTransition, UPARAM(ref) FVector& InOutScalarVelocityTransition, FVector InSrcScale, FVector InSrcScalarVelocity, float TimeSinceTransition, float BlendTime = 0.2f, float SmoothingTime = 0.2f);
 
-	UE_EXPERIMENTAL(5.8, "SpringMath is experimental")
 	/** Applies a dead blend to a transform, blending it with an extrapolation of the transform and velocity at the point of the last transition.
 	*
 	* @param InTransform					The transform to apply the dead blend to
@@ -1331,10 +1229,9 @@ public:
 	* @param SmoothingTime					Decay rate of the extrapolated velocity. Smaller values make the velocity decay faster.
 	* @return								The blended transform
 	*/
-	UFUNCTION(BlueprintPure, meta = (Category = "Math|Spring|Experimental"))
+	UFUNCTION(BlueprintPure, meta = (Category = "Math|Spring"))
 	static ENGINE_API FTransform DeadBlendApplyToTransform(FTransform InTransform, FTransform InTransformTransition, FSpringTransformVelocity InVelocityTransition, float TimeSinceTransition, float BlendTime = 0.2f, float SmoothingTime = 0.2f);
 
-	UE_EXPERIMENTAL(5.8, "SpringMath is experimental")
 	/** Transitions a dead blend, computing the new value and velocity at the point of transition.
 	*
 	* @param InOutTransformTransition	Transform at the point of transition
@@ -1345,7 +1242,7 @@ public:
 	* @param BlendTime					Timescale over which to apply the offset. Larger values result in longer blend times. Can be zero.
 	* @param SmoothingTime				Decay rate of the extrapolated velocity. Smaller values make the velocity decay faster.
 	*/
-	UFUNCTION(BlueprintCallable, meta = (Category = "Math|Spring|Experimental"))
+	UFUNCTION(BlueprintCallable, meta = (Category = "Math|Spring"))
 	static ENGINE_API void DeadBlendTransitionTransform(UPARAM(ref) FTransform& InOutTransformTransition, UPARAM(ref) FSpringTransformVelocity& InOutVelocityTransition, FTransform InSrcTransform, FSpringTransformVelocity InSrcVelocity, float TimeSinceTransition, float BlendTime = 0.2f, float SmoothingTime = 0.2f);
 
 };
