Editor: Fix crash when dismissing notifications after hot-reload

#rn Fixed crash when dismissing notifications after hot-reload
#Synced-CL 52148822

diff --git a/Engine/Source/Runtime/Engine/Private/RendererSettings.cpp b/Engine/Source/Runtime/Engine/Private/RendererSettings.cpp
--- a/Engine/Source/Runtime/Engine/Private/RendererSettings.cpp
+++ b/Engine/Source/Runtime/Engine/Private/RendererSettings.cpp
@@ -553,13 +553,16 @@ void URendererSettings::CheckForMissingShaderModels()
 
 		if (bProjectMissingD3DSM6 || bProjectMissingWindowsVulkanSM6 || bProjectMissingLinuxVulkanSM6 || bProjectMissingMacSM6)
 		{
-			auto DismissNotification = [this]()
+			auto DismissNotification = [This = TStrongObjectPtr(this)]
 			{
-				if (TSharedPtr<SNotificationItem> NotificationPin = ShaderModelNotificationPtr.Pin())
+				if (URendererSettings* Settings = This.Get())
 				{
-					NotificationPin->SetCompletionState(SNotificationItem::CS_None);
-					NotificationPin->ExpireAndFadeout();
-					ShaderModelNotificationPtr.Reset();
+					if (TSharedPtr<SNotificationItem> NotificationPin = Settings->ShaderModelNotificationPtr.Pin())
+					{
+						NotificationPin->SetCompletionState(SNotificationItem::CS_None);
+						NotificationPin->ExpireAndFadeout();
+						Settings->ShaderModelNotificationPtr.Reset();
+					}
 				}
 			};
 
diff --git a/Engine/Source/Runtime/Engine/Private/StaticMesh.cpp b/Engine/Source/Runtime/Engine/Private/StaticMesh.cpp
--- a/Engine/Source/Runtime/Engine/Private/StaticMesh.cpp
+++ b/Engine/Source/Runtime/Engine/Private/StaticMesh.cpp
@@ -8472,13 +8472,16 @@ void UStaticMesh::CheckForMissingShaderModels()
 		{
 			bWarnedAboutMissingShaderModel = true;
 
-			auto DismissNotification = [this]()
+			auto DismissNotification = [This = TStrongObjectPtr(this)]
 			{
-				if (TSharedPtr<SNotificationItem> NotificationPin = ShaderModelNotificationPtr.Pin())
+				if (UStaticMesh* Mesh = This.Get())
 				{
-					NotificationPin->SetCompletionState(SNotificationItem::CS_None);
-					NotificationPin->ExpireAndFadeout();
-					ShaderModelNotificationPtr.Reset();
+					if (TSharedPtr<SNotificationItem> NotificationPin = Mesh->ShaderModelNotificationPtr.Pin())
+					{
+						NotificationPin->SetCompletionState(SNotificationItem::CS_None);
+						NotificationPin->ExpireAndFadeout();
+						Mesh->ShaderModelNotificationPtr.Reset();
+					}
 				}
 			};
 
