Slate: Expose function for borderless window style

#rn Added FWindowStyle::GetBorderless() to expose style used by UGameEngine default window
#Synced-CL 52223014

diff --git a/Engine/Source/Runtime/Engine/Private/GameEngine.cpp b/Engine/Source/Runtime/Engine/Private/GameEngine.cpp
--- a/Engine/Source/Runtime/Engine/Private/GameEngine.cpp
+++ b/Engine/Source/Runtime/Engine/Private/GameEngine.cpp
@@ -661,16 +661,6 @@ TSharedRef<SWindow> UGameEngine::CreateGameWindow()
 		MaxWindowHeight = FMath::Max(DisplayMetrics.VirtualDisplayRect.Bottom - DisplayMetrics.VirtualDisplayRect.Top, ResY);
 	}
 
-	static FWindowStyle BorderlessStyle = FWindowStyle::GetDefault();
-	BorderlessStyle
-		.SetActiveTitleBrush(FSlateNoResource())
-		.SetInactiveTitleBrush(FSlateNoResource())
-		.SetFlashTitleBrush(FSlateNoResource())
-		.SetOutlineBrush(FSlateNoResource())
-		.SetBorderBrush(FSlateNoResource())
-		.SetBackgroundBrush(FSlateNoResource())
-		.SetChildBackgroundBrush(FSlateNoResource());
-
 	if (MonitorIndex != -1)
 	{
 		FDisplayMetrics DisplayMetrics;
@@ -709,7 +699,7 @@ TSharedRef<SWindow> UGameEngine::CreateGameWindow()
 	}
 
 	TSharedRef<SWindow> Window = SNew(SWindow)
-	.Style(bUseBorderlessWindow ? &BorderlessStyle : &FCoreStyle::Get().GetWidgetStyle<FWindowStyle>("Window"))
+	.Style(bUseBorderlessWindow ? &FWindowStyle::GetBorderless() : &FCoreStyle::Get().GetWidgetStyle<FWindowStyle>("Window"))
 	.ClientSize(FVector2D(ResX, ResY))
 	.AdjustInitialSizeAndPositionForDPIScale(false)
 	.Title(WindowTitle)
diff --git a/Engine/Source/Runtime/SlateCore/Private/Styling/SlateTypes.cpp b/Engine/Source/Runtime/SlateCore/Private/Styling/SlateTypes.cpp
--- a/Engine/Source/Runtime/SlateCore/Private/Styling/SlateTypes.cpp
+++ b/Engine/Source/Runtime/SlateCore/Private/Styling/SlateTypes.cpp
@@ -838,6 +838,24 @@ const FWindowStyle& FWindowStyle::GetDefault()
 	return Default;
 }
 
+const FWindowStyle& FWindowStyle::GetBorderless()
+{
+	static FWindowStyle BorderlessStyle = []
+	{
+		FWindowStyle NewBorderlessStyle;
+		NewBorderlessStyle
+			.SetActiveTitleBrush(FSlateNoResource())
+			.SetInactiveTitleBrush(FSlateNoResource())
+			.SetFlashTitleBrush(FSlateNoResource())
+			.SetOutlineBrush(FSlateNoResource())
+			.SetBorderBrush(FSlateNoResource())
+			.SetBackgroundBrush(FSlateNoResource())
+			.SetChildBackgroundBrush(FSlateNoResource());
+		return NewBorderlessStyle;
+	}();
+	return BorderlessStyle;
+}
+
 void FInvalidatableBrushAttribute::SetImage(SWidget& ThisWidget, const TAttribute< const FSlateBrush* >& InImage)
 {
 	const bool bImagePointerChanged = SetWidgetAttribute(ThisWidget, Image, InImage, EInvalidateWidgetReason::Layout);
diff --git a/Engine/Source/Runtime/SlateCore/Public/Styling/SlateTypes.h b/Engine/Source/Runtime/SlateCore/Public/Styling/SlateTypes.h
--- a/Engine/Source/Runtime/SlateCore/Public/Styling/SlateTypes.h
+++ b/Engine/Source/Runtime/SlateCore/Public/Styling/SlateTypes.h
@@ -2148,6 +2148,8 @@ struct FWindowStyle : public FSlateWidgetStyle
 
 	static SLATECORE_API const FWindowStyle& GetDefault();
 
+	static SLATECORE_API const FWindowStyle& GetBorderless();
+
 	/** Style used to draw the window minimize button */
 	UPROPERTY(EditAnywhere, BlueprintReadWrite, Category=Appearance)
 	FButtonStyle MinimizeButtonStyle;
