Fix Dynamic Translucency Resolution (DynTranslucencyRes) display in Stat Unit and other dynamic resolution states not using the correct resolution fraction.
Fix Dynamic Translucency current fraction in Stat RenderScaling using max current fraction instead of current fraction

#Synced-CL 52032250

diff --git a/Engine/Source/Runtime/Engine/Private/DynamicResolution.cpp b/Engine/Source/Runtime/Engine/Private/DynamicResolution.cpp
--- a/Engine/Source/Runtime/Engine/Private/DynamicResolution.cpp
+++ b/Engine/Source/Runtime/Engine/Private/DynamicResolution.cpp
@@ -482,7 +482,7 @@ void FDynamicResolutionHeuristicProxy::RefreshHeuristicStats_RenderThread()
 			// MeasuredMs is set in RefreshCurrentFrameResolutionFraction_RenderThread()
 			SET_FLOAT_STAT_FName(Budget.GetStatId_MinScaling().GetName(), HeuristicSettings.MinResolutionFraction);
 			SET_FLOAT_STAT_FName(Budget.GetStatId_MaxScaling().GetName(), HeuristicSettings.MaxResolutionFraction);
-			SET_FLOAT_STAT_FName(Budget.GetStatId_CurrentScaling().GetName(), CurrentFrameMaxResolutionFractions[Budget]);
+			SET_FLOAT_STAT_FName(Budget.GetStatId_CurrentScaling().GetName(), CurrentFrameResolutionFractions[Budget]);
 		}
 		else
 		{
diff --git a/Engine/Source/Runtime/Engine/Private/UnrealClient.cpp b/Engine/Source/Runtime/Engine/Private/UnrealClient.cpp
--- a/Engine/Source/Runtime/Engine/Private/UnrealClient.cpp
+++ b/Engine/Source/Runtime/Engine/Private/UnrealClient.cpp
@@ -898,6 +898,9 @@ int32 FStatUnitData::DrawStat(FViewport* InViewport, FCanvas* InCanvas, int32 In
 					DisplayName.ReplaceInline(TEXT("Dynamic"), TEXT("Dyn"));
 					DisplayName.ReplaceInline(TEXT("Resolution"), TEXT("Res"));
 
+					ResolutionFraction = DynamicResolutionStateInfos.ResolutionFractionApproximations[Budget];
+					ScreenPercentage = ResolutionFraction * 100.0f;
+
 					FColor Color = (ResolutionFraction < AlertResolutionFraction) ? StatRed : ((ResolutionFraction < FMath::Min(ResolutionFraction * 0.97f, 1.0f)) ? StatOrange : StatGreen);
 
 					DrawTitleString(*DisplayName, /* UnitGraphColor = */ FColor::White);
