[HWRT] Reexpose UPrimitiveComponent bRayTracingFarField flag in editor.
- updated description to clarify that far field assignment is exclusive to those components are never included in the near field.

#Synced-CL 52227391

diff --git a/Engine/Source/Runtime/Engine/Classes/Components/PrimitiveComponent.h b/Engine/Source/Runtime/Engine/Classes/Components/PrimitiveComponent.h
--- a/Engine/Source/Runtime/Engine/Classes/Components/PrimitiveComponent.h
+++ b/Engine/Source/Runtime/Engine/Classes/Components/PrimitiveComponent.h
@@ -698,8 +698,8 @@ public:
 	UPROPERTY(EditAnywhere, AdvancedDisplay, BlueprintReadOnly, Category = Rendering, meta = (DisplayName = "Hidden In Scene Capture", ToolTip = "When true, will not be captured by Scene Capture"))
 	uint8 bHiddenInSceneCapture : 1;
 
-	/** If true, this component will be available to ray trace as a far field primitive even if hidden. */
-	UPROPERTY()
+	/** If true, this component will be included in the ray tracing far field layer instead of the near field. The ray tracing far field layer ignores whether the component is hidden. */
+	UPROPERTY(EditAnywhere, AdvancedDisplay, BlueprintReadOnly, Category = RayTracing)
 	uint8 bRayTracingFarField : 1;
 
 	/** If true, a single top-down Lumen card is allocated */
