[AsyncLoadingTests]
- Fix redirector leaks in memory and on disk causing some tests to fail depending on execution order or how many time they were executed

#Synced-CL 52203528

diff --git a/Engine/Source/Runtime/Engine/Private/Tests/Loading/AsyncLoadingTests_Shared.cpp b/Engine/Source/Runtime/Engine/Private/Tests/Loading/AsyncLoadingTests_Shared.cpp
--- a/Engine/Source/Runtime/Engine/Private/Tests/Loading/AsyncLoadingTests_Shared.cpp
+++ b/Engine/Source/Runtime/Engine/Private/Tests/Loading/AsyncLoadingTests_Shared.cpp
@@ -1,8 +1,10 @@
 // Copyright Epic Games, Inc. All Rights Reserved.
 
 #include "Tests/Loading/AsyncLoadingTests_Shared.h"
+#include "HAL/FileManager.h"
 #include "Misc/AutomationTest.h"
 #include "Misc/PackageName.h"
+#include "Misc/RedirectCollector.h"
 #include "UObject/MetaData.h"
 #include "UObject/SavePackage.h"
 #include "AssetRegistry/AssetData.h"
@@ -121,8 +123,26 @@ void FLoadingTestsScope::LoadObjects()
 
 void FLoadingTestsScope::CleanupObjects()
 {
+	// Remove any asset path redirections that were registered for test packages
+	// (e.g., from ObjectRedirectors created by Rename operations and picked up by the asset registry).
+	// These persist in GRedirectCollector across tests and would corrupt soft object path resolution.
+#if WITH_EDITOR
+	for (const FString& PackageName : PackageNames)
+	{
+		FSoftObjectPath ObjectPath(PackageName + TEXT(".") + ObjectName);
+		GRedirectCollector.RemoveAssetPathRedirection(ObjectPath);
+	}
+#endif
+
 	GarbageCollect();
 
+	// Delete on-disk package files to avoid stale assets leaking between tests
+	for (const FString& PackageName : PackageNames)
+	{
+		FString FilePath = FPackageName::LongPackageNameToFilename(PackageName, FPackageName::GetAssetPackageExtension());
+		IFileManager::Get().Delete(*FilePath);
+	}
+
 	UAsyncLoadingTests_Shared::OnConstruct.Unbind();
 	UAsyncLoadingTests_Shared::OnPostLoad.Unbind();
 	UAsyncLoadingTests_Shared::OnSerialize.Unbind();
