Chaos SKM Cloth Asset:
- Fixed missing DDC invalidation when binding a Cloth Asset that was causing the simulation mapping to go missing.
- Fixed a race condition in the binding code when loading and async building the Skeletal Mesh.
[RN] Chaos SKM Cloth Asset: Fixed missing DDC invalidation when binding a Cloth Asset that was causing the simulation mapping to go missing. Fixed a race condition in the binding code when loading and async building the Skeletal Mesh.
#Synced-CL 52128657

diff --git a/Engine/Plugins/ChaosClothAsset/Source/ChaosClothAssetEngine/Private/ChaosClothAsset/ClothAssetSKMClothingAsset.cpp b/Engine/Plugins/ChaosClothAsset/Source/ChaosClothAssetEngine/Private/ChaosClothAsset/ClothAssetSKMClothingAsset.cpp
--- a/Engine/Plugins/ChaosClothAsset/Source/ChaosClothAssetEngine/Private/ChaosClothAsset/ClothAssetSKMClothingAsset.cpp
+++ b/Engine/Plugins/ChaosClothAsset/Source/ChaosClothAssetEngine/Private/ChaosClothAsset/ClothAssetSKMClothingAsset.cpp
@@ -247,6 +247,11 @@ void UChaosClothAssetSKMClothingAsset::PostEditChangeProperty(FPropertyChangedEv
 	}
 }
 
+USkinnedAsset* UChaosClothAssetSKMClothingAsset::GetSkinnedAssetDependency() const
+{
+	return const_cast<UChaosClothAssetBase*>(Asset.Get());
+}
+
 bool UChaosClothAssetSKMClothingAsset::BindToSkeletalMesh(USkeletalMesh* SkeletalMesh, const int32 MeshLodIndex, const int32 SectionIndex, const int32 AssetLodIndex)
 {
 	using namespace UE::Chaos::ClothAsset;
@@ -268,6 +273,11 @@ bool UChaosClothAssetSKMClothingAsset::BindToSkeletalMesh(USkeletalMesh* Skeleta
 	FSkeletalMeshLODModel& SkeletalMeshLODModel = SkeletalMesh->GetImportedModel()->LODModels[MeshLodIndex];
 	FSkelMeshSection& SkelMeshSection = SkeletalMeshLODModel.Sections[SectionIndex];
 
+	// Detect a partial binding state
+	// The GUID is already set but no mapping data was ever generated (e.g. if Asset was null on the first bind).
+	// Assigning the same GUID will not change the DDC key, so an explicit invalidation is done at the end of this function.
+	const bool bWasPartiallyBound = (SkelMeshSection.ClothingData.AssetGuid == AssetGuid) && !SkelMeshSection.HasClothingData();
+
 	// Clear the proxy deformer mappings data now in case the cloth simulation model is invalid
 	SkelMeshSection.ClothMappingDataLODs.Reset();
 
@@ -467,6 +477,14 @@ bool UChaosClothAssetSKMClothingAsset::BindToSkeletalMesh(USkeletalMesh* Skeleta
 	// Update the extra cloth deformer mapping LOD bias using this cloth entry
 	ClothingAssetUtils::UpdateLODBiasMappings(*this, SkeletalMesh, MeshLodIndex, SectionIndex);
 
+	// Invalidate the skeletal mesh for when the GUID was already set but the mapping data was missing
+	// Game thread only to avoid clashing with the async build where the cloth bindings are unbound/rebound.
+	// Partially bound sections can only be fixed on the GameThread and from the Persona Editor anyway.
+	if (bWasPartiallyBound && IsInGameThread())
+	{
+		SkeletalMesh->InvalidateDeriveDataCacheGUID();
+	}
+
 	return true;
 }
 
diff --git a/Engine/Plugins/ChaosClothAsset/Source/ChaosClothAssetEngine/Public/ChaosClothAsset/ClothAssetSKMClothingAsset.h b/Engine/Plugins/ChaosClothAsset/Source/ChaosClothAssetEngine/Public/ChaosClothAsset/ClothAssetSKMClothingAsset.h
--- a/Engine/Plugins/ChaosClothAsset/Source/ChaosClothAssetEngine/Public/ChaosClothAsset/ClothAssetSKMClothingAsset.h
+++ b/Engine/Plugins/ChaosClothAsset/Source/ChaosClothAssetEngine/Public/ChaosClothAsset/ClothAssetSKMClothingAsset.h
@@ -61,6 +61,7 @@ private:
 	virtual void UpdateAllLODBiasMappings(USkeletalMesh* SkeletalMesh) override;
 	virtual void GetSimulationMesh(const int32 InMeshLodIndex, TArray<FVector3f>& OutPositions, TArray<uint32>& OutIndices, TArray<float>& OutMaxDistances) const override;
 	virtual TArray<TPair<FText, FText>> GetStats() const override;
+	virtual USkinnedAsset* GetSkinnedAssetDependency() const override;
 #endif  // #if WITH_EDITOR
 	virtual void RefreshBoneMapping(USkeletalMesh* SkeletalMesh) override {}
 	virtual bool IsValid() const override
diff --git a/Engine/Source/Runtime/ClothingSystemRuntimeInterface/Public/ClothingAssetBase.h b/Engine/Source/Runtime/ClothingSystemRuntimeInterface/Public/ClothingAssetBase.h
--- a/Engine/Source/Runtime/ClothingSystemRuntimeInterface/Public/ClothingAssetBase.h
+++ b/Engine/Source/Runtime/ClothingSystemRuntimeInterface/Public/ClothingAssetBase.h
@@ -13,6 +13,7 @@
 
 class ITransactionObjectAnnotation;
 class USkeletalMesh;
+class USkinnedAsset;
 
 #define UE_API CLOTHINGSYSTEMRUNTIMEINTERFACE_API
 
@@ -102,6 +103,13 @@ PRAGMA_ENABLE_DEPRECATION_WARNINGS
 		return {};
 	}
 
+	/**
+	 * Get a Skinned Asset dependency that participates in async compilation, if any.
+	 * Used by the Skeletal Mesh to declare any compilation dependencies from the clothing assets.
+	 * @return The underlying Skinned Asset or nullptr if this clothing asset has no such dependency.
+	 */
+	virtual USkinnedAsset* GetSkinnedAssetDependency() const { return nullptr; }
+
 	//~ Begin UObject Interface
 	/** Make sure the object is transactional. */
 	UE_API virtual void PostInitProperties() override;
diff --git a/Engine/Source/Runtime/Engine/Classes/Engine/SkeletalMesh.h b/Engine/Source/Runtime/Engine/Classes/Engine/SkeletalMesh.h
--- a/Engine/Source/Runtime/Engine/Classes/Engine/SkeletalMesh.h
+++ b/Engine/Source/Runtime/Engine/Classes/Engine/SkeletalMesh.h
@@ -1020,6 +1020,10 @@ public:
 	ENGINE_API virtual bool HasAnyDependenciesCompiling() const override;
 private:
 	void OnNaniteAssemblyReferenceCached(USkeletalMesh* Reference);
+	void CacheClothDependencies();
+
+	/** Cache of SkinnedAsset dependencies accessible during the async build. */
+	TArray<TWeakObjectPtr<USkinnedAsset>> ClothDependencyCache;
 #endif
 
 private:
diff --git a/Engine/Source/Runtime/Engine/Private/SkeletalMesh.cpp b/Engine/Source/Runtime/Engine/Private/SkeletalMesh.cpp
--- a/Engine/Source/Runtime/Engine/Private/SkeletalMesh.cpp
+++ b/Engine/Source/Runtime/Engine/Private/SkeletalMesh.cpp
@@ -827,9 +827,21 @@ const TArray<TObjectPtr<USkeletalMesh>>& USkeletalMesh::GetCachedNaniteAssemblyR
 
 #if WITH_EDITOR
 
+void USkeletalMesh::CacheClothDependencies()
+{
+	ClothDependencyCache.Reset();
+	for (UClothingAssetBase* const ClothingAsset : GetMeshClothingAssets())
+	{
+		if (USkinnedAsset* const Dependency = ClothingAsset ? ClothingAsset->GetSkinnedAssetDependency() : nullptr)
+		{
+			ClothDependencyCache.AddUnique(Dependency);
+		}
+	}
+}
+
 bool USkeletalMesh::HasCompileDependencies() const
 {
-	return IsNaniteAssembly();
+	return IsNaniteAssembly() || ClothDependencyCache.Num();
 }
 
 bool USkeletalMesh::HasAnyDependenciesCompiling() const
@@ -841,14 +853,23 @@ bool USkeletalMesh::HasAnyDependenciesCompiling() const
 			return true;
 		}
 	}
-
+	for (const TWeakObjectPtr<USkinnedAsset>& Dependency : ClothDependencyCache)
+	{
+		if (USkinnedAsset* const ResolvedDependency = Dependency.Get())
+		{
+			if (ResolvedDependency && ResolvedDependency->IsCompiling())
+			{
+				return true;
+			}
+		}
+	}
 	return false;
 }
 
 TArray<USkinnedAsset*> USkeletalMesh::GetSkinnedAssetDependencies() const
 {
 	TArray<USkinnedAsset*> Deps;
-	Deps.Reserve(AssemblyReferenceCache.Num());
+	Deps.Reserve(AssemblyReferenceCache.Num() + ClothDependencyCache.Num());
 	for (USkeletalMesh* SkeletalMesh : GetCachedNaniteAssemblyReferences())
 	{
 		if (SkeletalMesh)
@@ -856,6 +877,13 @@ TArray<USkinnedAsset*> USkeletalMesh::GetSkinnedAssetDependencies() const
 			Deps.Add(SkeletalMesh);
 		}
 	}
+	for (const TWeakObjectPtr<USkinnedAsset>& Dependency : ClothDependencyCache)
+	{
+		if (USkinnedAsset* const ResolvedDependency = Dependency.Get())
+		{
+			Deps.Add(ResolvedDependency);
+		}
+	}
 	return Deps;
 }
 
@@ -2334,7 +2362,7 @@ void USkeletalMesh::Build()
 			AsyncTask->StartBackgroundTask(SkeletalMeshThreadPool, BasePriority, EQueuedWorkFlags::DoNotRunInsideBusyWait);
 		}
 
-		if (IsNaniteAssembly())
+		if (HasCompileDependencies())
 		{
 			FSkinnedAssetCompilingManager::Get().AddSkinnedAssetsWithDependencies({ this });
 		}
@@ -2364,11 +2392,13 @@ void USkeletalMesh::BeginBuildInternal(FSkinnedAssetBuildContext& Context)
 	{
 		RecacheNaniteAssemblyReferences(false);
 	}
-	
+
+	CacheClothDependencies();
+
 	// Make sure the InlineReductionCacheDatas structure has enough room for all available LODs. Since the reduction can happen on multiple threads
 	// this has to happen before we trigger the mesh build.
 	GetImportedModel()->InlineReductionCacheDatas.SetNum(GetNumSourceModels());
-	
+
 	// Flush the resource release commands to the rendering thread to ensure that the build doesn't occur while a resource is still
 	// allocated, and potentially accessing the USkeletalMesh.
 	ReleaseResourcesFence.Wait();
@@ -2440,6 +2470,8 @@ FEvent* USkeletalMesh::LockPropertiesUntil()
 	// compilation ongoing plus any dependency (i.e. UGroomBindings).
 	FAssetCompilingManager::Get().FinishCompilationForObjects({ this });
 
+	CacheClothDependencies();
+
 	//Use the async task compile to lock the properties
 	FSkinnedAsyncTaskContext Context(Event);
 	BeginAsyncTaskInternal(Context);
@@ -2456,7 +2488,7 @@ FEvent* USkeletalMesh::LockPropertiesUntil()
 		AsyncTask->StartBackgroundTask(SkeletalMeshThreadPool, BasePriority, EQueuedWorkFlags::DoNotRunInsideBusyWait, RequiredMemory, TEXT("SkeletalMeshAsset"));
 	}
 
-	if (IsNaniteAssembly())
+	if (HasCompileDependencies())
 	{
 		FSkinnedAssetCompilingManager::Get().AddSkinnedAssetsWithDependencies({ this });
 	}
@@ -3889,6 +3921,20 @@ void USkeletalMesh::BeginPostLoadInternal(FSkinnedAssetPostLoadContext& Context)
 		MeshClothingAsset->ConditionalPostLoad();
 	}
 
+	// Update the clothing asset dependencies, this is used to access the cloth skinned asset dependencies during the async build
+	CacheClothDependencies();
+
+	// Ensure inner skinned asset dependencies are PostLoaded
+	// Notes: This is so they register with the compiling manager before this mesh enters the dependency queue.
+	//        Unlike the ConditionalPostLoad above, these are standalone assets whose RF_NeedsPostLoad flag is reliably managed by the linker.
+	for (const TWeakObjectPtr<USkinnedAsset>& Dependency : ClothDependencyCache)
+	{
+		if (USkinnedAsset* const ResolvedDependency = Dependency.Get())
+		{
+			ResolvedDependency->ConditionalPostLoad();
+		}
+	}
+
 	if (!GetOutermost()->bIsCookedForEditor)
 	{
 		// If LODInfo is missing - create array of correct size.
