Iris
* Corrected documentation referencing COND_Group instead of COND_NetGroup

#rnx
#Synced-CL 52176657

diff --git a/Engine/Source/Runtime/Engine/Public/Net/Iris/ReplicationSystem/ReplicationSystemUtil.h b/Engine/Source/Runtime/Engine/Public/Net/Iris/ReplicationSystem/ReplicationSystemUtil.h
--- a/Engine/Source/Runtime/Engine/Public/Net/Iris/ReplicationSystem/ReplicationSystemUtil.h
+++ b/Engine/Source/Runtime/Engine/Public/Net/Iris/ReplicationSystem/ReplicationSystemUtil.h
@@ -135,10 +135,10 @@ struct FReplicationSystemUtil
 	 */
 	ENGINE_API static void SetNetConditionForActorComponent(const UActorComponent* SubObject, ELifetimeCondition Condition);
 
-	/** Update group memberships used by COND_Group sub object filtering for the specified SubObject */
+	/** Update group memberships used by COND_NetGroup sub object filtering for the specified SubObject */
 	ENGINE_API static void UpdateSubObjectGroupMemberships(const UObject* SubObject, const UWorld* World);
 
-	/** Update replication status for all NetGroups used by COND_Group sub object filtering for the provided PlayerController */
+	/** Update replication status for all NetGroups used by COND_NetGroup sub object filtering for the provided PlayerController */
 	ENGINE_API static void UpdateSubObjectGroupMemberships(const APlayerController* PC);
 
 	/** Update replication status for the PlayerController to not include the specified NetGroup */
