Substrate - Fix for crash when using Unlit BSDF as input of Select node.

#tests PC editor
#rnx
#Synced-CL 52251747

diff --git a/Engine/Source/Runtime/Engine/Private/Materials/MaterialExpressionSubstrate.cpp b/Engine/Source/Runtime/Engine/Private/Materials/MaterialExpressionSubstrate.cpp
--- a/Engine/Source/Runtime/Engine/Private/Materials/MaterialExpressionSubstrate.cpp
+++ b/Engine/Source/Runtime/Engine/Private/Materials/MaterialExpressionSubstrate.cpp
@@ -3357,6 +3357,17 @@ int32 UMaterialExpressionSubstrateSelect::Compile(class FMaterialCompiler* Compi
 	FSubstrateOperator& SubstrateOperator = Compiler->SubstrateCompilationGetOperator(Compiler->SubstrateTreeStackGetPathUniqueId());
 	FSubstrateOperator* OperatorA = Compiler->SubstrateCompilationGetOperatorFromIndex(SubstrateOperator.LeftIndex);
 	FSubstrateOperator* OperatorB = Compiler->SubstrateCompilationGetOperatorFromIndex(SubstrateOperator.RightIndex);
+
+	if (!OperatorA || !OperatorB)
+	{
+		const FString ErroneousInput = (!OperatorA && !OperatorB) ? TEXT("A and B") : (!OperatorA ? TEXT("A") : TEXT("B"));
+		return Compiler->Errorf(TEXT("Substrate Select node has invalid input %s."), *ErroneousInput);
+	}
+	if (OperatorA->BSDFRegisteredSharedLocalBasis.NormalCodeChunk == INDEX_NONE || OperatorB->BSDFRegisteredSharedLocalBasis.NormalCodeChunk == INDEX_NONE)
+	{
+		return Compiler->Errorf(TEXT("Substrate Select node requires both inputs to have a valid local basis. Unlit BSDF does not register a local basis, so it cannot be used with the Select node."));
+	}
+
 	const int32 NewNormalCodeChunk = SubstrateSelectNormal( OperatorA->BSDFRegisteredSharedLocalBasis.NormalCodeChunk, OperatorB->BSDFRegisteredSharedLocalBasis.NormalCodeChunk);
 	// The tangent is optional so we treat it differently if INDEX_NONE is specified
 	int32 NewTangentCodeChunk = INDEX_NONE;
diff --git a/Engine/Source/Runtime/Engine/Private/Materials/SubstrateTranslatorCommon.cpp b/Engine/Source/Runtime/Engine/Private/Materials/SubstrateTranslatorCommon.cpp
--- a/Engine/Source/Runtime/Engine/Private/Materials/SubstrateTranslatorCommon.cpp
+++ b/Engine/Source/Runtime/Engine/Private/Materials/SubstrateTranslatorCommon.cpp
@@ -405,12 +405,14 @@ bool FSubstrateCompilationContext::SubstrateGenerateDerivedMaterialOperatorData(
 		{
 			FSubstrateTreeValidation Result;
 			Result.MaterialType = A.MaterialType | B.MaterialType;
+			Result.bValid = A.bValid && B.bValid;
 			return Result;
 		};
 		auto MergeSubstrateTreeValidation = [](const FSubstrateTreeValidation& A, const FSubstrateTreeValidation& B)
 		{
 			FSubstrateTreeValidation Result;
 			Result.MaterialType = A.MaterialType | B.MaterialType;
+			Result.bValid = A.bValid && B.bValid;
 			return Result;
 		};
 
