Removing more implicit command list usage from streaming code.

#Synced-CL 52214067

diff --git a/Engine/Source/Runtime/Engine/Private/Streaming/Texture2DMipAllocator_AsyncReallocate.cpp b/Engine/Source/Runtime/Engine/Private/Streaming/Texture2DMipAllocator_AsyncReallocate.cpp
--- a/Engine/Source/Runtime/Engine/Private/Streaming/Texture2DMipAllocator_AsyncReallocate.cpp
+++ b/Engine/Source/Runtime/Engine/Private/Streaming/Texture2DMipAllocator_AsyncReallocate.cpp
@@ -126,10 +126,12 @@ void FTexture2DMipAllocator_AsyncReallocate::UnlockNewMips()
 	// Unlock any locked mips.
 	if (IntermediateTextureRHI)
 	{
+		FRHICommandListImmediate& RHICmdList = FRHICommandListImmediate::Get();
+
 		const bool bFlushRHIThread = CVarFlushRHIThreadOnSTreamingTextureLocks.GetValueOnAnyThread() > 0;
 		for (int32 MipIndex : LockedMipIndices)
 		{
-			RHIUnlockTexture2D(IntermediateTextureRHI, MipIndex, false, CVarFlushRHIThreadOnSTreamingTextureLocks.GetValueOnAnyThread() > 0 );
+			RHICmdList.UnlockTexture2D(IntermediateTextureRHI, MipIndex, false, bFlushRHIThread);
 		}
 		LockedMipIndices.Empty();
 	}
