Fix flickering sections when skinning scene extension is enabled and certain sections are disabled.

#Synced-CL 52259874

diff --git a/Engine/Source/Runtime/Engine/Private/SkeletalRenderGPUSkin.cpp b/Engine/Source/Runtime/Engine/Private/SkeletalRenderGPUSkin.cpp
--- a/Engine/Source/Runtime/Engine/Private/SkeletalRenderGPUSkin.cpp
+++ b/Engine/Source/Runtime/Engine/Private/SkeletalRenderGPUSkin.cpp
@@ -2566,7 +2566,7 @@ void FSkeletalMeshObjectGPUSkin::FVertexFactoryData::InitVertexFactories(
 			CreateVertexFactory(RHICmdList, VertexFactories, Passthrough, VertexBuffers, InFeatureLevel, VertexAttributeMask, NumBones, BoneOffset, Section.BaseVertexIndex, bUsedForPassthroughVertexFactory, bUnifiedBoneBuffer);
 		}
 
-		BoneOffset += bUnifiedBoneBuffer ? 0 : Section.BoneMap.Num();
+		BoneOffset += Section.IsValid() && !bUnifiedBoneBuffer ? Section.BoneMap.Num() : 0;
 	}
 }
 
