[Input] Add a GetAllHardwareDeviceIdentifiers function to make tooling and access to available device descriptions easier

#rnx
#Synced-CL 52059832

diff --git a/Engine/Source/Runtime/Engine/Classes/GameFramework/InputSettings.h b/Engine/Source/Runtime/Engine/Classes/GameFramework/InputSettings.h
--- a/Engine/Source/Runtime/Engine/Classes/GameFramework/InputSettings.h
+++ b/Engine/Source/Runtime/Engine/Classes/GameFramework/InputSettings.h
@@ -557,6 +557,13 @@ public:
 	*/
 	UFUNCTION()
 	static ENGINE_API const TArray<FName>& GetAllHardwareDeviceNames();
+
+	/**
+	 * Returns all FHardwareDeviceIdentifier types known across all platform configurations.
+	 * Unlike GetAllHardwareDeviceNames, this returns the full structs (including InputClassName
+	 * and PrimaryDeviceType) so callers can build descriptors for device simulation.
+	 */
+	static ENGINE_API TArray<FHardwareDeviceIdentifier> GetAllHardwareDeviceIdentifiers();
 #endif	// WITH_EDITOR
 
 
diff --git a/Engine/Source/Runtime/Engine/Private/UserInterface/InputSettings.cpp b/Engine/Source/Runtime/Engine/Private/UserInterface/InputSettings.cpp
--- a/Engine/Source/Runtime/Engine/Private/UserInterface/InputSettings.cpp
+++ b/Engine/Source/Runtime/Engine/Private/UserInterface/InputSettings.cpp
@@ -762,4 +762,29 @@ const TArray<FName>& UInputPlatformSettings::GetAllHardwareDeviceNames()
 
 	return HardwareDevices;
 }
+
+TArray<FHardwareDeviceIdentifier> UInputPlatformSettings::GetAllHardwareDeviceIdentifiers()
+{
+	TArray<FHardwareDeviceIdentifier> Result;
+
+	// Always include the built-in defaults.
+	Result.AddUnique(FHardwareDeviceIdentifier::DefaultKeyboardAndMouse);
+	Result.AddUnique(FHardwareDeviceIdentifier::DefaultGamepad);
+	Result.AddUnique(FHardwareDeviceIdentifier::DefaultMobileTouch);
+
+	TArray<TObjectPtr<UPlatformSettings>> AllInputSettings = UPlatformSettingsManager::Get().GetAllPlatformSettings<UInputPlatformSettings>();
+
+	for (const TObjectPtr<UPlatformSettings>& Setting : AllInputSettings)
+	{
+		if (const UInputPlatformSettings* InputSetting = Cast<UInputPlatformSettings>(Setting))
+		{
+			for (const FHardwareDeviceIdentifier& Device : InputSetting->HardwareDevices)
+			{
+				Result.AddUnique(Device);
+			}
+		}
+	}
+
+	return Result;
+}
 #endif	// WITH_EDITOR
