Added ISKM attachment API for attaching a child ISKM instance to another parent ISKM instance.
- Added SetInstanceBoneAttachment / ClearInstanceBoneAttachment. The child holds a weak ptr to the parent. The child bounds are expanded to include the parent.
- Added a GPU scene resolve shader that pulls the attached parent bone and transforms the child instance to it.

#Synced-CL 52214594

#Synced-CL 52214594

#Synced-CL 52214594

diff --git a/Engine/Shaders/Private/Skinning/ResolveAttachments.usf b/Engine/Shaders/Private/Skinning/ResolveAttachments.usf
new file mode 100644
--- /dev/null
+++ b/Engine/Shaders/Private/Skinning/ResolveAttachments.usf
@@ -0,0 +1,97 @@
+// Copyright Epic Games, Inc. All Rights Reserved.
+
+// NOTE: GPUSceneWriter.ush must be included BEFORE SceneData.ush as it defines USE_GPU_SCENE_DATA_RW.
+#include "/Engine/Private/Common.ush"
+#include "/Engine/Private/GPUScene/GPUSceneWriter.ush"
+#include "/Engine/Private/BoneTransform.ush"
+#include "/Engine/Private/GpuSkinCommon.ush"
+#include "/Engine/Private/Skinning/SkinningDataDecode.ush"
+#include "/Engine/Shared/SkinningDefinitions.h"
+
+StructuredBuffer<FResolveAttachmentBlockHeader> AttachmentBlockHeaderBuffer;
+FCompressedBoneTransformBuffer SkinningTransformBuffer;
+uint NumAttachments;
+
+float4x4 LoadRefPoseMatrix(FResolveAttachmentBlockHeader Header)
+{
+	float4 Row0 = float4(Header.BoneRefPoseRow0[0], Header.BoneRefPoseRow0[1], Header.BoneRefPoseRow0[2], Header.BoneRefPoseRow0[3]);
+	float4 Row1 = float4(Header.BoneRefPoseRow1[0], Header.BoneRefPoseRow1[1], Header.BoneRefPoseRow1[2], Header.BoneRefPoseRow1[3]);
+	float4 Row2 = float4(Header.BoneRefPoseRow2[0], Header.BoneRefPoseRow2[1], Header.BoneRefPoseRow2[2], Header.BoneRefPoseRow2[3]);
+	float3x4 Packed = float3x4(Row0, Row1, Row2);
+	float4x3 Unpacked = transpose(Packed);
+	return float4x4(float4(Unpacked[0], 0), float4(Unpacked[1], 0), float4(Unpacked[2], 0), float4(Unpacked[3], 1));
+}
+
+float4x4 LoadLocalOffsetMatrix(FResolveAttachmentBlockHeader Header)
+{
+	FQuat Quat = FQuat(
+		Header.LocalOffsetRotation[0],
+		Header.LocalOffsetRotation[1],
+		Header.LocalOffsetRotation[2],
+		Header.LocalOffsetRotation[3]);
+
+	float3 Trans = float3(Header.LocalOffsetTranslation[0], Header.LocalOffsetTranslation[1], Header.LocalOffsetTranslation[2]);
+
+	return MatrixFromBoneTransform(MakeBoneTransform(Quat, Trans));
+}
+
+// Resolve a bone from the skinning buffer, undo InvRefPose, compose with local offset and parent world transform.
+// If BoneIndex is 0xFFFF (bone not in skinning buffer), uses ref pose directly (no animation).
+float4x4 ComposeBoneAttachmentTransform(
+	uint BoneBase,
+	float4x4 RefPoseMatrix,
+	float4x4 LocalOffsetMatrix,
+	float4x4 ParentLocalToWorld,
+	float3 PositionHigh,
+	uint BoneIndex)
+{
+	float4x4 Bone;
+	if (BoneIndex != 0xFFFF)
+	{
+		float4x3 SkinMatrix = LoadCompressedBoneTransform(SkinningTransformBuffer, 0, BoneBase + BoneIndex);
+		Bone = mul(RefPoseMatrix, float4x4(float4(SkinMatrix[0], 0), float4(SkinMatrix[1], 0), float4(SkinMatrix[2], 0), float4(SkinMatrix[3], 1)));
+	}
+	else
+	{
+		Bone = RefPoseMatrix;
+	}
+
+	float4x4 Result = mul(mul(LocalOffsetMatrix, Bone), ParentLocalToWorld);
+	Result[3].xyz -= PositionHigh;
+	return Result;
+}
+
+[numthreads(64, 1, 1)]
+void ResolveAttachmentsCS(uint3 DispatchThreadId : SV_DispatchThreadID)
+{
+	uint AttachmentIndex = DispatchThreadId.x;
+	if (AttachmentIndex >= NumAttachments)
+	{
+		return;
+	}
+
+	FResolveAttachmentBlockHeader Header = AttachmentBlockHeaderBuffer[AttachmentIndex];
+
+	// Unpack packed fields.
+	uint ParentMaxTransformCount     = GetMaxTransformCount(Header);
+	uint BoneIndex                   = GetBoneIndex(Header);
+	uint ParentCurrentTransformSlot  = GetCurrentTransformSlot(Header);
+	uint ParentInstanceId            = GetParentInstanceId(Header);
+
+	FInstanceSceneData ParentInstance = GetInstanceSceneData(ParentInstanceId);
+	FPrimitiveSceneData ChildPrimitive = GetPrimitiveData(Header.ChildPrimitiveId);
+
+	float4x4 RefPoseMatrix = LoadRefPoseMatrix(Header);
+	float4x4 LocalOffsetMatrix = LoadLocalOffsetMatrix(Header);
+
+	uint CurrentBase = GetSkinningBoneTransformOffset(Header.ParentTransformBufferOffset, ParentInstance.SkinningData, ParentCurrentTransformSlot, ParentMaxTransformCount, false);
+	WriteInstanceLocalToRelativeWorld(
+		Header.ChildInstanceId,
+		ComposeBoneAttachmentTransform(CurrentBase, RefPoseMatrix, LocalOffsetMatrix, DFDemote(ParentInstance.LocalToWorld), ChildPrimitive.PositionHigh, BoneIndex),
+		true);
+
+	uint PreviousBase = GetSkinningBoneTransformOffset(Header.ParentTransformBufferOffset, ParentInstance.SkinningData, ParentCurrentTransformSlot, ParentMaxTransformCount, true);
+	WriteInstanceDynamicData(
+		GetInstancePayloadDataOffsets(Header.ChildInstanceId),
+		ComposeBoneAttachmentTransform(PreviousBase, RefPoseMatrix, LocalOffsetMatrix, DFDemote(ParentInstance.PrevLocalToWorld), ChildPrimitive.PositionHigh, BoneIndex));
+}
diff --git a/Engine/Shaders/Shared/SkinningDefinitions.h b/Engine/Shaders/Shared/SkinningDefinitions.h
--- a/Engine/Shaders/Shared/SkinningDefinitions.h
+++ b/Engine/Shaders/Shared/SkinningDefinitions.h
@@ -172,6 +172,37 @@ struct FScatterTransformsBlockHeader
 };
 SKINNING_STATIC_ASSERT(sizeof(FScatterTransformsBlockHeader) == 24, "FScatterTransformsBlockHeader must be 24 bytes");
 
+struct FResolveAttachmentBlockHeader
+{
+	uint ParentTransformBufferOffset;
+	uint ParentMaxTransformCountAndBoneIndex;       // bits [15:0] = MaxTransformCount, [31:16] = BoneIndex
+	uint ParentCurrentTransformSlotAndInstanceId;   // bits [7:0] = CurrentTransformSlot, [31:8] = ParentInstanceId
+	uint ChildInstanceId;
+	uint ChildPrimitiveId;
+	float LocalOffsetTranslation[3];
+	float LocalOffsetRotation[4];
+	float BoneRefPoseRow0[4];
+	float BoneRefPoseRow1[4];
+	float BoneRefPoseRow2[4];
+};
+SKINNING_STATIC_ASSERT(sizeof(FResolveAttachmentBlockHeader) == 96, "FResolveAttachmentBlockHeader must be 96 bytes");
+
+#ifdef __cplusplus
+inline void SetMaxTransformCount(FResolveAttachmentBlockHeader& H, uint32 V)    { H.ParentMaxTransformCountAndBoneIndex = (H.ParentMaxTransformCountAndBoneIndex & 0xFFFF0000) | (V & 0xFFFF); }
+inline void SetBoneIndex(FResolveAttachmentBlockHeader& H, uint32 V)            { H.ParentMaxTransformCountAndBoneIndex = (H.ParentMaxTransformCountAndBoneIndex & 0x0000FFFF) | (V << 16); }
+inline void SetCurrentTransformSlot(FResolveAttachmentBlockHeader& H, uint32 V) { H.ParentCurrentTransformSlotAndInstanceId = (H.ParentCurrentTransformSlotAndInstanceId & 0xFFFFFF00) | (V & 0xFF); }
+inline void SetParentInstanceId(FResolveAttachmentBlockHeader& H, uint32 V)     { H.ParentCurrentTransformSlotAndInstanceId = (H.ParentCurrentTransformSlotAndInstanceId & 0x000000FF) | (V << 8); }
+inline uint32 GetMaxTransformCount(const FResolveAttachmentBlockHeader& H)    { return H.ParentMaxTransformCountAndBoneIndex & 0xFFFF; }
+inline uint32 GetBoneIndex(const FResolveAttachmentBlockHeader& H)            { return H.ParentMaxTransformCountAndBoneIndex >> 16; }
+inline uint32 GetCurrentTransformSlot(const FResolveAttachmentBlockHeader& H) { return H.ParentCurrentTransformSlotAndInstanceId & 0xFF; }
+inline uint32 GetParentInstanceId(const FResolveAttachmentBlockHeader& H)     { return H.ParentCurrentTransformSlotAndInstanceId >> 8; }
+#else
+uint GetMaxTransformCount(FResolveAttachmentBlockHeader H)    { return H.ParentMaxTransformCountAndBoneIndex & 0xFFFF; }
+uint GetBoneIndex(FResolveAttachmentBlockHeader H)            { return H.ParentMaxTransformCountAndBoneIndex >> 16; }
+uint GetCurrentTransformSlot(FResolveAttachmentBlockHeader H) { return H.ParentCurrentTransformSlotAndInstanceId & 0xFF; }
+uint GetParentInstanceId(FResolveAttachmentBlockHeader H)     { return H.ParentCurrentTransformSlotAndInstanceId >> 8; }
+#endif
+
 #define ANIM_BANK_FLAG_NONE			0x0
 #define ANIM_BANK_FLAG_LOOPING		0x1
 #define ANIM_BANK_FLAG_AUTOSTART	0x2
diff --git a/Engine/Source/Runtime/Core/Public/UObject/FortniteMainBranchObjectVersions.inl b/Engine/Source/Runtime/Core/Public/UObject/FortniteMainBranchObjectVersions.inl
--- a/Engine/Source/Runtime/Core/Public/UObject/FortniteMainBranchObjectVersions.inl
+++ b/Engine/Source/Runtime/Core/Public/UObject/FortniteMainBranchObjectVersions.inl
@@ -802,5 +802,8 @@ FORTNITE_MAIN_VERSION(DynamicDelegatePayloads, 262)
 // Compact representation of serializable dynamic delegate payloads which avoids shifting binary data
 FORTNITE_MAIN_VERSION(CompactDynamicDelegatePayloads, 263)
 
+// Added bone attachment serialization for UInstancedSkinnedMeshComponent
+FORTNITE_MAIN_VERSION(InstancedSkinnedMeshBoneAttachments, 264)
+
 // -----<new versions can be added above this line>-------------------------------------------------
 #undef FORTNITE_MAIN_VERSION
diff --git a/Engine/Source/Runtime/Engine/Classes/Components/InstancedSkinnedMeshComponent.h b/Engine/Source/Runtime/Engine/Classes/Components/InstancedSkinnedMeshComponent.h
--- a/Engine/Source/Runtime/Engine/Classes/Components/InstancedSkinnedMeshComponent.h
+++ b/Engine/Source/Runtime/Engine/Classes/Components/InstancedSkinnedMeshComponent.h
@@ -9,6 +9,8 @@
 #include "InstanceDataSceneProxy.h"
 #include "InstanceData/InstanceDataManager.h"
 #include "InstancedSkinnedMeshSceneProxyDesc.h"
+#include "Matrix3x4.h"
+#include "SkinningSceneExtensionProxy.h"
 #include "InstancedSkinnedMeshComponent.generated.h"
 
 class FPrimitiveSceneProxy;
@@ -18,6 +20,7 @@ class UTexture;
 class UAnimBank;
 struct FAnimBankRecordHandle;
 struct FPrimitiveSceneDesc;
+struct FBoneAttachmentSocket;
 
 USTRUCT()
 struct FSkinnedMeshInstanceData
@@ -109,6 +112,13 @@ protected:
 	UPROPERTY(EditAnywhere, BlueprintReadOnly, Category = Instances)
 	bool bHasPreviousTransforms = false;
 
+	/** Forces all socket bones to be animated on the GPU so that they are always available for attachment, regardless
+     *  of whether the mesh vertices are influenced by them. When false, only bones that are actively influencing vertices
+     *  are included in the transform set and can be attached to.
+     */
+	UPROPERTY(EditAnywhere, BlueprintReadOnly, Category = Animation)
+	bool bForceAnimateSockets = false;
+
 	/** If true, this component will use GPU LOD selection. */
 	UPROPERTY(EditAnywhere, BlueprintReadOnly, Category = Culling)
 	uint8 bUseGpuLodSelection : 1;
@@ -268,6 +278,22 @@ public:
 
 	ENGINE_API bool IsEnabled() const;
 
+	/** Attach an instance to a parent instance's bone at a socket or bone. */
+	UFUNCTION(BlueprintCallable, Category = "Components|InstancedSkinnedMesh")
+	ENGINE_API void SetInstanceBoneAttachment(
+		FPrimitiveInstanceId InstanceId,
+		FPrimitiveInstanceId ParentInstanceId,
+		UInstancedSkinnedMeshComponent* Parent,
+		FName SocketName);
+
+	/** Detach an instance from bone attachment. */
+	UFUNCTION(BlueprintCallable, Category = "Components|InstancedSkinnedMesh")
+	ENGINE_API void ClearInstanceBoneAttachment(FPrimitiveInstanceId InstanceId);
+
+	/** Returns true if this component has any bone attachment sockets. */
+	UFUNCTION(BlueprintCallable, Category = "Components|InstancedSkinnedMesh")
+	bool HasBoneAttachmentParents() const { return BoneAttachmentSockets.Num() > 0; }
+
 	UFUNCTION(BlueprintCallable, Category = "Components|InstancedSkinnedMesh")
 	ENGINE_API int32 GetInstanceCount() const;
 
@@ -342,6 +368,36 @@ private:
 
 	FInstanceDataManager InstanceDataManager;
 
+	struct FBoneAttachmentSocketRef
+	{
+		TWeakObjectPtr<UInstancedSkinnedMeshComponent> Parent;
+		FName SocketName;
+		int32 RefCount = 0;
+		bool bNeedsResolve = false;
+
+		/** Resolved socket data, populated at acquire time or lazily on first flush. */
+		FBoneAttachmentSocket Resolved;
+	};
+
+	TSparseArray<FBoneAttachmentSocketRef> BoneAttachmentSockets;
+	TArray<FBoneAttachmentBinding> BoneAttachmentBindings;
+	TBitArray<> DirtyBoneAttachmentBindings;
+	int32 NumDirtyBoneAttachmentBindings = 0;
+
+	void FlushBoneAttachmentSockets();
+
+	int32 AcquireBoneAttachmentSocket(UInstancedSkinnedMeshComponent* Parent, FName SocketName, bool bDeferResolve = false);
+	bool ResolveBoneAttachmentSocket(FBoneAttachmentSocketRef& Socket);
+	void ReleaseBoneAttachmentSocket(int32 SocketIndex);
+
+	void MarkBoneAttachmentBindingDirty(int32 InstanceIndex);
+	void ResetBoneAttachmentDirtyState();
+
+	void SaveBoneAttachments(FArchive& Ar);
+	void LoadBoneAttachments(FArchive& Ar);
+
+	TArray<FBoneAttachmentSocket> GetBoneAttachmentSockets();
+
 	friend class FInstancedSkinnedMeshSceneProxy;
 	friend class FInstancedSkinnedMeshComponentHelper;
 	friend struct FInstancedSkinnedMeshComponentInstanceData;
diff --git a/Engine/Source/Runtime/Engine/Private/Components/InstancedSkinnedMeshComponent.cpp b/Engine/Source/Runtime/Engine/Private/Components/InstancedSkinnedMeshComponent.cpp
--- a/Engine/Source/Runtime/Engine/Private/Components/InstancedSkinnedMeshComponent.cpp
+++ b/Engine/Source/Runtime/Engine/Private/Components/InstancedSkinnedMeshComponent.cpp
@@ -12,6 +12,9 @@
 #include "HitProxies.h"
 #include "NaniteSceneProxy.h"
 #include "PrimitiveSceneInfo.h"
+#include "Engine/SkeletalMeshSocket.h"
+#include "SkinningSceneExtensionProxy.h"
+#include "UObject/FortniteMainBranchObjectVersion.h"
 #include "Rendering/NaniteResourcesHelper.h"
 #include "Rendering/RenderCommandPipes.h"
 #include "Rendering/SkeletalMeshRenderData.h"
@@ -32,6 +35,8 @@
 
 #include UE_INLINE_GENERATED_CPP_BY_NAME(InstancedSkinnedMeshComponent)
 
+DEFINE_LOG_CATEGORY_STATIC(LogInstancedSkinnedMesh, Log, All);
+
 TAutoConsoleVariable<int32> CVarInstancedSkinnedMeshesForceRefPose(
 	TEXT("r.InstancedSkinnedMeshes.ForceRefPose"),
 	0,
@@ -132,6 +137,7 @@ void UInstancedSkinnedMeshComponent::Serialize(FArchive& Ar)
 	Super::Serialize(Ar);
 
 	Ar.UsingCustomVersion(FUE5MainStreamObjectVersion::GUID);
+	Ar.UsingCustomVersion(FFortniteMainBranchObjectVersion::GUID);
 
 	bool bCooked = Ar.IsCooking();
 	Ar << bCooked;
@@ -205,6 +211,25 @@ void UInstancedSkinnedMeshComponent::Serialize(FArchive& Ar)
 		InstanceDataManager.Reset(InstanceData.Num());
 	}
 
+	if (Ar.CustomVer(FFortniteMainBranchObjectVersion::GUID) >= FFortniteMainBranchObjectVersion::InstancedSkinnedMeshBoneAttachments)
+	{
+		if (Ar.IsSaving())
+		{
+			SaveBoneAttachments(Ar);
+		}
+		else
+		{
+			LoadBoneAttachments(Ar);
+		}
+	}
+	else if (Ar.IsLoading())
+	{
+		// Old asset without bone attachment data. Initialize parallel arrays to match instance count.
+		BoneAttachmentBindings.SetNum(InstanceData.Num());
+		DirtyBoneAttachmentBindings.Init(false, InstanceData.Num());
+		NumDirtyBoneAttachmentBindings = 0;
+	}
+
 	if (bCooked)
 	{
 		if (Ar.IsLoading())
@@ -312,6 +337,17 @@ void UInstancedSkinnedMeshComponent::OptimizeInstanceData(bool bShouldRetainIdMa
 #if WITH_EDITOR
 		ReorderArray(SelectedInstances, IndexRemap);
 #endif
+
+		ReorderArray(BoneAttachmentBindings, IndexRemap);
+
+		// Dirty bits are invalid after reorder — reset and force proxy recreation
+		// so the render thread gets the reordered bindings.
+		DirtyBoneAttachmentBindings.Init(false, BoneAttachmentBindings.Num());
+		NumDirtyBoneAttachmentBindings = 0;
+		if (BoneAttachmentSockets.Num() > 0)
+		{
+			MarkRenderStateDirty();
+		}
 	}
 }
 
@@ -321,6 +357,20 @@ void UInstancedSkinnedMeshComponent::ApplyInheritedPerInstanceData(const UInstan
 	InstanceData = InArchetype->InstanceData;
 	InstanceCustomData = InArchetype->InstanceCustomData;
 	NumCustomDataFloats = InArchetype->NumCustomDataFloats;
+
+	// Bone attachment sockets reference specific parent component instances
+	// that are not valid for the inherited copy. Release and clear.
+	for (const FBoneAttachmentBinding& Binding : BoneAttachmentBindings)
+	{
+		if (Binding.IsAttached())
+		{
+			ReleaseBoneAttachmentSocket(Binding.GetSocketIndex());
+		}
+	}
+	BoneAttachmentBindings.Reset();
+	BoneAttachmentBindings.SetNum(InstanceData.Num());
+	DirtyBoneAttachmentBindings.Init(false, InstanceData.Num());
+	NumDirtyBoneAttachmentBindings = 0;
 }
 
 bool UInstancedSkinnedMeshComponent::ShouldInheritPerInstanceData() const
@@ -391,7 +441,11 @@ void UInstancedSkinnedMeshComponent::SendRenderInstanceData_Concurrent()
 		TransformProvider->SubmitChanges();
 	}
 
-	// If instance data is entirely GPU driven, don't upload from CPU.
+	// Flush bone attachment changes before the GPU-only early return —
+	// GPU-only components still need incremental binding updates on the proxy.
+	FlushBoneAttachmentSockets();
+
+	// If instance data is entirely GPU driven, don't upload CPU instance data.
 	if (bIsInstanceDataGPUOnly)
 	{
 		return;
@@ -419,6 +473,395 @@ void UInstancedSkinnedMeshComponent::SendRenderInstanceData_Concurrent()
 	}
 }
 
+int32 UInstancedSkinnedMeshComponent::AcquireBoneAttachmentSocket(UInstancedSkinnedMeshComponent* Parent, FName SocketName, bool bDeferResolve)
+{
+	// Find existing socket for this parent + socket name.
+	for (TSparseArray<FBoneAttachmentSocketRef>::TIterator It(BoneAttachmentSockets); It; ++It)
+	{
+		FBoneAttachmentSocketRef& Socket = *It;
+		if (Socket.Parent == Parent && Socket.SocketName == SocketName)
+		{
+			++Socket.RefCount;
+			return It.GetIndex();
+		}
+	}
+
+	// Create new socket entry.
+	FBoneAttachmentSocketRef Socket;
+	Socket.Parent = Parent;
+	Socket.SocketName = SocketName;
+	Socket.RefCount = 1;
+
+	if (bDeferResolve)
+	{
+		// During Serialize, parent's mesh may not have completed PostLoad.
+		// Leave Resolved as default (invalid BoneIndex, identity RefPose).
+		Socket.bNeedsResolve = true;
+	}
+	else
+	{
+		if (!ResolveBoneAttachmentSocket(Socket))
+		{
+			return INDEX_NONE;
+		}
+	}
+
+	return BoneAttachmentSockets.Add(MoveTemp(Socket));
+}
+
+bool UInstancedSkinnedMeshComponent::ResolveBoneAttachmentSocket(FBoneAttachmentSocketRef& Socket)
+{
+	Socket.bNeedsResolve = false;
+
+	UInstancedSkinnedMeshComponent* Parent = Socket.Parent.Get();
+	if (!Parent)
+	{
+		return false;
+	}
+
+	const USkinnedAsset* ParentAsset = Parent->GetSkinnedAsset();
+	if (!ParentAsset)
+	{
+		UE_LOG(LogInstancedSkinnedMesh, Warning, TEXT("ResolveBoneAttachmentSocket: Parent '%s' has no skinned asset."), *Parent->GetPathName());
+		return false;
+	}
+
+	Socket.Resolved.ParentComponentId = Parent->GetPrimitiveSceneId();
+
+	const USkeletalMeshSocket* MeshSocket = ParentAsset->FindSocket(Socket.SocketName);
+	const FName BoneName = MeshSocket ? MeshSocket->BoneName : Socket.SocketName;
+	const int32 BoneIndex = ParentAsset->GetRefSkeleton().FindBoneIndex(BoneName);
+
+	if (BoneIndex == INDEX_NONE)
+	{
+		UE_LOG(LogInstancedSkinnedMesh, Warning, TEXT("ResolveBoneAttachmentSocket: Socket/bone '%s' not found on parent '%s'"),
+			*Socket.SocketName.ToString(), *ParentAsset->GetPathName());
+		return false;
+	}
+
+	Socket.Resolved.BoneIndex = BoneIndex;
+
+	TConstArrayView<FMatrix44f> InvRefMatrices = ParentAsset->GetRefBasesInvMatrix();
+	if (InvRefMatrices.IsValidIndex(BoneIndex))
+	{
+		TransposeTransform(Socket.Resolved.RefPoseMatrix, InvRefMatrices[BoneIndex].Inverse());
+	}
+
+	if (MeshSocket)
+	{
+		Socket.Resolved.Rotation = FQuat4f(MeshSocket->RelativeRotation.Quaternion());
+		Socket.Resolved.Translation = FVector3f(MeshSocket->RelativeLocation);
+	}
+
+	return true;
+}
+
+void UInstancedSkinnedMeshComponent::ReleaseBoneAttachmentSocket(int32 SocketIndex)
+{
+	if (!BoneAttachmentSockets.IsValidIndex(SocketIndex) || !BoneAttachmentSockets.IsAllocated(SocketIndex))
+	{
+		return;
+	}
+
+	FBoneAttachmentSocketRef& Socket = BoneAttachmentSockets[SocketIndex];
+	check(Socket.RefCount > 0);
+	if (--Socket.RefCount == 0)
+	{
+		BoneAttachmentSockets.RemoveAt(SocketIndex);
+	}
+}
+
+void UInstancedSkinnedMeshComponent::MarkBoneAttachmentBindingDirty(int32 InstanceIndex)
+{
+	if (!DirtyBoneAttachmentBindings[InstanceIndex])
+	{
+		DirtyBoneAttachmentBindings[InstanceIndex] = true;
+		++NumDirtyBoneAttachmentBindings;
+	}
+	MarkRenderInstancesDirty();
+}
+
+void UInstancedSkinnedMeshComponent::ResetBoneAttachmentDirtyState()
+{
+	DirtyBoneAttachmentBindings.Init(false, DirtyBoneAttachmentBindings.Num());
+	NumDirtyBoneAttachmentBindings = 0;
+}
+
+void UInstancedSkinnedMeshComponent::SetInstanceBoneAttachment(FPrimitiveInstanceId InstanceId, FPrimitiveInstanceId ParentInstanceId,
+	UInstancedSkinnedMeshComponent* Parent, FName SocketName)
+{
+	if (!Parent)
+	{
+		return;
+	}
+
+	if (!InstanceDataManager.IsValidId(InstanceId))
+	{
+		UE_LOG(LogInstancedSkinnedMesh, Warning, TEXT("SetInstanceBoneAttachment: Invalid InstanceId."));
+		return;
+	}
+
+	const int32 ChildIndex = InstanceDataManager.IdToIndex(InstanceId);
+
+	if (!BoneAttachmentBindings.IsValidIndex(ChildIndex))
+	{
+		UE_LOG(LogInstancedSkinnedMesh, Error, TEXT("SetInstanceBoneAttachment: Instance index %d out of range."), ChildIndex);
+		return;
+	}
+
+	if (!ParentInstanceId.IsValid())
+	{
+		UE_LOG(LogInstancedSkinnedMesh, Warning, TEXT("SetInstanceBoneAttachment: Invalid ParentInstanceId."));
+		return;
+	}
+
+	// Release old socket ref if this instance was already attached.
+	const FBoneAttachmentBinding OldBinding = BoneAttachmentBindings[ChildIndex];
+	if (OldBinding.IsAttached())
+	{
+		ReleaseBoneAttachmentSocket(OldBinding.GetSocketIndex());
+	}
+
+	const int32 SocketIndex = AcquireBoneAttachmentSocket(Parent, SocketName);
+	if (SocketIndex == INDEX_NONE)
+	{
+		BoneAttachmentBindings[ChildIndex] = FBoneAttachmentBinding();
+		MarkBoneAttachmentBindingDirty(ChildIndex);
+		return;
+	}
+
+	const int32 ParentIndex = ParentInstanceId.GetAsIndex();
+
+	if (SocketIndex >= 255)
+	{
+		UE_LOG(LogInstancedSkinnedMesh, Error, TEXT("SetInstanceBoneAttachment: Socket index %d exceeds maximum (254). Too many unique parent+socket combinations."), SocketIndex);
+		ReleaseBoneAttachmentSocket(SocketIndex);
+		BoneAttachmentBindings[ChildIndex] = FBoneAttachmentBinding();
+		MarkBoneAttachmentBindingDirty(ChildIndex);
+		return;
+	}
+
+	BoneAttachmentBindings[ChildIndex] = FBoneAttachmentBinding(SocketIndex, ParentIndex);
+	MarkBoneAttachmentBindingDirty(ChildIndex);
+}
+
+void UInstancedSkinnedMeshComponent::ClearInstanceBoneAttachment(FPrimitiveInstanceId InstanceId)
+{
+	if (!InstanceDataManager.IsValidId(InstanceId))
+	{
+		return;
+	}
+
+	const int32 ChildIndex = InstanceDataManager.IdToIndex(InstanceId);
+
+	if (!BoneAttachmentBindings.IsValidIndex(ChildIndex))
+	{
+		return;
+	}
+
+	const FBoneAttachmentBinding Binding = BoneAttachmentBindings[ChildIndex];
+
+	if (!Binding.IsAttached())
+	{
+		return;
+	}
+
+	ReleaseBoneAttachmentSocket(Binding.GetSocketIndex());
+
+	BoneAttachmentBindings[ChildIndex] = FBoneAttachmentBinding();
+	MarkBoneAttachmentBindingDirty(ChildIndex);
+}
+
+void UInstancedSkinnedMeshComponent::FlushBoneAttachmentSockets()
+{
+	if (NumDirtyBoneAttachmentBindings == 0)
+	{
+		return;
+	}
+
+	if (!GetWorld() || !GetWorld()->Scene || !SceneProxy)
+	{
+		return;
+	}
+
+	struct FBindingPatch { int32 Index; FBoneAttachmentBinding Binding; };
+	TArray<FBindingPatch> BindingPatches;
+
+	if (NumDirtyBoneAttachmentBindings > 0)
+	{
+		BindingPatches.Reserve(NumDirtyBoneAttachmentBindings);
+		for (TConstSetBitIterator<> BitIt(DirtyBoneAttachmentBindings); BitIt; ++BitIt)
+		{
+			const int32 DirtyIndex = BitIt.GetIndex();
+			if (BoneAttachmentBindings.IsValidIndex(DirtyIndex))
+			{
+				BindingPatches.Add({ DirtyIndex, BoneAttachmentBindings[DirtyIndex] });
+			}
+		}
+	}
+
+	ENQUEUE_RENDER_COMMAND(UpdateBoneAttachmentSockets)(
+		[
+			  Scene = GetWorld()->Scene
+			, ComponentId = GetPrimitiveSceneId()
+			, ResolvedSockets = GetBoneAttachmentSockets()
+			, BindingPatches = MoveTemp(BindingPatches)
+			, ExpectedBindingCount = BoneAttachmentBindings.Num()
+		] (FRHICommandList&) mutable
+	{
+		FPrimitiveSceneInfo* SceneInfo = Scene->GetPrimitiveSceneInfo(ComponentId);
+		if (!SceneInfo || !SceneInfo->Proxy)
+		{
+			return;
+		}
+
+		FSkinningSceneExtensionProxy* Proxy = SceneInfo->Proxy->GetSkinningSceneExtensionProxy();
+		if (!Proxy || !Proxy->UseInstancing())
+		{
+			return;
+		}
+
+		FInstancedSkinningSceneExtensionProxy* InstancedProxy = static_cast<FInstancedSkinningSceneExtensionProxy*>(Proxy);
+		InstancedProxy->SetBoneAttachmentSockets(MoveTemp(ResolvedSockets));
+		InstancedProxy->ResizeBoneAttachmentBindings(ExpectedBindingCount);
+
+		for (const FBindingPatch& Patch : BindingPatches)
+		{
+			InstancedProxy->SetBoneAttachmentBinding(Patch.Index, Patch.Binding);
+		}
+	});
+
+	ResetBoneAttachmentDirtyState();
+}
+
+TArray<FBoneAttachmentSocket> UInstancedSkinnedMeshComponent::GetBoneAttachmentSockets()
+{
+	TArray<FBoneAttachmentSocket> Result;
+	if (BoneAttachmentSockets.Num() == 0)
+	{
+		return Result;
+	}
+
+	// Resolve any deferred sockets before copying to render thread.
+	for (auto& Socket : BoneAttachmentSockets)
+	{
+		if (Socket.bNeedsResolve)
+		{
+			ResolveBoneAttachmentSocket(Socket);
+		}
+	}
+
+	// Size to accommodate sparse indices (may have gaps from released sockets).
+	Result.SetNum(BoneAttachmentSockets.GetMaxIndex());
+	for (const auto& Socket : BoneAttachmentSockets)
+	{
+		const int32 Index = BoneAttachmentSockets.PointerToIndex(&Socket);
+		Result[Index] = Socket.Resolved;
+	}
+	return Result;
+}
+
+void UInstancedSkinnedMeshComponent::SaveBoneAttachments(FArchive& Ar)
+{
+	int32 NumSockets = BoneAttachmentSockets.Num();
+	Ar << NumSockets;
+
+	for (TSparseArray<FBoneAttachmentSocketRef>::TConstIterator It(BoneAttachmentSockets); It; ++It)
+	{
+		TWeakObjectPtr<UInstancedSkinnedMeshComponent> Parent = It->Parent;
+		FName SocketName = It->SocketName;
+		int32 Index = It.GetIndex();
+		Ar << Parent;
+		Ar << SocketName;
+		Ar << Index;
+	}
+
+	Ar << BoneAttachmentBindings;
+}
+
+void UInstancedSkinnedMeshComponent::LoadBoneAttachments(FArchive& Ar)
+{
+	int32 NumSockets = 0;
+	Ar << NumSockets;
+
+	TMap<int32, int32> IndexRemap;
+	for (int32 SocketIndex = 0; SocketIndex < NumSockets; ++SocketIndex)
+	{
+		TWeakObjectPtr<UInstancedSkinnedMeshComponent> Parent;
+		FName SocketName;
+		int32 SavedIndex;
+		Ar << Parent;
+		Ar << SocketName;
+		Ar << SavedIndex;
+
+		int32 AcquiredIndex = INDEX_NONE;
+		if (UInstancedSkinnedMeshComponent* ParentPtr = Parent.Get())
+		{
+			AcquiredIndex = AcquireBoneAttachmentSocket(ParentPtr, SocketName, /*bDeferResolve=*/ true);
+		}
+
+		if (SavedIndex != AcquiredIndex)
+		{
+			IndexRemap.Add(SavedIndex, AcquiredIndex);
+		}
+	}
+
+	Ar << BoneAttachmentBindings;
+	const int32 ExpectedBindingCount = bIsInstanceDataGPUOnly ? NumInstancesGPUOnly : InstanceData.Num();
+	const int32 SerializedCount = BoneAttachmentBindings.Num();
+	BoneAttachmentBindings.SetNum(ExpectedBindingCount);
+	DirtyBoneAttachmentBindings.Init(false, ExpectedBindingCount);
+	NumDirtyBoneAttachmentBindings = 0;
+
+	// Reset ref counts — will recount from actual bindings during the remap pass below.
+	for (auto& Socket : BoneAttachmentSockets)
+	{
+		Socket.RefCount = 0;
+	}
+
+	// Remap binding socket indices that changed during re-acquire, and recount refs.
+	for (FBoneAttachmentBinding& Binding : BoneAttachmentBindings)
+	{
+		if (!Binding.IsAttached())
+		{
+			continue;
+		}
+
+		const int32 OldSocketIndex = Binding.GetSocketIndex();
+		const int32* NewSocketIndex = IndexRemap.Find(OldSocketIndex);
+
+		if (NewSocketIndex)
+		{
+			Binding = (*NewSocketIndex != INDEX_NONE)
+				? FBoneAttachmentBinding(*NewSocketIndex, Binding.GetParentInstanceIndex())
+				: FBoneAttachmentBinding();
+		}
+		else if (!BoneAttachmentSockets.IsAllocated(OldSocketIndex))
+		{
+			Binding = FBoneAttachmentBinding();
+		}
+
+		// Count actual reference (after potential remap).
+		if (Binding.IsAttached())
+		{
+			const int32 FinalSocketIndex = Binding.GetSocketIndex();
+			if (BoneAttachmentSockets.IsAllocated(FinalSocketIndex))
+			{
+				BoneAttachmentSockets[FinalSocketIndex].RefCount++;
+			}
+		}
+	}
+
+	// Remove any sockets with zero references.
+	for (auto It = BoneAttachmentSockets.CreateIterator(); It; ++It)
+	{
+		if (It->RefCount == 0)
+		{
+			It.RemoveCurrent();
+		}
+	}
+}
+
 bool UInstancedSkinnedMeshComponent::IsHLODRelevant() const
 {
 	if (!CanBeHLODRelevant(this))
@@ -471,6 +914,10 @@ void UInstancedSkinnedMeshComponent::CreateRenderState_Concurrent(FRegisterCompo
 {
 	MeshObjectFactory = &CreateInstancedSkeletalMeshObjectFunction;
 	Super::CreateRenderState_Concurrent(Context);
+
+	// Bone attachment data is populated on the proxy at CreateSceneProxy() time.
+	// Clear dirty state since the proxy has a fresh full copy.
+	ResetBoneAttachmentDirtyState();
 }
 
 void UInstancedSkinnedMeshComponent::DestroyRenderState_Concurrent()
@@ -511,7 +958,7 @@ FPrimitiveSceneProxy* UInstancedSkinnedMeshComponent::CreateSceneProxy()
 
 	if (CheckPSOPrecachingAndBoostPriority() && GetPSOPrecacheProxyCreationStrategy() == EPSOPrecacheProxyCreationStrategy::DelayUntilPSOPrecached)
 	{
-		UE_LOG(LogAnimBank, Verbose, TEXT("Skipping CreateSceneProxy for UInstancedSkinnedMeshComponent %s (UInstancedSkinnedMeshComponent PSOs are still compiling)"), *GetFullName());
+		UE_LOG(LogInstancedSkinnedMesh, Verbose, TEXT("Skipping CreateSceneProxy for UInstancedSkinnedMeshComponent %s (UInstancedSkinnedMeshComponent PSOs are still compiling)"), *GetFullName());
 		return nullptr;
 	}
 
@@ -526,6 +973,19 @@ FPrimitiveSceneProxy* UInstancedSkinnedMeshComponent::CreateSceneProxy()
 		InstanceDataManager.FlushChanges(GetComponentDesc(Result->GetScene().GetShaderPlatform()));
 	}
 
+	// Populate bone attachment data directly on the proxy (already resolved at acquire time).
+	if (Result && BoneAttachmentSockets.Num() > 0)
+	{
+		FSkinningSceneExtensionProxy* SkinProxy = Result->GetSkinningSceneExtensionProxy();
+
+		if (SkinProxy && SkinProxy->UseInstancing())
+		{
+			FInstancedSkinningSceneExtensionProxy* InstancedProxy = static_cast<FInstancedSkinningSceneExtensionProxy*>(SkinProxy);
+			InstancedProxy->SetBoneAttachmentSockets(GetBoneAttachmentSockets());
+			InstancedProxy->SetBoneAttachmentBindings(TArray<FBoneAttachmentBinding>(BoneAttachmentBindings));
+		}
+	}
+
 #if !(UE_BUILD_SHIPPING || UE_BUILD_TEST)
 	SendRenderDebugPhysics(Result);
 #endif
@@ -607,6 +1067,21 @@ void UInstancedSkinnedMeshComponent::PostEditChangeChainProperty(FPropertyChange
 void UInstancedSkinnedMeshComponent::PostEditUndo()
 {
 	Super::PostEditUndo();
+
+	// Bone attachment arrays aren't UPROPERTYs — they aren't restored by the transaction system.
+	// Clear everything and let gameplay code re-establish attachments if needed.
+	for (const FBoneAttachmentBinding& Binding : BoneAttachmentBindings)
+	{
+		if (Binding.IsAttached())
+		{
+			ReleaseBoneAttachmentSocket(Binding.GetSocketIndex());
+		}
+	}
+	BoneAttachmentBindings.Reset();
+	BoneAttachmentBindings.SetNum(InstanceData.Num());
+	DirtyBoneAttachmentBindings.Init(false, InstanceData.Num());
+	NumDirtyBoneAttachmentBindings = 0;
+
 	MarkRenderStateDirty();
 }
 
@@ -753,7 +1228,35 @@ void UInstancedSkinnedMeshComponent::RefreshBoneTransforms(FActorComponentTickFu
 
 void UInstancedSkinnedMeshComponent::SetNumGPUInstances(int32 InCount)
 {
+	if (!bIsInstanceDataGPUOnly)
+	{
+		UE_LOG(LogInstancedSkinnedMesh, Error, TEXT("SetNumGPUInstances: Component '%s' is not in GPU-only mode. Call SetInstanceDataGPUOnly(true) first."), *GetPathName());
+		return;
+	}
+
 	NumInstancesGPUOnly = InCount;
+
+	// Release socket refs for any bindings being truncated.
+	const int32 OldNum = BoneAttachmentBindings.Num();
+	for (int32 i = InCount; i < OldNum; ++i)
+	{
+		const FBoneAttachmentBinding& Binding = BoneAttachmentBindings[i];
+		if (Binding.IsAttached())
+		{
+			ReleaseBoneAttachmentSocket(Binding.GetSocketIndex());
+		}
+	}
+
+	// Resize bone attachment parallel arrays to match.
+	BoneAttachmentBindings.SetNum(InCount);
+	DirtyBoneAttachmentBindings.Init(true, InCount);
+	NumDirtyBoneAttachmentBindings = InCount;
+
+	// Trigger flush of resized bindings to the render proxy.
+	if (BoneAttachmentSockets.Num() > 0)
+	{
+		MarkRenderInstancesDirty();
+	}
 }
 
 FPrimitiveInstanceId UInstancedSkinnedMeshComponent::AddInstance(const FTransform& InstanceTransform, int32 AnimationIndex, bool bWorldSpace)
@@ -1003,6 +1506,10 @@ void UInstancedSkinnedMeshComponent::RemoveInstances(const TArray<FPrimitiveInst
 
 	for (FPrimitiveInstanceId InstanceId : InstancesToRemove)
 	{
+		if (!InstanceDataManager.IsValidId(InstanceId))
+		{
+			continue;
+		}
 		int32 InstanceIndex = InstanceDataManager.IdToIndex(InstanceId);
 		RemoveInstanceInternal(InstanceIndex, false);
 	}
@@ -1017,6 +1524,11 @@ void UInstancedSkinnedMeshComponent::ClearInstances()
 	InstanceCustomData.Empty();
 	PerInstancePrevTransform.Empty();
 
+	BoneAttachmentBindings.Empty();
+	BoneAttachmentSockets.Empty();
+	DirtyBoneAttachmentBindings.Empty();
+	NumDirtyBoneAttachmentBindings = 0;
+
 #if WITH_EDITOR
 	SelectedInstances.Empty();
 #endif
@@ -1113,6 +1625,9 @@ FPrimitiveInstanceId UInstancedSkinnedMeshComponent::AddInstanceInternal(int32 I
 		PerInstancePrevTransform.Add(LocalTransform);
 	}
 
+	BoneAttachmentBindings.Add(FBoneAttachmentBinding());
+	DirtyBoneAttachmentBindings.Add(false);
+
 #if WITH_EDITOR
 	SelectedInstances.Add(false);
 #endif
@@ -1151,6 +1666,32 @@ bool UInstancedSkinnedMeshComponent::RemoveInstanceInternal(int32 InstanceIndex,
 		SelectedInstances.RemoveAtSwap(InstanceIndex);
 	}
 #endif
+
+	if (BoneAttachmentBindings.IsValidIndex(InstanceIndex))
+	{
+		const FBoneAttachmentBinding OldBinding = BoneAttachmentBindings[InstanceIndex];
+		if (OldBinding.IsAttached())
+		{
+			ReleaseBoneAttachmentSocket(OldBinding.GetSocketIndex());
+		}
+
+		BoneAttachmentBindings.RemoveAtSwap(InstanceIndex);
+
+		if (DirtyBoneAttachmentBindings[InstanceIndex])
+		{
+			--NumDirtyBoneAttachmentBindings;
+		}
+		DirtyBoneAttachmentBindings.RemoveAtSwap(InstanceIndex);
+
+		// The swapped-in element at InstanceIndex now has a different binding than the proxy expects.
+		// Mark it dirty so the next flush patches the proxy. The proxy's array is 1 element too large
+		// but the extra entry is Unattached and harmlessly skipped by FillAttachmentHeaders.
+		if (InstanceIndex < BoneAttachmentBindings.Num())
+		{
+			MarkBoneAttachmentBindingDirty(InstanceIndex);
+		}
+	}
+
 	return true;
 }
 
diff --git a/Engine/Source/Runtime/Engine/Private/InstancedSkeletalMeshObject.cpp b/Engine/Source/Runtime/Engine/Private/InstancedSkeletalMeshObject.cpp
--- a/Engine/Source/Runtime/Engine/Private/InstancedSkeletalMeshObject.cpp
+++ b/Engine/Source/Runtime/Engine/Private/InstancedSkeletalMeshObject.cpp
@@ -26,6 +26,7 @@ FInstancedSkeletalMeshObject::FInstancedSkeletalMeshObject(
 	}
 
 	bUseGpuLodSelection = NumLODs > 1 && InMeshDesc.bUseGpuLodSelection && bIsGPUSkin && InRenderData->HasUnifiedBoneMap();
+	bForceAnimateSockets = InMeshDesc.bForceAnimateSockets;
 	BoneTransformStorageMode = (bIsGPUSkin || InRenderData->HasUnifiedBoneMap()) ? EBoneTransformStorageMode::BoneMap : EBoneTransformStorageMode::Direct;
 
 	InitResources(InMeshDesc);
@@ -97,7 +98,7 @@ FSkinningSceneExtensionProxy* FInstancedSkeletalMeshObject::CreateSceneExtension
 {
 	if (Type == EType::Nanite || IsGPUSkinSceneExtensionEnabled())
 	{
-		return new FInstancedSkinningSceneExtensionProxy(TransformProvider, this, InSkinnedAsset, bAllowScaling);
+		return new FInstancedSkinningSceneExtensionProxy(TransformProvider, this, InSkinnedAsset, bAllowScaling, bForceAnimateSockets);
 	}
 	return nullptr;
 }
diff --git a/Engine/Source/Runtime/Engine/Private/InstancedSkeletalMeshObject.h b/Engine/Source/Runtime/Engine/Private/InstancedSkeletalMeshObject.h
--- a/Engine/Source/Runtime/Engine/Private/InstancedSkeletalMeshObject.h
+++ b/Engine/Source/Runtime/Engine/Private/InstancedSkeletalMeshObject.h
@@ -107,6 +107,7 @@ protected:
 	FDynamicData* DynamicData = nullptr;
 	EType Type;
 	bool bUseGpuLodSelection = false;
+	bool bForceAnimateSockets = false;
 
 	friend class FInstancedSkeletalMeshUpdatePacket;
 };
\ No newline at end of file
diff --git a/Engine/Source/Runtime/Engine/Private/InstancedSkinnedMeshSceneProxyDesc.cpp b/Engine/Source/Runtime/Engine/Private/InstancedSkinnedMeshSceneProxyDesc.cpp
--- a/Engine/Source/Runtime/Engine/Private/InstancedSkinnedMeshSceneProxyDesc.cpp
+++ b/Engine/Source/Runtime/Engine/Private/InstancedSkinnedMeshSceneProxyDesc.cpp
@@ -88,4 +88,5 @@ void FInstancedSkinnedMeshSceneProxyDesc::InitializeFromInstancedSkinnedMeshComp
 	InstanceDataSceneProxy = InComponent->GetInstanceDataSceneProxy();
 
 	bUseGpuLodSelection = InComponent->bUseGpuLodSelection;
+	bForceAnimateSockets = InComponent->bForceAnimateSockets;
 }
\ No newline at end of file
diff --git a/Engine/Source/Runtime/Engine/Private/SkinningSceneExtensionProxy.cpp b/Engine/Source/Runtime/Engine/Private/SkinningSceneExtensionProxy.cpp
--- a/Engine/Source/Runtime/Engine/Private/SkinningSceneExtensionProxy.cpp
+++ b/Engine/Source/Runtime/Engine/Private/SkinningSceneExtensionProxy.cpp
@@ -4,13 +4,18 @@
 #include "SkeletalRenderPublic.h"
 #include "SkinningDefinitions.h"
 #include "AnimationRuntime.h"
+#include "Engine/SkeletalMeshSocket.h"
 #include "Rendering/SkeletalMeshRenderData.h"
 #include "SkinnedMeshSceneProxyDesc.h"
 #include "InstancedSkinnedMeshSceneProxyDesc.h"
 
 extern TAutoConsoleVariable<int32> CVarInstancedSkinnedMeshesForceRefPose;
 
-FSkinningSceneExtensionProxy::FSkinningSceneExtensionProxy(const FSkeletalMeshObject* InMeshObject, const USkinnedAsset* InSkinnedAsset, bool bAllowScaling)
+FSkinningSceneExtensionProxy::FSkinningSceneExtensionProxy(
+	const FSkeletalMeshObject* InMeshObject,
+	const USkinnedAsset* InSkinnedAsset,
+	bool bAllowScaling,
+	bool bInIncludeSocketsInBoneMap)
 	: SkinnedAsset(InSkinnedAsset)
 	, MeshObject(InMeshObject)
 {
@@ -18,6 +23,7 @@ FSkinningSceneExtensionProxy::FSkinningSceneExtensionProxy(const FSkeletalMeshOb
 	MaxBoneInfluenceCount = RenderData.GetNumBoneInfluences();
 	BoneTransformStorageMode = MeshObject->GetBoneTransformStorageMode();
 	bHasUnifiedBoneMap = RenderData.HasUnifiedBoneMap();
+	bHasSocketsInBoneMap = bInIncludeSocketsInBoneMap;
 
 	if (BoneTransformStorageMode == EBoneTransformStorageMode::Direct)
 	{
@@ -105,6 +111,30 @@ void FSkinningSceneExtensionProxy::InitForDirectStorage(bool bAllowScaling)
 	}
 }
 
+void FSkinningSceneExtensionProxy::GetSocketBoneIndices(const USkinnedAsset* InSkinnedAsset, TArray<FBoneIndexType>& OutBoneIndices)
+{
+	if (!InSkinnedAsset)
+	{
+		return;
+	}
+
+	const FReferenceSkeleton& RefSkeleton = InSkinnedAsset->GetRefSkeleton();
+	const int32 NumRawBones = RefSkeleton.GetRawBoneNum();
+	const TArray<USkeletalMeshSocket*> Sockets = InSkinnedAsset->GetActiveSocketList();
+
+	for (const USkeletalMeshSocket* Socket : Sockets)
+	{
+		if (Socket)
+		{
+			const int32 BoneIndex = RefSkeleton.FindBoneIndex(Socket->BoneName);
+			if (BoneIndex != INDEX_NONE && BoneIndex < NumRawBones)
+			{
+				OutBoneIndices.AddUnique(BoneIndex);
+			}
+		}
+	}
+}
+
 void FSkinningSceneExtensionProxy::InitForBoneMapStorage()
 {
 	check(BoneTransformStorageMode == EBoneTransformStorageMode::BoneMap);
@@ -141,6 +171,18 @@ void FSkinningSceneExtensionProxy::InitForBoneMapStorage()
 		}
 	}
 
+	// Append socket bones so they have TransformIndex slots in the skinning buffer.
+	if (bHasSocketsInBoneMap)
+	{
+		TArray<FBoneIndexType> SocketBones;
+		GetSocketBoneIndices(SkinnedAsset, SocketBones);
+		for (const FBoneIndexType BoneIndex : SocketBones)
+		{
+			BoneMap.AddUnique(BoneIndex);
+		}
+		MaxBoneTransformCount = BoneMap.Num();
+	}
+
 	check(MaxBoneTransformCount > 0);
 }
 
@@ -150,8 +192,13 @@ const FGuid& FSkinningSceneExtensionProxy::GetTransformProviderId() const
 	return MeshObjectProviderId;
 }
 
-FInstancedSkinningSceneExtensionProxy::FInstancedSkinningSceneExtensionProxy(const UTransformProviderData* TransformProvider, const FSkeletalMeshObject* InMeshObject, const USkinnedAsset* InSkinnedAsset, bool bAllowScaling)
-	: FSkinningSceneExtensionProxy(InMeshObject, InSkinnedAsset, bAllowScaling)
+FInstancedSkinningSceneExtensionProxy::FInstancedSkinningSceneExtensionProxy(
+	const UTransformProviderData* TransformProvider,
+	const FSkeletalMeshObject* InMeshObject,
+	const USkinnedAsset* InSkinnedAsset,
+	bool bAllowScaling,
+	bool bInIncludeSocketsInBoneMap)
+	: FSkinningSceneExtensionProxy(InMeshObject, InSkinnedAsset, bAllowScaling, bInIncludeSocketsInBoneMap)
 {
 	bUseInstancing = true;
 
diff --git a/Engine/Source/Runtime/Engine/Public/InstancedSkinnedMeshComponentHelper.h b/Engine/Source/Runtime/Engine/Public/InstancedSkinnedMeshComponentHelper.h
--- a/Engine/Source/Runtime/Engine/Public/InstancedSkinnedMeshComponentHelper.h
+++ b/Engine/Source/Runtime/Engine/Public/InstancedSkinnedMeshComponentHelper.h
@@ -53,7 +53,13 @@ FInstanceDataManagerSourceDataDesc FInstancedSkinnedMeshComponentHelper::GetComp
 
 	Flags.bHasPerInstanceHierarchyOffset = false;
 	Flags.bHasPerInstanceLocalBounds = TransformProvider ? TransformProvider->HasAnimationBounds() : false;
-	Flags.bHasPerInstanceDynamicData = InstancePrevTransforms.Num() > 0 && InstancePrevTransforms.Num() == InComponent.GetInstanceCount();
+	const bool bHasCPUPrevTransforms = InstancePrevTransforms.Num() > 0 && InstancePrevTransforms.Num() == InComponent.GetInstanceCount();
+	bool bHasAttachmentParent = false;
+	if constexpr (std::is_same_v<T, UInstancedSkinnedMeshComponent>)
+	{
+		bHasAttachmentParent = InComponent.HasBoneAttachmentParents();
+	}
+	Flags.bHasPerInstanceDynamicData = bHasCPUPrevTransforms || bHasAttachmentParent;
 	Flags.bHasPerInstanceSkinningData = true;
 
 	Flags.bHasPerInstanceLMSMUVBias = false;//IsStaticLightingAllowed();
@@ -82,7 +88,17 @@ FInstanceDataManagerSourceDataDesc FInstancedSkinnedMeshComponentHelper::GetComp
 	{
 		// publish data
 		ChangeSet.GetTransformWriter().Gather([&InComponent](int32 InstanceIndex) -> FRenderTransform { return FRenderTransform(InComponent.InstanceData[InstanceIndex].Transform.ToMatrixWithScale()); });
-		ChangeSet.GetPrevTransformWriter().Gather([InstancePrevTransforms](int32 InstanceIndex) -> FRenderTransform { return FRenderTransform(InstancePrevTransforms[InstanceIndex].ToMatrixWithScale()); });
+
+		if (InstancePrevTransforms.Num() == InComponent.GetInstanceCount())
+		{
+			ChangeSet.GetPrevTransformWriter().Gather([InstancePrevTransforms](int32 InstanceIndex) -> FRenderTransform { return FRenderTransform(InstancePrevTransforms[InstanceIndex].ToMatrixWithScale()); });
+		}
+		else if (ChangeSet.Flags.bHasPerInstanceDynamicData)
+		{
+			// No CPU previous transforms, but dynamic data is required (e.g., GPU-driven bone attachments).
+			// Seed with current transforms so the first frame has no motion blur artifacts.
+			ChangeSet.GetPrevTransformWriter().Gather([&InComponent](int32 InstanceIndex) -> FRenderTransform { return FRenderTransform(InComponent.InstanceData[InstanceIndex].Transform.ToMatrixWithScale()); });
+		}
 		ChangeSet.GetCustomDataWriter().Gather(MakeArrayView(InComponent.InstanceCustomData), InComponent.NumCustomDataFloats);
 
 		if (TransformProvider)
@@ -142,6 +158,33 @@ FInstanceDataManagerSourceDataDesc FInstancedSkinnedMeshComponentHelper::GetComp
 template <class T>
 FBoxSphereBounds FInstancedSkinnedMeshComponentHelper::CalcBounds(const T& InComponent, const FTransform& LocalToWorld)
 {
+	// Helper to union parent attachment bounds for child-driven bone attachments.
+	// Expands parent bounds by child mesh extent to account for geometry extending
+	// beyond the attachment point (e.g., a rifle attached to a hand bone).
+	auto UnionParentBounds = [&InComponent](FBoxSphereBounds InBounds) -> FBoxSphereBounds
+	{
+		if constexpr (std::is_same_v<T, UInstancedSkinnedMeshComponent>)
+		{
+			if (InComponent.BoneAttachmentSockets.Num() > 0)
+			{
+				const USkinnedAsset* ChildAsset = InComponent.GetSkinnedAsset();
+				const float ChildMeshExtent = ChildAsset ? ChildAsset->GetBounds().SphereRadius : 0.0f;
+
+				for (const auto& Socket : InComponent.BoneAttachmentSockets)
+				{
+					if (Socket.Parent.IsValid())
+					{
+						FBoxSphereBounds ExpandedParentBounds = Socket.Parent->Bounds;
+						ExpandedParentBounds.SphereRadius += ChildMeshExtent;
+						ExpandedParentBounds.BoxExtent += FVector(ChildMeshExtent);
+						InBounds = InBounds + ExpandedParentBounds;
+					}
+				}
+			}
+		}
+		return InBounds;
+	};
+
 	const USkinnedAsset* SkinnedAssetPtr = InComponent.GetSkinnedAsset();
 	if (SkinnedAssetPtr && InComponent.InstanceData.Num() > 0)
 	{
@@ -195,7 +238,7 @@ FBoxSphereBounds FInstancedSkinnedMeshComponentHelper::CalcBounds(const T& InCom
 			// Only use bounds if valid, otherwise use the skinned asset bounds in ref pose.
 			if (BoundsBuilder.IsValid())
 			{
-				return BoundsBuilder;
+				return UnionParentBounds(BoundsBuilder);
 			}
 		}
 
@@ -207,11 +250,11 @@ FBoxSphereBounds FInstancedSkinnedMeshComponentHelper::CalcBounds(const T& InCom
 				BoundsBuilder += InstanceBounds.TransformBy(FTransform(InComponent.InstanceData[InstanceIndex].Transform) * LocalToWorld);
 			}
 
-			return BoundsBuilder;
+			return UnionParentBounds(BoundsBuilder);
 		}
 	}
 
-	return InComponent.CalcMeshBound(FVector3f::ZeroVector, false, LocalToWorld);
+	return UnionParentBounds(InComponent.CalcMeshBound(FVector3f::ZeroVector, false, LocalToWorld));
 }
 
 template <class T>
diff --git a/Engine/Source/Runtime/Engine/Public/InstancedSkinnedMeshSceneProxyDesc.h b/Engine/Source/Runtime/Engine/Public/InstancedSkinnedMeshSceneProxyDesc.h
--- a/Engine/Source/Runtime/Engine/Public/InstancedSkinnedMeshSceneProxyDesc.h
+++ b/Engine/Source/Runtime/Engine/Public/InstancedSkinnedMeshSceneProxyDesc.h
@@ -41,4 +41,5 @@ struct FInstancedSkinnedMeshSceneProxyDesc : public FSkinnedMeshSceneProxyDesc
 #endif
 
 	bool bUseGpuLodSelection = true;
+	bool bForceAnimateSockets = false;
 };
diff --git a/Engine/Source/Runtime/Engine/Public/SkinningSceneExtensionProxy.h b/Engine/Source/Runtime/Engine/Public/SkinningSceneExtensionProxy.h
--- a/Engine/Source/Runtime/Engine/Public/SkinningSceneExtensionProxy.h
+++ b/Engine/Source/Runtime/Engine/Public/SkinningSceneExtensionProxy.h
@@ -3,6 +3,7 @@
 #pragma once
 
 #include "CoreMinimal.h"
+#include "Matrix3x4.h"
 #include "PrimitiveSceneProxy.h"
 #include "SkeletalRenderPublic.h"
 
@@ -11,13 +12,18 @@ class FSkeletalMeshRenderData;
 class UTransformProviderData;
 class FTransformProviderRenderProxy;
 class FSkinningSceneExtension;
+class FPrimitiveSceneInfo;
 struct FSkinnedMeshSceneProxyDesc;
 struct FInstancedSkinnedMeshSceneProxyDesc;
 
 class FSkinningSceneExtensionProxy
 {
 public:
-	ENGINE_API FSkinningSceneExtensionProxy(const FSkeletalMeshObject* InMeshObject, const USkinnedAsset* InSkinnedAsset, bool bAllowScaling);
+	ENGINE_API FSkinningSceneExtensionProxy(
+		const FSkeletalMeshObject* InMeshObject,
+		const USkinnedAsset* InSkinnedAsset,
+		bool bAllowScaling,
+		bool bIncludeSocketsInBoneMap = false);
 
 	virtual ~FSkinningSceneExtensionProxy() = default;
 
@@ -92,6 +98,11 @@ public:
 		return bHasUnifiedBoneMap;
 	}
 
+	bool HasSocketsInBoneMap() const
+	{
+		return bHasSocketsInBoneMap;
+	}
+
 	bool UseSkeletonBatching() const
 	{
 		return bUseSkeletonBatching;
@@ -113,6 +124,9 @@ public:
 		return 4 /* quat */ + 3 /* XYZ translation */ + (HasScale() ? 3 : 0 /* XYZ scale */);
 	}
 
+	/** Collect bone indices referenced by sockets on the skinned asset. Only raw bones (< GetRawBoneNum) are included. */
+	static ENGINE_API void GetSocketBoneIndices(const USkinnedAsset* InSkinnedAsset, TArray<FBoneIndexType>& OutBoneIndices);
+
 	virtual void CreateRenderThreadResources(FSceneInterface& Scene, FRHICommandListBase& RHICmdList) {}
 
 	virtual void DestroyRenderThreadResources() {}
@@ -154,10 +168,11 @@ protected:
 
 	EBoneTransformStorageMode BoneTransformStorageMode = EBoneTransformStorageMode::Direct;
 
-	uint8 bHasScale            : 1 = false;
-	uint8 bHasUnifiedBoneMap   : 1 = false;
-	uint8 bUseSkeletonBatching : 1 = false;
-	uint8 bUseInstancing       : 1 = false;
+	uint8 bHasScale                     : 1 = false;
+	uint8 bHasUnifiedBoneMap            : 1 = false;
+	uint8 bHasSocketsInBoneMap          : 1 = false;
+	uint8 bUseSkeletonBatching          : 1 = false;
+	uint8 bUseInstancing                : 1 = false;
 
 	friend class FSkinningSceneExtension;
 };
@@ -168,10 +183,85 @@ public:
 	virtual ~FInstancedSkinningSceneExtensionProxyUserData() = default;
 };
 
+class FBoneAttachmentBinding
+{
+	static const uint32 Unattached   = ~0u;
+	static const uint32 MaxSockets   = (1 << 8 ) - 1;
+	static const uint32 MaxInstances = (1 << 24) - 1;
+
+public:
+	FBoneAttachmentBinding() = default;
+
+	FBoneAttachmentBinding(uint32 InSocketIndex, uint32 InParentInstanceIndex)
+	{
+		check(InSocketIndex < MaxSockets);
+		check(InParentInstanceIndex < MaxInstances);
+		Bits.SocketIndex   = InSocketIndex;
+		Bits.InstanceIndex = InParentInstanceIndex;
+	}
+
+	uint32 GetSocketIndex() const
+	{
+		check(IsAttached());
+		return Bits.SocketIndex;
+	}
+
+	uint32 GetParentInstanceIndex() const
+	{
+		check(IsAttached());
+		return Bits.InstanceIndex;
+	}
+
+	bool IsAttached() const
+	{
+		return Bits.Packed != Unattached;
+	}
+
+	bool operator==(const FBoneAttachmentBinding& Other) const { return Bits.Packed == Other.Bits.Packed; }
+	bool operator!=(const FBoneAttachmentBinding& Other) const { return Bits.Packed != Other.Bits.Packed; }
+
+	friend FArchive& operator<<(FArchive& Ar, FBoneAttachmentBinding& Binding)
+	{
+		Ar << Binding.Bits.Packed;
+		return Ar;
+	}
+
+private:
+	union
+	{
+		struct
+		{
+			uint32 InstanceIndex : 24;
+			uint32 SocketIndex   :  8;
+		};
+		uint32 Packed = Unattached;
+	} Bits;
+};
+
+/** Fully resolved bone attachment socket data for the render thread. */
+struct FBoneAttachmentSocket
+{
+	FBoneAttachmentSocket()
+	{
+		RefPoseMatrix.SetIdentity();
+	}
+
+	FPrimitiveComponentId ParentComponentId;
+	uint16 BoneIndex = INVALID_BONE_INDEX;
+	FMatrix3x4 RefPoseMatrix;
+	FQuat4f Rotation = FQuat4f::Identity;
+	FVector3f Translation = FVector3f::ZeroVector;
+};
+
 class FInstancedSkinningSceneExtensionProxy : public FSkinningSceneExtensionProxy
 {
 public:
-	ENGINE_API FInstancedSkinningSceneExtensionProxy(const UTransformProviderData* InTransformProvider, const FSkeletalMeshObject* InMeshObject, const USkinnedAsset* InSkinnedAsset, bool bAllowScaling);
+	ENGINE_API FInstancedSkinningSceneExtensionProxy(
+		const UTransformProviderData* InTransformProvider,
+		const FSkeletalMeshObject* InMeshObject,
+		const USkinnedAsset* InSkinnedAsset,
+		bool bAllowScaling,
+		bool bIncludeSocketsInBoneMap);
 
 	ENGINE_API void CreateRenderThreadResources(FSceneInterface& Scene, FRHICommandListBase& RHICmdList);
 	ENGINE_API void DestroyRenderThreadResources();
@@ -204,8 +294,53 @@ public:
 		return *UserData;
 	}
 
+	void SetBoneAttachmentSockets(TArray<FBoneAttachmentSocket>&& InSockets)
+	{
+		BoneAttachmentSockets = MoveTemp(InSockets);
+	}
+
+	void SetBoneAttachmentBindings(TArray<FBoneAttachmentBinding>&& InBindings)
+	{
+		BoneAttachmentBindings = MoveTemp(InBindings);
+	}
+
+	/** Resize the bindings array to match instance count. */
+	void ResizeBoneAttachmentBindings(int32 NewNum)
+	{
+		BoneAttachmentBindings.SetNum(NewNum);
+	}
+
+	/** Update a single binding entry. Index must be in range. */
+	void SetBoneAttachmentBinding(int32 Index, FBoneAttachmentBinding Binding)
+	{
+		check(BoneAttachmentBindings.IsValidIndex(Index));
+		BoneAttachmentBindings[Index] = Binding;
+	}
+
+	TConstArrayView<FBoneAttachmentSocket> GetBoneAttachmentSockets() const
+	{
+		return BoneAttachmentSockets;
+	}
+
+	TConstArrayView<FBoneAttachmentBinding> GetBoneAttachmentBindings() const
+	{
+		return BoneAttachmentBindings;
+	}
+
+	int32 GetNumBoneAttachmentBindings() const
+	{
+		return BoneAttachmentBindings.Num();
+	}
+
+	bool HasBoneAttachments() const
+	{
+		return BoneAttachmentSockets.Num() > 0;
+	}
+
 protected:
 	FTransformProviderRenderProxy* TransformProviderProxy = nullptr;
 	FInstancedSkinningSceneExtensionProxyUserData* UserData = nullptr;
 	FGuid TransformProviderId;
+	TArray<FBoneAttachmentSocket> BoneAttachmentSockets;
+	TArray<FBoneAttachmentBinding> BoneAttachmentBindings;
 };
\ No newline at end of file
diff --git a/Engine/Source/Runtime/Renderer/Private/GPUScene.cpp b/Engine/Source/Runtime/Renderer/Private/GPUScene.cpp
--- a/Engine/Source/Runtime/Renderer/Private/GPUScene.cpp
+++ b/Engine/Source/Runtime/Renderer/Private/GPUScene.cpp
@@ -5,6 +5,7 @@
 =============================================================================*/
 
 #include "GPUScene.h"
+#include "Skinning/SkinningSceneExtension.h"
 #include "CoreMinimal.h"
 #include "DataDrivenShaderPlatformInfo.h"
 #include "RHI.h"
@@ -765,7 +766,7 @@ void FGPUScene::InitLightData(const FLightSceneInfoCompact& LightInfoCompact, bo
 	DataOut.InverseExposureBlend = LightParams.InverseExposureBlend;
 }
 
-void FGPUScene::UpdateInternal(FRDGBuilder& GraphBuilder, FRDGExternalAccessQueue& ExternalAccessQueue, const UE::Tasks::FTask& UpdateTaskPrerequisites, const FUpdateFromComputeCommands& UpdatesFromCompute)
+void FGPUScene::UpdateInternal(FRDGBuilder& GraphBuilder, FRDGExternalAccessQueue& ExternalAccessQueue, const UE::Tasks::FTask& UpdateTaskPrerequisites, const FUpdateFromComputeCommands& UpdatesFromCompute, FSceneExtensionsUpdaters& SceneExtensionsUpdaters)
 {
 	LLM_SCOPE_BYTAG(GPUScene);
 
@@ -902,7 +903,9 @@ void FGPUScene::UpdateInternal(FRDGBuilder& GraphBuilder, FRDGExternalAccessQueu
 	}
 
 	// Update instance data using GPU compute.
-	if (!UpdatesFromCompute.IsEmpty())
+	auto SkinningUpdater = SceneExtensionsUpdaters.GetUpdaterPtr<FSkinningSceneExtension::FUpdater>();
+
+	if (!UpdatesFromCompute.IsEmpty() || SkinningUpdater)
 	{
 		RDG_EVENT_SCOPE(GraphBuilder, "UpdateGPUScene UpdateInstancesFromCompute");
 
@@ -923,6 +926,11 @@ void FGPUScene::UpdateInternal(FRDGBuilder& GraphBuilder, FRDGExternalAccessQueu
 
 			Command.Payload.GPUSceneWriter.Execute(GraphBuilder, Params);
 		}
+
+		if (SkinningUpdater)
+		{
+			SkinningUpdater->ResolveAttachments(GraphBuilder, Params);
+		}
 	}
 }
 
@@ -1772,19 +1780,21 @@ void FGPUScene::AddPrimitiveToUpdate(FPersistentPrimitiveIndex PersistentPrimiti
 }
 
 
-void FGPUScene::Update(FRDGBuilder& GraphBuilder, FRDGExternalAccessQueue& ExternalAccessQueue, const FUpdateFromComputeCommands& UpdatesFromCompute, const UE::Tasks::FTask& UpdateTaskPrerequisites)
+void FGPUScene::Update(FRDGBuilder& GraphBuilder, FRDGExternalAccessQueue& ExternalAccessQueue, const FUpdateFromComputeCommands& UpdatesFromCompute, const UE::Tasks::FTask& UpdateTaskPrerequisites, FSceneExtensionsUpdaters& SceneExtensionsUpdaters)
 {
-	if (bIsEnabled)
+	if (!bIsEnabled)
 	{
-		RDG_GPU_MASK_SCOPE(GraphBuilder, FRHIGPUMask::All());
+		return;
+	}
 
-		ensure(!bInBeginEndBlock);
+	RDG_GPU_MASK_SCOPE(GraphBuilder, FRHIGPUMask::All());
 
-		// Default state when updated (no "dynamic primitives" pushed)
-		DynamicPrimitivesOffset = Scene.GetMaxPersistentPrimitiveIndex();
+	ensure(!bInBeginEndBlock);
 
-		UpdateInternal(GraphBuilder, ExternalAccessQueue, UpdateTaskPrerequisites, UpdatesFromCompute);
-	}
+	// Default state when updated (no "dynamic primitives" pushed)
+	DynamicPrimitivesOffset = Scene.GetMaxPersistentPrimitiveIndex();
+
+	UpdateInternal(GraphBuilder, ExternalAccessQueue, UpdateTaskPrerequisites, UpdatesFromCompute, SceneExtensionsUpdaters);
 }
 
 void FGPUScene::UploadDynamicPrimitiveShaderDataForView(FRDGBuilder& GraphBuilder, FViewInfo& View, bool bRayTracing, UE::Renderer::Private::IShadowInvalidatingInstances *ShadowInvalidatingInstances)
diff --git a/Engine/Source/Runtime/Renderer/Private/GPUScene.h b/Engine/Source/Runtime/Renderer/Private/GPUScene.h
--- a/Engine/Source/Runtime/Renderer/Private/GPUScene.h
+++ b/Engine/Source/Runtime/Renderer/Private/GPUScene.h
@@ -32,6 +32,7 @@ class FScenePostUpdateChangeSet;
 class FPrimitiveSceneProxy;
 struct FLightSceneChangeSet;
 class FGPUSceneWriterParameters;
+class FSceneExtensionsUpdaters;
 
 namespace UE::Renderer::Private
 {
@@ -263,7 +264,7 @@ public:
 	/**
 	 * Pull all pending updates from Scene and upload primitive & instance data.
 	 */
-	void Update(FRDGBuilder& GraphBuilder, FRDGExternalAccessQueue& ExternalAccessQueue, const FUpdateFromComputeCommands& UpdatesFromCompute, const UE::Tasks::FTask& UpdateTaskPrerequisites = {});
+	void Update(FRDGBuilder& GraphBuilder, FRDGExternalAccessQueue& ExternalAccessQueue, const FUpdateFromComputeCommands& UpdatesFromCompute, const UE::Tasks::FTask& UpdateTaskPrerequisites, FSceneExtensionsUpdaters& SceneExtensionsUpdaters);
 
 	/**
 	 * Queue the given primitive for upload to GPU at next call to Update.
@@ -492,7 +493,7 @@ private:
 
 	void UploadDynamicPrimitiveShaderDataForViewInternal(FRDGBuilder& GraphBuilder, FViewInfo& View, bool bRayTracing, UE::Renderer::Private::IShadowInvalidatingInstances *ShadowInvalidatingInstances);
 
-	void UpdateInternal(FRDGBuilder& GraphBuilder, FRDGExternalAccessQueue& ExternalAccessQueue, const UE::Tasks::FTask& UpdateTaskPrerequisites, const FUpdateFromComputeCommands& UpdatesFromCompute);
+	void UpdateInternal(FRDGBuilder& GraphBuilder, FRDGExternalAccessQueue& ExternalAccessQueue, const UE::Tasks::FTask& UpdateTaskPrerequisites, const FUpdateFromComputeCommands& UpdatesFromCompute, FSceneExtensionsUpdaters& SceneExtensionsUpdaters);
 
 	void AddUpdatePrimitiveIdsPass(FRDGBuilder& GraphBuilder, FInstanceGPULoadBalancer& IdOnlyUpdateItems);
 
diff --git a/Engine/Source/Runtime/Renderer/Private/RendererScene.cpp b/Engine/Source/Runtime/Renderer/Private/RendererScene.cpp
--- a/Engine/Source/Runtime/Renderer/Private/RendererScene.cpp
+++ b/Engine/Source/Runtime/Renderer/Private/RendererScene.cpp
@@ -6661,7 +6661,7 @@ void FScene::Update(FRDGBuilder& GraphBuilder, const FUpdateParameters& Paramete
 
 		FRDGExternalAccessQueue ExternalAccessQueue;
 
-		GPUScene.Update(GraphBuilder, ExternalAccessQueue, SceneUpdateChangeSetStorage.PrimitiveUpdates.GetRangeView<FUpdateInstanceFromComputeCommand>(), Parameters.GPUSceneUpdateTaskPrerequisites);
+		GPUScene.Update(GraphBuilder, ExternalAccessQueue, SceneUpdateChangeSetStorage.PrimitiveUpdates.GetRangeView<FUpdateInstanceFromComputeCommand>(), Parameters.GPUSceneUpdateTaskPrerequisites, SceneExtensionsUpdaters);
 
 		ExternalAccessQueue.Submit(GraphBuilder);
 	}
diff --git a/Engine/Source/Runtime/Renderer/Private/Skinning/AnimSequenceTransformProvider.cpp b/Engine/Source/Runtime/Renderer/Private/Skinning/AnimSequenceTransformProvider.cpp
--- a/Engine/Source/Runtime/Renderer/Private/Skinning/AnimSequenceTransformProvider.cpp
+++ b/Engine/Source/Runtime/Renderer/Private/Skinning/AnimSequenceTransformProvider.cpp
@@ -398,13 +398,39 @@ void FAnimSequenceTransformProvider::InitExtension(FScene& InScene)
 void FAnimSequenceTransformProvider::BuildBoneTables(
 	FBoneTables& BoneTables,
 	const USkinnedAsset* SkinnedMesh,
-	const FSkeletalMeshLODRenderData& LODRenderData)
+	const FSkeletalMeshLODRenderData& LODRenderData,
+	bool bHasSocketsInBoneMap)
 {
 	const USkeleton* TargetSkeleton = SkinnedMesh->GetSkeleton();
 	const FReferenceSkeleton& MeshRefSkeleton = SkinnedMesh->GetRefSkeleton();
 	const FReferenceSkeleton& TargetSkeletonRefSkeleton = TargetSkeleton->GetReferenceSkeleton();
 	const uint32 NumMeshBones = MeshRefSkeleton.GetNum();
-	const uint32 NumCompactBones = LODRenderData.RequiredBones.Num();
+
+	// Start with the LOD's RequiredBones. Optionally inject force-animated socket bones + parent chains.
+	TArray<FBoneIndexType> ExtendedRequiredBones;
+	bool bUseExtended = false;
+
+	if (bHasSocketsInBoneMap)
+	{
+		TArray<FBoneIndexType> SocketBones;
+		FSkinningSceneExtensionProxy::GetSocketBoneIndices(SkinnedMesh, SocketBones);
+
+		if (SocketBones.Num() > 0)
+		{
+			ExtendedRequiredBones = LODRenderData.RequiredBones;
+			for (const FBoneIndexType BoneIndex : SocketBones)
+			{
+				ExtendedRequiredBones.AddUnique(BoneIndex);
+			}
+			ExtendedRequiredBones.Sort();
+			FAnimationRuntime::EnsureParentsPresent(ExtendedRequiredBones, MeshRefSkeleton);
+			ExtendedRequiredBones.Sort();
+			bUseExtended = true;
+		}
+	}
+
+	TConstArrayView<FBoneIndexType> RequiredBones = bUseExtended ? TConstArrayView<FBoneIndexType>(ExtendedRequiredBones) : TConstArrayView<FBoneIndexType>(LODRenderData.RequiredBones);
+	const uint32 NumCompactBones = RequiredBones.Num();
 
 	BoneTables.MeshToCompact.Init(INVALID_BONE_INDEX, NumMeshBones);
 	BoneTables.CompactToMesh.SetNumUninitialized(NumCompactBones);
@@ -415,12 +441,12 @@ void FAnimSequenceTransformProvider::BuildBoneTables(
 
 	for (uint32 CompactBoneIndex = 0; CompactBoneIndex < NumCompactBones; ++CompactBoneIndex)
 	{
-		BoneTables.MeshToCompact[LODRenderData.RequiredBones[CompactBoneIndex]] = CompactBoneIndex;
+		BoneTables.MeshToCompact[RequiredBones[CompactBoneIndex]] = CompactBoneIndex;
 	}
 
 	for (uint32 CompactBoneIndex = 0; CompactBoneIndex < NumCompactBones; ++CompactBoneIndex)
 	{
-		const uint16 MeshBoneIndex = LODRenderData.RequiredBones[CompactBoneIndex];
+		const uint16 MeshBoneIndex = RequiredBones[CompactBoneIndex];
 		BoneTables.CompactToMesh[CompactBoneIndex] = MeshBoneIndex;
 
 		const uint16 MeshParentBoneIndex = MeshRefSkeleton.GetParentIndex(MeshBoneIndex);
@@ -718,7 +744,7 @@ void FAnimSequenceTransformProvider::RegisterProxy(FAnimSequenceTransformProvide
 	for (int32 LODIndex = MinLODIndex; LODIndex < NumLODs; ++LODIndex)
 	{
 		FProxyUserData::FLOD& LOD = ProxyUserData->LODs[LODIndex];
-		BuildBoneTables(LOD.BoneTables, SkinnedAsset, RenderData->LODRenderData[LODIndex]);
+		BuildBoneTables(LOD.BoneTables, SkinnedAsset, RenderData->LODRenderData[LODIndex], ExtensionProxy->HasSocketsInBoneMap());
 
 		// Single pass mode skips building leaf bone paths in favor of doing everything in one compute dispatch.
 		ComputeSinglePassGroupParams(LOD.BoneTables.NumCompactBones, LOD.GroupSize, LOD.InstancesPerGroup, LOD.SlotsPerInstance);
@@ -771,7 +797,7 @@ void FAnimSequenceTransformProvider::RegisterProxy(FAnimSequenceTransformProvide
 
 		for (int32 LODIndex = MinLODIndex; LODIndex < NumLODs; ++LODIndex)
 		{
-			const FMappingRegistryKey MappingRegistryKey(SourceSkeleton, SkinnedAsset, LODIndex);
+			const FMappingRegistryKey MappingRegistryKey(SourceSkeleton, SkinnedAsset, LODIndex, ExtensionProxy->HasSocketsInBoneMap());
 			TUniquePtr<FMapping>& Mapping = MappingRegistry.FindOrAdd(MappingRegistryKey);
 			if (!Mapping)
 			{
@@ -782,6 +808,7 @@ void FAnimSequenceTransformProvider::RegisterProxy(FAnimSequenceTransformProvide
 				Mapping->SkinnedAsset    = SkinnedAsset;
 				Mapping->NumMeshBones    = ProxyUserData->NumMeshBones;
 				Mapping->LODIndex        = LODIndex;
+				Mapping->bHasSocketsInBoneMap = ExtensionProxy->HasSocketsInBoneMap();
 				Mapping->NumCompactBones = LOD.BoneTables.NumCompactBones;
 				BuildRetargetingTables(Mapping->RetargetingTables, SourceSkeleton, SkinnedAsset, LOD.BoneTables);
 
@@ -872,7 +899,7 @@ void FAnimSequenceTransformProvider::UnregisterProxy(FAnimSequenceTransformProvi
 			}
 
 			MappingsToUpload.RemoveSingleSwap(Mapping, EAllowShrinking::No);
-			MappingRegistry.Remove(FMappingRegistryKey(Mapping->SourceSkeleton, Mapping->SkinnedAsset, Mapping->LODIndex));
+			MappingRegistry.Remove(FMappingRegistryKey(Mapping->SourceSkeleton, Mapping->SkinnedAsset, Mapping->LODIndex, Mapping->bHasSocketsInBoneMap));
 			Mapping = nullptr;
 
 			DEC_DWORD_STAT(STAT_AnimSequenceTransformProvider_Mappings);
diff --git a/Engine/Source/Runtime/Renderer/Private/Skinning/AnimSequenceTransformProvider.h b/Engine/Source/Runtime/Renderer/Private/Skinning/AnimSequenceTransformProvider.h
--- a/Engine/Source/Runtime/Renderer/Private/Skinning/AnimSequenceTransformProvider.h
+++ b/Engine/Source/Runtime/Renderer/Private/Skinning/AnimSequenceTransformProvider.h
@@ -78,7 +78,8 @@ private:
 	void BuildBoneTables(
 		FBoneTables& BoneTables,
 		const USkinnedAsset* SkinnedMesh,
-		const FSkeletalMeshLODRenderData& LODRenderData);
+		const FSkeletalMeshLODRenderData& LODRenderData,
+		bool bHasSocketsInBoneMap);
 
 	struct FRetargetingTables
 	{
@@ -155,24 +156,26 @@ private:
 	{
 		FMappingRegistryKey() = default;
 
-		FMappingRegistryKey(const USkeleton* InSourceSkeleton, const USkinnedAsset* InSkinnedAsset, uint32 InLODIndex)
+		FMappingRegistryKey(const USkeleton* InSourceSkeleton, const USkinnedAsset* InSkinnedAsset, uint32 InLODIndex, bool bInHasSocketsInBoneMap = false)
 			: SourceSkeleton(InSourceSkeleton)
 			, SkinnedAsset(InSkinnedAsset)
 			, LODIndex(InLODIndex)
+			, bHasSocketsInBoneMap(bInHasSocketsInBoneMap)
 		{}
 
 		const USkeleton* SourceSkeleton = nullptr;
 		const USkinnedAsset* SkinnedAsset = nullptr;
 		uint32 LODIndex = 0;
+		bool bHasSocketsInBoneMap = false;
 
 		bool operator==(const FMappingRegistryKey& Other) const
 		{
-			return SourceSkeleton == Other.SourceSkeleton && SkinnedAsset == Other.SkinnedAsset && LODIndex == Other.LODIndex;
+			return SourceSkeleton == Other.SourceSkeleton && SkinnedAsset == Other.SkinnedAsset && LODIndex == Other.LODIndex && bHasSocketsInBoneMap == Other.bHasSocketsInBoneMap;
 		}
 
 		friend uint32 GetTypeHash(const FMappingRegistryKey& Key)
 		{
-			return PointerHash(Key.SourceSkeleton, PointerHash(Key.SkinnedAsset, Key.LODIndex));
+			return PointerHash(Key.SourceSkeleton, PointerHash(Key.SkinnedAsset, HashCombine(Key.LODIndex, uint32(Key.bHasSocketsInBoneMap))));
 		}
 	};
 
@@ -230,6 +233,7 @@ private:
 		uint32 NumMeshBones = 0;
 		uint32 NumCompactBones = 0;
 		uint32 RefCount = 0;
+		bool bHasSocketsInBoneMap = false;
 	};
 
 	struct FTrackData
diff --git a/Engine/Source/Runtime/Renderer/Private/Skinning/SkinningSceneExtension.cpp b/Engine/Source/Runtime/Renderer/Private/Skinning/SkinningSceneExtension.cpp
--- a/Engine/Source/Runtime/Renderer/Private/Skinning/SkinningSceneExtension.cpp
+++ b/Engine/Source/Runtime/Renderer/Private/Skinning/SkinningSceneExtension.cpp
@@ -13,6 +13,8 @@
 #include "Rendering/SkeletalMeshLODRenderData.h"
 #include "SkinningDefinitions.h"
 #include "SkinningSceneExtensionProxy.h"
+#include "GPUSceneWriter.h"
+#include "GPUScene.h"
 
 DECLARE_STATS_GROUP(TEXT("SkinningSceneExtension"), STATGROUP_SkinningSceneExtension, STATCAT_Advanced);
 
@@ -1650,6 +1652,16 @@ void FSkinningSceneExtension::FUpdater::PreGPUSceneUpdate(FRDGBuilder& GraphBuil
 	SceneData->PerformSkinning(Parameters, GraphBuilder, CurrentTime);
 }
 
+void FSkinningSceneExtension::FUpdater::ResolveAttachments(FRDGBuilder& GraphBuilder, const FGPUSceneWriteDelegateParams& Params)
+{
+	if (!SceneData->IsEnabled())
+	{
+		return;
+	}
+
+	SceneData->ResolveAttachments(GraphBuilder, Params);
+}
+
 void FSkinningSceneExtension::FRenderer::UpdateSceneUniformBuffer(FRDGBuilder& GraphBuilder, FSceneUniformBuffer& SceneUniformBuffer)
 {
 	SCOPED_NAMED_EVENT(FSkinningSceneExtension_FRenderer_UpdateSceneUniformBuffer, FColor::Silver);
@@ -1907,3 +1919,258 @@ void FSkinningSceneExtension::ProvideMeshObjectTransforms(FSkinningTransformProv
 	IMPLEMENT_STATIC_UNIFORM_BUFFER_SCENE_UB(StructType, SceneUbEx##FieldName)
 
 IMPLEMENT_SCENE_UB_STRUCT_EX(FSkinningSceneParameters, Skinning, GetDefaultSkinningParameters);
+
+///////////////////////////////////////////////////////////////////////////////////////////////////
+// Bone Attachment Resolve
+
+class FResolveAttachmentsCS : public FGlobalShader
+{
+public:
+	DECLARE_GLOBAL_SHADER(FResolveAttachmentsCS);
+	SHADER_USE_PARAMETER_STRUCT(FResolveAttachmentsCS, FGlobalShader);
+
+	static bool ShouldCompilePermutation(const FGlobalShaderPermutationParameters& Parameters)
+	{
+		return true;
+	}
+
+	static void ModifyCompilationEnvironment(const FGlobalShaderPermutationParameters& Parameters, FShaderCompilerEnvironment& OutEnvironment)
+	{
+		FGlobalShader::ModifyCompilationEnvironment(Parameters, OutEnvironment);
+		OutEnvironment.CompilerFlags.Add(CFLAG_WarningsAsErrors);
+		OutEnvironment.CompilerFlags.Add(CFLAG_HLSL2021);
+		OutEnvironment.SetDefine(TEXT("VF_SUPPORTS_PRIMITIVE_SCENE_DATA"), 1);
+	}
+
+	BEGIN_SHADER_PARAMETER_STRUCT(FParameters, )
+		SHADER_PARAMETER_RDG_BUFFER_SRV(StructuredBuffer<UE::HLSL::FResolveAttachmentBlockHeader>, AttachmentBlockHeaderBuffer)
+		SHADER_PARAMETER_RDG_COMPRESSED_BONE_TRANSFORM_SRV(SkinningTransformBuffer)
+		SHADER_PARAMETER_STRUCT_INCLUDE(FGPUSceneWriterParameters, GPUSceneWriterParameters)
+		SHADER_PARAMETER(uint32, NumAttachments)
+	END_SHADER_PARAMETER_STRUCT()
+};
+IMPLEMENT_GLOBAL_SHADER(FResolveAttachmentsCS, "/Engine/Private/Skinning/ResolveAttachments.usf", "ResolveAttachmentsCS", SF_Compute);
+
+void FSkinningSceneExtension::ResolveAttachments(FRDGBuilder& GraphBuilder, const FGPUSceneWriteDelegateParams& Params)
+{
+	TRACE_CPUPROFILER_EVENT_SCOPE(FSkinningSceneExtension::ResolveAttachments);
+
+	if (!IsEnabled())
+	{
+		return;
+	}
+
+	// Collect all attachments across all child instanced primitives.
+	uint32 NumTotalAttachments = 0;
+
+	for (const FHeaderData& Header : HeaderDatas)
+	{
+		if (Header.bIsInstanced && Header.Proxy)
+		{
+			const FInstancedSkinningSceneExtensionProxy* ChildProxy = static_cast<const FInstancedSkinningSceneExtensionProxy*>(Header.Proxy);
+			if (ChildProxy->HasBoneAttachments())
+			{
+				for (FBoneAttachmentBinding Binding : ChildProxy->GetBoneAttachmentBindings())
+				{
+					if (Binding.IsAttached())
+					{
+						++NumTotalAttachments;
+					}
+				}
+			}
+		}
+	}
+
+	if (NumTotalAttachments == 0)
+	{
+		return;
+	}
+
+	// Allocate the block header buffer and lock for write.
+	FRDGBuffer* AttachmentHeaderBuffer = GraphBuilder.CreateBuffer(
+		FRDGBufferDesc::CreateStructuredDesc(sizeof(UE::HLSL::FResolveAttachmentBlockHeader), NumTotalAttachments),
+		TEXT("SkinningSceneExtension.AttachmentHeaders"));
+
+	GraphBuilder.ConvertToExternalBuffer(AttachmentHeaderBuffer);
+
+	// Fill the headers asynchronously inside a command list setup task.
+	// Lock/Unlock happens on the command list timeline inside the task.
+	uint32& NumActualAttachments = *GraphBuilder.AllocPOD<uint32>();
+	NumActualAttachments = 0;
+
+	UE::Tasks::FTask FillTask = GraphBuilder.AddCommandListSetupTask(
+		[this, AttachmentHeaderBuffer, NumTotalAttachments, &NumActualAttachments](FRHICommandList& RHICmdList)
+	{
+		auto* RawHeaderData = reinterpret_cast<UE::HLSL::FResolveAttachmentBlockHeader*>(
+			RHICmdList.LockBuffer(AttachmentHeaderBuffer->GetRHIUnchecked(), 0,
+				NumTotalAttachments * sizeof(UE::HLSL::FResolveAttachmentBlockHeader), RLM_WriteOnly));
+
+		TRACE_CPUPROFILER_EVENT_SCOPE(FSkinningSceneExtension::FillAttachmentHeaders);
+
+		uint32 AttachmentIndex = 0;
+		TArray<uint16, SceneRenderingAllocator> ReverseBoneMap;
+
+		// Iterate child proxies that have bone attachment parent sockets.
+		for (const FHeaderData& ChildHeader : HeaderDatas)
+		{
+			if (!ChildHeader.bIsInstanced || !ChildHeader.Proxy)
+			{
+				continue;
+			}
+
+			const FInstancedSkinningSceneExtensionProxy* ChildProxy = static_cast<const FInstancedSkinningSceneExtensionProxy*>(ChildHeader.Proxy);
+			if (!ChildProxy->HasBoneAttachments())
+			{
+				continue;
+			}
+
+			const FPrimitiveSceneInfo* ChildSceneInfo = ChildHeader.PrimitiveSceneInfo;
+			if (!ChildSceneInfo)
+			{
+				continue;
+			}
+
+			// Pre-resolve each parent socket: find parent header, build reverse bone map, resolve transform bone index.
+			struct FResolvedParentSocket
+			{
+				const FHeaderData* ParentHeader = nullptr;
+				uint16 TransformBoneIndex = INVALID_BONE_INDEX;
+			};
+			TArray<FResolvedParentSocket, SceneRenderingAllocator> ResolvedSockets;
+			ResolvedSockets.SetNum(ChildProxy->GetBoneAttachmentSockets().Num());
+
+			for (int32 SocketIndex = 0; SocketIndex < ChildProxy->GetBoneAttachmentSockets().Num(); ++SocketIndex)
+			{
+				const FBoneAttachmentSocket& Socket = ChildProxy->GetBoneAttachmentSockets()[SocketIndex];
+				FResolvedParentSocket& Resolved = ResolvedSockets[SocketIndex];
+
+				if (!Socket.ParentComponentId.IsValid())
+				{
+					continue;
+				}
+
+				const FPrimitiveSceneInfo* ParentSceneInfo = Scene.GetPrimitiveSceneInfo(Socket.ParentComponentId);
+				if (!ParentSceneInfo)
+				{
+					continue;
+				}
+
+				const int32 ParentPersistentIndex = ParentSceneInfo->GetPersistentIndex().Index;
+				if (!HeaderDatas.IsValidIndex(ParentPersistentIndex))
+				{
+					continue;
+				}
+
+				const FHeaderData& ParentHeader = HeaderDatas[ParentPersistentIndex];
+				Resolved.ParentHeader = &ParentHeader;
+
+				// Build reverse bone map for this parent and resolve the bone index.
+				TConstArrayView<uint16> ParentBoneMap = ParentHeader.Proxy ? ParentHeader.Proxy->GetBoneMap() : TConstArrayView<uint16>();
+				Resolved.TransformBoneIndex = Socket.BoneIndex;
+
+				if (!ParentBoneMap.IsEmpty() && ParentHeader.Proxy)
+				{
+					ReverseBoneMap.SetNumUninitialized(ParentHeader.Proxy->GetSkinnedAsset()->GetRefSkeleton().GetNum(), EAllowShrinking::No);
+					FMemory::Memset(ReverseBoneMap.GetData(), 0xFF, ReverseBoneMap.Num() * sizeof(uint16));
+					for (int32 BoneMapIndex = 0; BoneMapIndex < ParentBoneMap.Num(); ++BoneMapIndex)
+					{
+						if (ParentBoneMap[BoneMapIndex] < (uint16)ReverseBoneMap.Num())
+						{
+							ReverseBoneMap[ParentBoneMap[BoneMapIndex]] = (uint16)BoneMapIndex;
+						}
+					}
+
+					Resolved.TransformBoneIndex = (Socket.BoneIndex < (uint16)ReverseBoneMap.Num())
+						? ReverseBoneMap[Socket.BoneIndex]
+						: INVALID_BONE_INDEX;
+				}
+			}
+
+			// Iterate child instances and emit headers for attached ones.
+			const int32 NumChildInstances = ChildProxy->GetNumBoneAttachmentBindings();
+			for (int32 ChildInstanceIndex = 0; ChildInstanceIndex < NumChildInstances; ++ChildInstanceIndex)
+			{
+				const FBoneAttachmentBinding Binding = ChildProxy->GetBoneAttachmentBindings()[ChildInstanceIndex];
+				if (!Binding.IsAttached())
+				{
+					continue;
+				}
+
+				const uint8 ParentSocketIndex = Binding.GetSocketIndex();
+				const uint32 ParentInstanceIndex = Binding.GetParentInstanceIndex();
+
+				if (!ResolvedSockets.IsValidIndex(ParentSocketIndex) || !ResolvedSockets[ParentSocketIndex].ParentHeader)
+				{
+					continue;
+				}
+
+				const FResolvedParentSocket& Resolved = ResolvedSockets[ParentSocketIndex];
+				const FHeaderData& ParentHeader = *Resolved.ParentHeader;
+				const FBoneAttachmentSocket& Socket = ChildProxy->GetBoneAttachmentSockets()[ParentSocketIndex];
+
+				// Skip attachments where the bone is not in the parent's skinning buffer (e.g., LOD-culled).
+				if (Resolved.TransformBoneIndex == INVALID_BONE_INDEX)
+				{
+					continue;
+				}
+
+				// Skip if parent instance index is out of range (parent may have had instances removed).
+				if (ParentInstanceIndex >= ParentHeader.NumInstanceSceneDataEntries)
+				{
+					continue;
+				}
+
+				const uint32 ParentInstanceId = ParentHeader.InstanceSceneDataOffset + ParentInstanceIndex;
+				const uint32 ChildGPUSceneInstanceId = ChildHeader.InstanceSceneDataOffset + ChildInstanceIndex;
+
+				UE::HLSL::FResolveAttachmentBlockHeader& BlockHeader = RawHeaderData[AttachmentIndex++];
+				BlockHeader.ParentTransformBufferOffset = ParentHeader.TransformBufferOffset;
+				SetMaxTransformCount(BlockHeader, ParentHeader.MaxTransformCount);
+				SetBoneIndex(BlockHeader, Resolved.TransformBoneIndex);
+				SetCurrentTransformSlot(BlockHeader, ParentHeader.CurrentTransformSlot);
+				SetParentInstanceId(BlockHeader, ParentInstanceId);
+				BlockHeader.ChildInstanceId = ChildGPUSceneInstanceId;
+				BlockHeader.ChildPrimitiveId = ChildSceneInfo->GetPersistentIndex().Index;
+
+				BlockHeader.LocalOffsetTranslation[0] = Socket.Translation.X;
+				BlockHeader.LocalOffsetTranslation[1] = Socket.Translation.Y;
+				BlockHeader.LocalOffsetTranslation[2] = Socket.Translation.Z;
+
+				BlockHeader.LocalOffsetRotation[0] = Socket.Rotation.X;
+				BlockHeader.LocalOffsetRotation[1] = Socket.Rotation.Y;
+				BlockHeader.LocalOffsetRotation[2] = Socket.Rotation.Z;
+				BlockHeader.LocalOffsetRotation[3] = Socket.Rotation.W;
+
+				FMemory::Memcpy(&BlockHeader.BoneRefPoseRow0, Socket.RefPoseMatrix.M, sizeof(float) * 12);
+			}
+		}
+
+		NumActualAttachments = AttachmentIndex;
+
+		RHICmdList.UnlockBuffer(AttachmentHeaderBuffer->GetRHIUnchecked());
+	});
+
+	FResolveAttachmentsCS::FParameters* PassParameters = GraphBuilder.AllocParameters<FResolveAttachmentsCS::FParameters>();
+	PassParameters->AttachmentBlockHeaderBuffer = GraphBuilder.CreateSRV(AttachmentHeaderBuffer);
+
+	// Re-register the persistent transform buffer with this GraphBuilder to get an RDG SRV.
+	FRDGBufferRef TransformBuffer = GraphBuilder.RegisterExternalBuffer(Buffers->TransformDataBuffer.GetPooledBuffer());
+	PassParameters->SkinningTransformBuffer = GetCompressedBoneTransformSRV(GraphBuilder, TransformBuffer);
+	PRAGMA_DISABLE_DEPRECATION_WARNINGS
+	PassParameters->GPUSceneWriterParameters = Params.GPUWriteParams;
+	PRAGMA_ENABLE_DEPRECATION_WARNINGS
+
+	TShaderMapRef<FResolveAttachmentsCS> ComputeShader(GetGlobalShaderMap(GMaxRHIFeatureLevel));
+
+	FComputeShaderUtils::AddPass(
+		GraphBuilder,
+		RDG_EVENT_NAME("ResolveAttachments(%d)", NumTotalAttachments),
+		ComputeShader,
+		PassParameters,
+		[PassParameters, &NumActualAttachments]() -> FIntVector
+		{
+			PassParameters->NumAttachments = NumActualAttachments;
+			return FIntVector(FMath::DivideAndRoundUp(NumActualAttachments, 64u), 1, 1);
+		}
+	);
+}
\ No newline at end of file
diff --git a/Engine/Source/Runtime/Renderer/Private/Skinning/SkinningSceneExtension.h b/Engine/Source/Runtime/Renderer/Private/Skinning/SkinningSceneExtension.h
--- a/Engine/Source/Runtime/Renderer/Private/Skinning/SkinningSceneExtension.h
+++ b/Engine/Source/Runtime/Renderer/Private/Skinning/SkinningSceneExtension.h
@@ -18,6 +18,7 @@
 #include "SceneRendering.h"
 
 class FSkinningSceneParameters;
+struct FGPUSceneWriteDelegateParams;
 
 class FSkinnedSceneProxy;
 class FSkinningSceneExtension : public ISceneExtension
@@ -39,6 +40,8 @@ public:
 
 		void PreGPUSceneUpdate(FRDGBuilder& GraphBuilder, const FGameTime& CurrentTime);
 
+		void ResolveAttachments(FRDGBuilder& GraphBuilder, const FGPUSceneWriteDelegateParams& Params);
+
 		void PostMeshUpdate(FRDGBuilder& GraphBuilder, const TConstArrayView<FPrimitiveSceneInfo*>& SceneInfosWithStaticDrawListUpdate, const FSkeletalMeshUpdater::FSubmitData* SkeletalMeshSubmitData);
 
 	private:
@@ -205,6 +208,8 @@ private:
 		const FGameTime& CurrentTime
 	);
 
+	void ResolveAttachments(FRDGBuilder& GraphBuilder, const FGPUSceneWriteDelegateParams& Params);
+
 	bool ProcessBufferDefragmentation();
 
 	UWorld* GetWorld() const;
