Remove FSceneViewport::bCrashOnDestruction

It wasn't hitting. The scene viewport is becoming null for some other reason.

#Synced-CL 52232959

diff --git a/Engine/Source/Runtime/Engine/Private/GameEngine.cpp b/Engine/Source/Runtime/Engine/Private/GameEngine.cpp
--- a/Engine/Source/Runtime/Engine/Private/GameEngine.cpp
+++ b/Engine/Source/Runtime/Engine/Private/GameEngine.cpp
@@ -279,7 +279,6 @@ void UGameEngine::CreateGameViewport( UGameViewportClient* GameViewportClient )
 	}
 
 	SceneViewport = MakeShareable( GameViewportClient->CreateGameViewport(GameViewportWidgetRef) );
-	SceneViewport->bCrashOnDestruction = true; // TEMP: Check for premature destruction
 	GameViewportClient->Viewport = SceneViewport.Get();
 	//GameViewportClient->CreateHighresScreenshotCaptureRegionWidget(); //  Disabled until mouse based input system can be made to work correctly.
 
@@ -796,17 +795,8 @@ void UGameEngine::OnGameWindowClosed(const TSharedRef<SWindow>& WindowBeingClose
 	{
 		// This will shutdown the game
 		GameViewport->CloseRequested(SceneViewport.Get());
-		SceneViewport->bCrashOnDestruction = false; // TEMP: Check for premature destruction
 		SceneViewport.Reset();
 	}
-	// TEMP - BEGIN: Check for premature destruction
-#if !UE_BUILD_SHIPPING
-	else
-	{
-		UE_LOGF(LogViewport, Fatal, "FSceneViewport is missing somehow.");
-	}
-#endif
-	// TEMP - END: Check for premature destruction
 }
 
 void UGameEngine::OnGameWindowMoved_SavePosition(const TSharedRef<SWindow>& WindowBeingMoved, int32 SaveWinPosIndex)
diff --git a/Engine/Source/Runtime/Engine/Private/Slate/SceneViewport.cpp b/Engine/Source/Runtime/Engine/Private/Slate/SceneViewport.cpp
--- a/Engine/Source/Runtime/Engine/Private/Slate/SceneViewport.cpp
+++ b/Engine/Source/Runtime/Engine/Private/Slate/SceneViewport.cpp
@@ -79,11 +79,6 @@ FSceneViewport::FSceneViewport( FViewportClient* InViewportClient, TSharedPtr<SV
 
 FSceneViewport::~FSceneViewport()
 {
-	// TEMP - BEGIN: Check for premature destruction
-#if !UE_BUILD_SHIPPING
-	UE_CLOGF(bCrashOnDestruction, LogViewport, Fatal, "Freeing FSceneViewport too early.");
-#endif
-	// TEMP - END: Check for premature destruction
 	Destroy();
 	// Wait for resources to be deleted
 	FlushRenderingCommands();
diff --git a/Engine/Source/Runtime/Engine/Public/Slate/SceneViewport.h b/Engine/Source/Runtime/Engine/Public/Slate/SceneViewport.h
--- a/Engine/Source/Runtime/Engine/Public/Slate/SceneViewport.h
+++ b/Engine/Source/Runtime/Engine/Public/Slate/SceneViewport.h
@@ -35,9 +35,6 @@ public:
 	ENGINE_API FSceneViewport( FViewportClient* InViewportClient, TSharedPtr<SViewport> InViewportWidget );
 	ENGINE_API ~FSceneViewport();
 
-	/** TEMP: Check for premature destruction */
-	UE_INTERNAL bool bCrashOnDestruction = false;
-
 	virtual void* GetWindow() override { return NULL; }
 
 	/** FViewport interface */
