[NetPhysics] Make Physics Replication LOD act even if there is no focal particle, i.e. return Predictive Interpolation as the replication mode.

#Synced-CL 52255161

diff --git a/Engine/Source/Runtime/Engine/Private/PhysicsEngine/PhysicsReplication.cpp b/Engine/Source/Runtime/Engine/Private/PhysicsEngine/PhysicsReplication.cpp
--- a/Engine/Source/Runtime/Engine/Private/PhysicsEngine/PhysicsReplication.cpp
+++ b/Engine/Source/Runtime/Engine/Private/PhysicsEngine/PhysicsReplication.cpp
@@ -958,6 +958,8 @@ void FPhysicsReplicationAsync::OnPostInitialize_Internal()
 
 void FPhysicsReplicationAsync::AddResimulationRequest_Internal(const float DeltaSeconds)
 {
+	//NOTE: This gets called if we perform rewind/resimulation before current physics frame steps - Chaos::RewindBeforeAdvance != 0
+
 	if (FPhysicsReplication::ShouldSkipPhysicsReplication())
 	{
 		return;
@@ -972,29 +974,45 @@ void FPhysicsReplicationAsync::AddResimulationRequest_Internal(const float Delta
 		check(RigidsSolver);
 
 		Chaos::FReadPhysicsObjectInterface_Internal Interface = Chaos::FPhysicsObjectInternalInterface::GetRead();
-		for(const FPhysicsRepAsyncInputData& InputData : AsyncInput->InputData)
+		for (const FPhysicsRepAsyncInputData& InputData : AsyncInput->InputData)
 		{
-			if( InputData.TargetState.Flags != ERigidBodyFlags::None)
+			if (InputData.TargetState.Flags != ERigidBodyFlags::None)
 			{
 				Chaos::FRewindData* RewindData = RigidsSolver->GetRewindData();
-				if (!RewindData && InputData.RepMode == EPhysicsReplicationMode::Resimulation)
+
+				EPhysicsReplicationMode EffectiveRepMode = InputData.RepMode;
+
+				// Get replication mode from LOD if enabled
+				IPhysicsReplicationLODAsync* PhysRepLod = RigidsSolver->GetPhysicsReplicationLOD_Internal();
+				if (PhysRepLod && PhysRepLod->IsEnabled())
+				{
+					FPhysicsRepLodData* LodData = PhysRepLod->GetLODData_Internal(InputData.PhysicsObject, EPhysicsReplicationLODFlags::LODFlag_All);
+					if (LodData && LodData->DataAssigned)
+					{
+						EffectiveRepMode = LodData->ReplicationMode;
+					}
+				}
+
+				// Create RewindData if we use resim and not yet created
+				if (!RewindData && EffectiveRepMode == EPhysicsReplicationMode::Resimulation)
 				{
-					// We don't have rewind data but an actor is set to replicate using resimulation; we need to enable rewind capture.
 					if (ensure(Chaos::FPBDRigidsSolver::IsNetworkPhysicsPredictionEnabled() && RigidsSolver->IsUsingFixedDt()))
 					{
 						RigidsSolver->EnableRewindCapture();
 						RewindData = RigidsSolver->GetRewindData();
 					}
 				}
+
+				// Update rewind data with target state from server no matter which replication mode is used
 				UpdateRewindDataTarget(InputData);
 
-				if(InputData.RepMode == EPhysicsReplicationMode::Resimulation)
+				if (EffectiveRepMode == EPhysicsReplicationMode::Resimulation)
 				{ 
-					if(RewindData)
+					if (RewindData)
 					{ 
-						if(Chaos::FGeometryParticleHandle* GeometryHandle = Interface.GetParticle(InputData.PhysicsObject))
+						if (Chaos::FGeometryParticleHandle* GeometryHandle = Interface.GetParticle(InputData.PhysicsObject))
 						{
-							if(Chaos::FPBDRigidParticleHandle* RigidHandle = GeometryHandle->CastToRigidParticle())
+							if (Chaos::FPBDRigidParticleHandle* RigidHandle = GeometryHandle->CastToRigidParticle())
 							{
 								FReplicatedPhysicsTargetAsync Target;
 
@@ -1065,33 +1083,49 @@ void FPhysicsReplicationAsync::OnPreSimulate_Internal()
 
 	if (AsyncInput)
 	{
-		// Update async targets with target input
 		for (const FPhysicsRepAsyncInputData& Input : AsyncInput->InputData)
 		{
 			if (Input.TargetState.Flags == ERigidBodyFlags::None)
 			{
-				// Remove replication target
 				RemoveObjectFromReplication(Input.PhysicsObject);
 				continue;
 			}
 
-			if (!bRewindDataExist && Input.RepMode == EPhysicsReplicationMode::Resimulation)
+			// When Chaos::RewindBeforeAdvance != 0 this is done inside AddResimulationRequest_Internal()
+			if (Chaos::RewindBeforeAdvance == 0)
 			{
-				// We don't have rewind data but an actor is set to replicate using resimulation; we need to enable rewind capture.
-				if (ensure(Chaos::FPBDRigidsSolver::IsNetworkPhysicsPredictionEnabled() && RigidsSolver->IsUsingFixedDt()))
+				EPhysicsReplicationMode EffectiveRepMode = Input.RepMode;
+
+				// Get replication mode from LOD if enabled
+				IPhysicsReplicationLODAsync* PhysRepLod = RigidsSolver->GetPhysicsReplicationLOD_Internal();
+				if (PhysRepLod && PhysRepLod->IsEnabled())
 				{
-					RigidsSolver->EnableRewindCapture();
-					RewindData = RigidsSolver->GetRewindData();
-					bRewindDataExist = RewindData != nullptr;
+					FPhysicsRepLodData* LodData = PhysRepLod->GetLODData_Internal(Input.PhysicsObject, EPhysicsReplicationLODFlags::LODFlag_All);
+					if (LodData && LodData->DataAssigned)
+					{
+						EffectiveRepMode = LodData->ReplicationMode;
+					}
+				}
+
+				// Create RewindData if we use resim and it's not yet created
+				if (EffectiveRepMode == EPhysicsReplicationMode::Resimulation)
+				{
+					if (!bRewindDataExist)
+					{
+						if (ensure(Chaos::FPBDRigidsSolver::IsNetworkPhysicsPredictionEnabled() && RigidsSolver->IsUsingFixedDt()))
+						{
+							RigidsSolver->EnableRewindCapture();
+							RewindData = RigidsSolver->GetRewindData();
+							bRewindDataExist = RewindData != nullptr;
+						}
+					}
 				}
-			}
 
-			if (Chaos::RewindBeforeAdvance == 0)
-			{ 
 				UpdateRewindDataTarget(Input);
 			}
+
 			UpdateAsyncTarget(Input, RigidsSolver);
-			
+
 			DebugDrawReplicationMode(Input);
 
 			// Deprecated, legacy BodyInstance flow for Default Replication
@@ -2567,84 +2601,97 @@ bool FPhysicsReplicationAsync::ResimulationReplication(Chaos::FPBDRigidParticleH
 		return true;
 	}
 
+	Chaos::FPBDRigidsSolver* RigidsSolver = static_cast<Chaos::FPBDRigidsSolver*>(GetSolver());
+	if (!RigidsSolver)
+	{
+		return true;
+	}
+
+	Chaos::FRewindData* RewindData = RigidsSolver->GetRewindData();
+	if (!RewindData)
+	{
+		if (ensure(Chaos::FPBDRigidsSolver::IsNetworkPhysicsPredictionEnabled() && RigidsSolver->IsUsingFixedDt()))
+		{
+			RigidsSolver->EnableRewindCapture();
+			RewindData = RigidsSolver->GetRewindData();
+		}
+	}
+
 	bool bClearTarget = true;
-	if(Chaos::FPBDRigidsSolver* RigidsSolver = static_cast<Chaos::FPBDRigidsSolver*>(GetSolver()))
+	if (RewindData)
 	{ 
-		if(Chaos::FRewindData* RewindData = RigidsSolver->GetRewindData())
+		if(RigidsSolver->GetEvolution())
 		{ 
-			if(RigidsSolver->GetEvolution())
-			{ 
-				const int32 LocalFrame = Target.ServerFrame - NetworkPhysicsTickOffset;
-				static constexpr Chaos::FFrameAndPhase::EParticleHistoryPhase RewindPhase = Chaos::FFrameAndPhase::EParticleHistoryPhase::PrePushData;
+			const int32 LocalFrame = Target.ServerFrame - NetworkPhysicsTickOffset;
+			static constexpr Chaos::FFrameAndPhase::EParticleHistoryPhase RewindPhase = Chaos::FFrameAndPhase::EParticleHistoryPhase::PrePushData;
 
-				// Get state from locally cached history for frame corresponding to received data
-				const Chaos::FGeometryParticleState PastState = RewindData->GetPastStateAtFrame(*Handle, LocalFrame, RewindPhase);
+			// Get state from locally cached history for frame corresponding to received data
+			const Chaos::FGeometryParticleState PastState = RewindData->GetPastStateAtFrame(*Handle, LocalFrame, RewindPhase);
 
-				const bool bShouldTriggerResim = (Chaos::RewindBeforeAdvance == 0) ? ShouldResimulateParticle(Handle, Target, DeltaSeconds) : false;
+			const bool bShouldTriggerResim = (Chaos::RewindBeforeAdvance == 0) ? ShouldResimulateParticle(Handle, Target, DeltaSeconds) : false;
 
-				if (bShouldTriggerResim && Target.TickCount == 0 && LocalFrame > RewindData->GetBlockedResimFrame())
-				{
-					// Request resimulation
-					RewindData->RequestResimulation(LocalFrame, Handle);
-				}
-				else if (SettingsCurrent.ResimulationSettings.GetRuntimeCorrectionEnabled())
-				{
-					const int32 NumPredictedFrames = RigidsSolver->GetCurrentFrame() - LocalFrame - Target.TickCount;
+			if (bShouldTriggerResim && Target.TickCount == 0 && LocalFrame > RewindData->GetBlockedResimFrame())
+			{
+				// Request resimulation
+				RewindData->RequestResimulation(LocalFrame, Handle);
+			}
+			else if (SettingsCurrent.ResimulationSettings.GetRuntimeCorrectionEnabled())
+			{
+				const int32 NumPredictedFrames = RigidsSolver->GetCurrentFrame() - LocalFrame - Target.TickCount;
 
-					if (Target.TickCount <= NumPredictedFrames && NumPredictedFrames > 0)
-					{
-						const FVector ErrorOffset = (Target.TargetState.Position - PastState.GetX());
+				if (Target.TickCount <= NumPredictedFrames && NumPredictedFrames > 0)
+				{
+					const FVector ErrorOffset = (Target.TargetState.Position - PastState.GetX());
 
-						// Positional Correction
-						const float CorrectionAmountX = SettingsCurrent.ResimulationSettings.GetPosStabilityMultiplier() / NumPredictedFrames;
-						const FVector PosDiffCorrection = ErrorOffset * CorrectionAmountX; // Same result as (ErrorOffset / NumPredictedFrames) * PosStabilityMultiplier
-						const FVector CorrectedX = Handle->GetX() + PosDiffCorrection;
+					// Positional Correction
+					const float CorrectionAmountX = SettingsCurrent.ResimulationSettings.GetPosStabilityMultiplier() / NumPredictedFrames;
+					const FVector PosDiffCorrection = ErrorOffset * CorrectionAmountX; // Same result as (ErrorOffset / NumPredictedFrames) * PosStabilityMultiplier
+					const FVector CorrectedX = Handle->GetX() + PosDiffCorrection;
 
-						// Rotational Correction
-						const float CorrectionAmountR = SettingsCurrent.ResimulationSettings.GetRotStabilityMultiplier() / NumPredictedFrames;
-						const FQuat DeltaQuat = PastState.GetR().Inverse() * Target.TargetState.Quaternion;
-						const FQuat TargetCorrectionR = Handle->GetR() * DeltaQuat;
-						const FQuat CorrectedR = FQuat::Slerp(Handle->GetR(), TargetCorrectionR, CorrectionAmountR);
+					// Rotational Correction
+					const float CorrectionAmountR = SettingsCurrent.ResimulationSettings.GetRotStabilityMultiplier() / NumPredictedFrames;
+					const FQuat DeltaQuat = PastState.GetR().Inverse() * Target.TargetState.Quaternion;
+					const FQuat TargetCorrectionR = Handle->GetR() * DeltaQuat;
+					const FQuat CorrectedR = FQuat::Slerp(Handle->GetR(), TargetCorrectionR, CorrectionAmountR);
 
-						if (SettingsCurrent.ResimulationSettings.GetRuntimeVelocityCorrectionEnabled())
-						{
-							// Linear Velocity Correction
-							const FVector LinVelDiff = Target.TargetState.LinVel - PastState.GetV(); // Velocity vector that the server covers but the client doesn't
-							const float CorrectionAmountV = SettingsCurrent.ResimulationSettings.GetVelStabilityMultiplier() / NumPredictedFrames;
-							const FVector VelCorrection = LinVelDiff * CorrectionAmountV; // Same result as (LinVelDiff / NumPredictedFrames) * VelStabilityMultiplier
-							const FVector CorrectedV = Handle->GetV() + VelCorrection;
-
-							// Angular Velocity Correction
-							const FVector AngVelDiff = FMath::DegreesToRadians(Target.TargetState.AngVel) - PastState.GetW(); // Angular velocity vector that the server covers but the client doesn't
-							const float CorrectionAmountW = SettingsCurrent.ResimulationSettings.GetAngVelStabilityMultiplier() / NumPredictedFrames;
-							const FVector AngVelCorrection = AngVelDiff * CorrectionAmountW; // Same result as (AngVelDiff / NumPredictedFrames) * VelStabilityMultiplier
-							const FVector CorrectedW = Handle->GetW() + AngVelCorrection;
+					if (SettingsCurrent.ResimulationSettings.GetRuntimeVelocityCorrectionEnabled())
+					{
+						// Linear Velocity Correction
+						const FVector LinVelDiff = Target.TargetState.LinVel - PastState.GetV(); // Velocity vector that the server covers but the client doesn't
+						const float CorrectionAmountV = SettingsCurrent.ResimulationSettings.GetVelStabilityMultiplier() / NumPredictedFrames;
+						const FVector VelCorrection = LinVelDiff * CorrectionAmountV; // Same result as (LinVelDiff / NumPredictedFrames) * VelStabilityMultiplier
+						const FVector CorrectedV = Handle->GetV() + VelCorrection;
+
+						// Angular Velocity Correction
+						const FVector AngVelDiff = FMath::DegreesToRadians(Target.TargetState.AngVel) - PastState.GetW(); // Angular velocity vector that the server covers but the client doesn't
+						const float CorrectionAmountW = SettingsCurrent.ResimulationSettings.GetAngVelStabilityMultiplier() / NumPredictedFrames;
+						const FVector AngVelCorrection = AngVelDiff * CorrectionAmountW; // Same result as (AngVelDiff / NumPredictedFrames) * VelStabilityMultiplier
+						const FVector CorrectedW = Handle->GetW() + AngVelCorrection;
 					
-							// Apply correction to velocities
-							Handle->SetV(CorrectedV);
-							Handle->SetW(CorrectedW);
-						}
-
-			#if !(UE_BUILD_SHIPPING || UE_BUILD_TEST)
-						if (PhysicsReplicationCVars::ResimulationCVars::bDrawDebug)
-						{
-							Chaos::FDebugDrawQueue::GetInstance().DrawDebugDirectionalArrow(Handle->GetX(), CorrectedX, 5.0f, FColor::MakeRandomSeededColor(LocalFrame), false, PhysicsReplicationCVars::DebugDrawLifeTime, 0, 0.5f);
-						}
-			#endif
-						// Apply correction to position and rotation
-						RigidsSolver->GetEvolution()->ApplyParticleTransformCorrection(Handle, CorrectedX, CorrectedR, SettingsCurrent.ResimulationSettings.GetRuntimeCorrectConnectedBodies(), /*bInRecalculateFrictionOnConnectedBodies*/true, ReplicatedParticleIDs);
+						// Apply correction to velocities
+						Handle->SetV(CorrectedV);
+						Handle->SetW(CorrectedW);
 					}
 
-					// Keep target for NumPredictedFrames time to perform runtime corrections with until a new target is received
-					bClearTarget = Target.TickCount >= NumPredictedFrames;
+		#if !(UE_BUILD_SHIPPING || UE_BUILD_TEST)
+					if (PhysicsReplicationCVars::ResimulationCVars::bDrawDebug)
+					{
+						Chaos::FDebugDrawQueue::GetInstance().DrawDebugDirectionalArrow(Handle->GetX(), CorrectedX, 5.0f, FColor::MakeRandomSeededColor(LocalFrame), false, PhysicsReplicationCVars::DebugDrawLifeTime, 0, 0.5f);
+					}
+		#endif
+					// Apply correction to position and rotation
+					RigidsSolver->GetEvolution()->ApplyParticleTransformCorrection(Handle, CorrectedX, CorrectedR, SettingsCurrent.ResimulationSettings.GetRuntimeCorrectConnectedBodies(), /*bInRecalculateFrictionOnConnectedBodies*/true, ReplicatedParticleIDs);
 				}
 
-				const bool bShouldSleep = (Target.TargetState.Flags & ERigidBodyFlags::Sleeping) != 0;
-				if (Target.IsWaiting() && !bShouldSleep)
-				{
-					// Don't clear the target if we are waiting for a specific target frame and not sleeping
-					bClearTarget = false;
-				}
+				// Keep target for NumPredictedFrames time to perform runtime corrections with until a new target is received
+				bClearTarget = Target.TickCount >= NumPredictedFrames;
+			}
+
+			const bool bShouldSleep = (Target.TargetState.Flags & ERigidBodyFlags::Sleeping) != 0;
+			if (Target.IsWaiting() && !bShouldSleep)
+			{
+				// Don't clear the target if we are waiting for a specific target frame and not sleeping
+				bClearTarget = false;
 			}
 		}
 	}
diff --git a/Engine/Source/Runtime/Engine/Private/PhysicsEngine/PhysicsReplicationLOD.cpp b/Engine/Source/Runtime/Engine/Private/PhysicsEngine/PhysicsReplicationLOD.cpp
--- a/Engine/Source/Runtime/Engine/Private/PhysicsEngine/PhysicsReplicationLOD.cpp
+++ b/Engine/Source/Runtime/Engine/Private/PhysicsEngine/PhysicsReplicationLOD.cpp
@@ -146,7 +146,7 @@ bool FPhysicsReplicationLODAsync::IsEnabled() const
 	{
 		return false;
 	}
-	return (bEnabled || (PhysicsReplicationLODCVars::OverrideEnabled == 1)) && FocalParticles.Num() > 0;
+	return (bEnabled || (PhysicsReplicationLODCVars::OverrideEnabled == 1));
 }
 
 void FPhysicsReplicationLODAsync::OnPostInitialize_Internal()
@@ -297,11 +297,19 @@ void FPhysicsReplicationLODAsync::UnregisterFocalPoint_Internal(Chaos::FConstPhy
 
 FPhysicsRepLodData* FPhysicsReplicationLODAsync::GetLODData_Internal(Chaos::FConstPhysicsObjectHandle PhysicsObject, const uint32 LODFlags)
 {
-	if (FocalParticles.Num() == 0)
+	if (!IsEnabled())
 	{
 		return nullptr;
 	}
 
+	if (FocalParticles.Num() == 0)
+	{
+		// No focal particles, all objects use PredictiveInterpolation
+		LodData.Reset(); // Reset sets ReplicationMode = PredictiveInterpolation by default and AlignedFrame gets initialized to 0 which stops any forward prediction to be performed
+		LodData.DataAssigned = true;
+		return &LodData;
+	}
+
 	Chaos::FReadPhysicsObjectInterface_Internal Interface = Chaos::FPhysicsObjectInternalInterface::GetRead();
 	const Chaos::FGeometryParticleHandle* ParticleHandle = Interface.GetParticle(PhysicsObject);
 	if (!ParticleHandle)
