Prevent auto-centering for window moves on the same monitor

Also remove FSystemResolution::LockUnknownMonitor() since it shouldn't be necessary anymore.

#Synced-CL 52232959

diff --git a/Engine/Source/Runtime/Engine/Private/GameEngine.cpp b/Engine/Source/Runtime/Engine/Private/GameEngine.cpp
--- a/Engine/Source/Runtime/Engine/Private/GameEngine.cpp
+++ b/Engine/Source/Runtime/Engine/Private/GameEngine.cpp
@@ -486,7 +486,7 @@ TSharedRef<SWindow> UGameEngine::CreateGameWindow()
 	int32 ResX = GSystemResolution.ResX;
 	int32 ResY = GSystemResolution.ResY;
 	EWindowMode::Type WindowMode = GSystemResolution.WindowMode;
-	int32 MonitorIndex = GSystemResolution.MonitorIndex;
+	const int32 MonitorIndex = GSystemResolution.MonitorIndex;
 	{
 		ConditionallyOverrideSettings(ResX, ResY, WindowMode, MonitorIndex);
 
@@ -631,13 +631,6 @@ TSharedRef<SWindow> UGameEngine::CreateGameWindow()
 		}
 	}
 
-	if (AutoCenterType == EAutoCenter::None)
-	{
-		// Position has been set explicitly, so disable multiple monitor behavior.
-		MonitorIndex = -1;
-		GSystemResolution.LockUnknownMonitor();
-	}
-
 	// Give the window the max width/height of either the requested resolution, or your available desktop resolution
 	// We need to do this as we request some 4K windows when rendering sequences, and the OS may try and clamp that
 	// window to your available desktop resolution
@@ -660,7 +653,7 @@ TSharedRef<SWindow> UGameEngine::CreateGameWindow()
 		MaxWindowHeight = FMath::Max(DisplayMetrics.VirtualDisplayRect.Bottom - DisplayMetrics.VirtualDisplayRect.Top, ResY);
 	}
 
-	if (MonitorIndex != -1)
+	if ((MonitorIndex != -1) && (AutoCenterType == EAutoCenter::PrimaryWorkArea))
 	{
 		FDisplayMetrics DisplayMetrics;
 		if (FSlateApplication::IsInitialized())
diff --git a/Engine/Source/Runtime/Engine/Private/Slate/SceneViewport.cpp b/Engine/Source/Runtime/Engine/Private/Slate/SceneViewport.cpp
--- a/Engine/Source/Runtime/Engine/Private/Slate/SceneViewport.cpp
+++ b/Engine/Source/Runtime/Engine/Private/Slate/SceneViewport.cpp
@@ -1572,9 +1572,16 @@ void FSceneViewport::ResizeFrame(uint32 NewWindowSizeX, uint32 NewWindowSizeY, E
 				// scaling so we actual render to the resolution we've been asked for.
 				if (NewWindowMode == EWindowMode::Windowed)
 				{
-					if (!bValidMonitorIndex && (OldWindowMode == EWindowMode::Windowed) && (NewWindowSize == OldWindowSize))
+					bool bKeepOldWindowPos = (OldWindowMode == EWindowMode::Windowed) && (NewWindowSize == OldWindowSize);
+					if (bKeepOldWindowPos && bValidMonitorIndex)
+					{
+						const FPlatformRect& DisplayRect = DisplayMetrics.MonitorInfo[MonitorIndexForResize].DisplayRect;
+						const FSlateRect FullScreenInfo = WindowToResize->GetFullScreenInfo();
+						bKeepOldWindowPos = (DisplayRect.Left == FullScreenInfo.Left) && (DisplayRect.Top == FullScreenInfo.Top);
+					}
+
+					if (bKeepOldWindowPos)
 					{
-						// Leave the window position alone!
 						NewWindowPos.Reset();
 					}
 					else
diff --git a/Engine/Source/Runtime/Engine/Private/UnrealEngine.cpp b/Engine/Source/Runtime/Engine/Private/UnrealEngine.cpp
--- a/Engine/Source/Runtime/Engine/Private/UnrealEngine.cpp
+++ b/Engine/Source/Runtime/Engine/Private/UnrealEngine.cpp
@@ -1151,7 +1151,6 @@ void SystemResolutionSinkCallback()
 	}
 
 	const FString ResString = CVarSystemResolution->GetString();
-	const int32 MonitorIndex = CVarMonitorIndex->GetInt();
 
 	uint32 ResX, ResY;
 	int32 WindowModeInt = GSystemResolution.WindowMode;
@@ -1159,6 +1158,7 @@ void SystemResolutionSinkCallback()
 	if (ParseResolution(*ResString, ResX, ResY, WindowModeInt))
 	{
 		EWindowMode::Type WindowMode = EWindowMode::ConvertIntToWindowMode(WindowModeInt);
+		const int32 MonitorIndex = CVarMonitorIndex->GetInt();
 
 		if( GSystemResolution.ResX != ResX ||
 			GSystemResolution.ResY != ResY ||
@@ -19041,19 +19041,9 @@ bool AllowHighQualityLightmaps(const FStaticFeatureLevel FeatureLevel)
 		&& !IsMobilePlatform(GShaderPlatformForFeatureLevel[FeatureLevel]);
 }
 
-static bool bLockUnknownMonitor = false;
-void FSystemResolution::LockUnknownMonitor()
-{
-	SetMonitorIndex(-1);
-	bLockUnknownMonitor = true;
-}
-
 void FSystemResolution::SetMonitorIndex(int32 MonitorIndex)
 {
-	if (!bLockUnknownMonitor)
-	{
-		CVarMonitorIndex->Set(MonitorIndex);
-	}
+	CVarMonitorIndex->Set(MonitorIndex);
 }
 
 // Helper function for changing system resolution via the r.setres console command
diff --git a/Engine/Source/Runtime/Engine/Public/UnrealEngine.h b/Engine/Source/Runtime/Engine/Public/UnrealEngine.h
--- a/Engine/Source/Runtime/Engine/Public/UnrealEngine.h
+++ b/Engine/Source/Runtime/Engine/Public/UnrealEngine.h
@@ -452,7 +452,6 @@ struct FSystemResolution
 	int32 MonitorIndex = -1;
 	bool bForceRefresh = false;
 
-	static void LockUnknownMonitor();
 	static ENGINE_API void SetMonitorIndex(int32 MonitorIndex);
 
 	// Helper function for changing system resolution via the r.setres console command
