Enable physics field on mobile only for CPU

#Synced-CL 52257317

diff --git a/Engine/Plugins/Experimental/ChaosNiagara/Source/ChaosNiagara/Classes/NiagaraDataInterfacePhysicsField.h b/Engine/Plugins/Experimental/ChaosNiagara/Source/ChaosNiagara/Classes/NiagaraDataInterfacePhysicsField.h
--- a/Engine/Plugins/Experimental/ChaosNiagara/Source/ChaosNiagara/Classes/NiagaraDataInterfacePhysicsField.h
+++ b/Engine/Plugins/Experimental/ChaosNiagara/Source/ChaosNiagara/Classes/NiagaraDataInterfacePhysicsField.h
@@ -59,7 +59,11 @@ public:
 
 	/** UNiagaraDataInterface Interface */
 	UE_API virtual void GetVMExternalFunction(const FVMExternalFunctionBindingInfo& BindingInfo, void* InstanceData, FVMExternalFunction& OutFunc) override;
-	virtual bool CanExecuteOnTarget(ENiagaraSimTarget Target) const override { return Target == ENiagaraSimTarget::GPUComputeSim; }
+	virtual bool CanExecuteOnTarget(ENiagaraSimTarget Target) const override { return true; }
+#if WITH_EDITOR
+	virtual bool ShouldCompile(EShaderPlatform ShaderPlatform) const override { return !IsMobilePlatform(ShaderPlatform) && 
+		IsFeatureLevelSupported(ShaderPlatform, ERHIFeatureLevel::SM5);}
+#endif
 	UE_API virtual bool InitPerInstanceData(void* PerInstanceData, FNiagaraSystemInstance* SystemInstance) override;
 	UE_API virtual void DestroyPerInstanceData(void* PerInstanceData, FNiagaraSystemInstance* SystemInstance)override;
 	UE_API virtual bool PerInstanceTickPostSimulate(void* PerInstanceData, FNiagaraSystemInstance* SystemInstance, float DeltaSeconds) override;
diff --git a/Engine/Source/Runtime/Engine/Classes/PhysicsField/PhysicsFieldComponent.h b/Engine/Source/Runtime/Engine/Classes/PhysicsField/PhysicsFieldComponent.h
--- a/Engine/Source/Runtime/Engine/Classes/PhysicsField/PhysicsFieldComponent.h
+++ b/Engine/Source/Runtime/Engine/Classes/PhysicsField/PhysicsFieldComponent.h
@@ -174,7 +174,7 @@ public:
 	 * Initializes the instance for the given resource.
 	 * @param TextureSize - The resource texture size to be used.
 	 */
-	void InitInstance(const TArray<EFieldPhysicsType>& TargetTypes, const bool bBuildClipmap);
+	void InitInstance(const TArray<EFieldPhysicsType>& TargetTypes, const bool bBuildClipmap, const bool bNeedsResource);
 
 	/**
 	 * Release the resource of the instance.
diff --git a/Engine/Source/Runtime/Engine/Private/PhysicsField/PhysicsFieldComponent.cpp b/Engine/Source/Runtime/Engine/Private/PhysicsField/PhysicsFieldComponent.cpp
--- a/Engine/Source/Runtime/Engine/Private/PhysicsField/PhysicsFieldComponent.cpp
+++ b/Engine/Source/Runtime/Engine/Private/PhysicsField/PhysicsFieldComponent.cpp
@@ -188,7 +188,7 @@ public:
 
 	static bool ShouldCompilePermutation(const FGlobalShaderPermutationParameters& Parameters)
 	{
-		return IsFeatureLevelSupported(Parameters.Platform, ERHIFeatureLevel::SM5);
+		return !IsMobilePlatform(Parameters.Platform) && IsFeatureLevelSupported(Parameters.Platform, ERHIFeatureLevel::SM5);
 	}
 
 	static void ModifyCompilationEnvironment(const FGlobalShaderPermutationParameters& Parameters, FShaderCompilerEnvironment& OutEnvironment)
@@ -302,7 +302,7 @@ public:
 
 	static bool ShouldCompilePermutation(const FGlobalShaderPermutationParameters& Parameters)
 	{
-		return IsFeatureLevelSupported(Parameters.Platform, ERHIFeatureLevel::SM5);
+		return !IsMobilePlatform(Parameters.Platform) && IsFeatureLevelSupported(Parameters.Platform, ERHIFeatureLevel::SM5);
 	}
 
 	static void ModifyCompilationEnvironment(const FGlobalShaderPermutationParameters& Parameters, FShaderCompilerEnvironment& OutEnvironment)
@@ -609,7 +609,7 @@ void FPhysicsFieldResource::UpdateResource(FRHICommandList& RHICmdList,
 *	FPhysicsFieldInstance
 */
 
-void FPhysicsFieldInstance::InitInstance( const TArray<EFieldPhysicsType>& TargetTypes, const bool bBuildClipmap)
+void FPhysicsFieldInstance::InitInstance( const TArray<EFieldPhysicsType>& TargetTypes, const bool bBuildClipmap, const bool bNeedsResource)
 {
 	FPhysicsFieldInfos::BufferOffsets VectorTargets(InPlace, -1), ScalarTargets(InPlace, -1), IntegerTargets(InPlace, -1), PhysicsTargets(InPlace, -1);
 	TStaticArray<UE::Core::TAlignedElement<FIntVector4, 16>, MAX_PHYSICS_FIELD_TARGETS> PhysicsBounds(InPlace, FIntVector4(-1));
@@ -644,7 +644,7 @@ void FPhysicsFieldInstance::InitInstance( const TArray<EFieldPhysicsType>& Targe
 		}
 	}
 	
-	if (!FieldResource)
+	if (!FieldResource && bNeedsResource)
 	{
 		FieldResource = new FPhysicsFieldResource(TargetCount, TargetTypes, VectorTargets, ScalarTargets, IntegerTargets, PhysicsTargets, PhysicsBounds, bBuildClipmap);
 
@@ -902,6 +902,13 @@ void UPhysicsFieldComponent::SendRenderDynamicData_Concurrent()
 void UPhysicsFieldComponent::OnRegister()
 {
 	Super::OnRegister();
+	
+	bool bNeedsResource = false;
+	if (GetWorld() && GetWorld()->Scene )
+	{
+		const EShaderPlatform ShaderPlatfom = GetWorld()->Scene->GetShaderPlatform();
+		bNeedsResource = !IsMobilePlatform(ShaderPlatfom) && IsFeatureLevelSupported(ShaderPlatfom, ERHIFeatureLevel::SM5);
+	}
 
 	TArray<bool> bBuildClipmaps = { GPhysicsFieldBuildClipmap == 1, false };
 	for (uint32 FieldIndex = 0; FieldIndex < 2; ++FieldIndex)
@@ -914,7 +921,7 @@ void UPhysicsFieldComponent::OnRegister()
 			if (GPhysicsFieldSingleTarget == 1)
 			{
 				TArray<EFieldPhysicsType> TargetTypes = { EFieldPhysicsType::Field_LinearForce };
-				LocalInstance->InitInstance(TargetTypes, bBuildClipmaps[FieldIndex]);
+				LocalInstance->InitInstance(TargetTypes, bBuildClipmaps[FieldIndex], bNeedsResource);
 			}
 			else
 			{
@@ -932,7 +939,7 @@ void UPhysicsFieldComponent::OnRegister()
 															EFieldPhysicsType::Field_ActivateDisabled };
 				TargetTypes.Sort();
 
-				LocalInstance->InitInstance(TargetTypes, bBuildClipmaps[FieldIndex]);
+				LocalInstance->InitInstance(TargetTypes, bBuildClipmaps[FieldIndex], bNeedsResource);
 			}
 		}
 	}
diff --git a/Engine/Source/Runtime/Engine/Private/World.cpp b/Engine/Source/Runtime/Engine/Private/World.cpp
--- a/Engine/Source/Runtime/Engine/Private/World.cpp
+++ b/Engine/Source/Runtime/Engine/Private/World.cpp
@@ -2960,7 +2960,7 @@ void UWorld::UpdateWorldComponents(bool bRerunConstructionScripts, bool bCurrent
 		}
 
 		static IConsoleVariable* PhysicsFieldEnableClipmapCVar = IConsoleManager::Get().FindConsoleVariable(TEXT("r.PhysicsField.EnableField"));
-		if (PhysicsFieldEnableClipmapCVar && PhysicsFieldEnableClipmapCVar->GetInt() == 1 && Scene && !IsMobilePlatform(Scene->GetShaderPlatform()))
+		if (PhysicsFieldEnableClipmapCVar && PhysicsFieldEnableClipmapCVar->GetInt() == 1)
 		{
 			if (!PhysicsField)
 			{
