Substrate - Enable eye shading model to go through the Select node.
Tested and works with new material translator.

#tests PC editor dx12
#rnx
#Synced-CL 52290745

diff --git a/Engine/Source/Runtime/Engine/Private/Materials/HLSLMaterialTranslator.cpp b/Engine/Source/Runtime/Engine/Private/Materials/HLSLMaterialTranslator.cpp
--- a/Engine/Source/Runtime/Engine/Private/Materials/HLSLMaterialTranslator.cpp
+++ b/Engine/Source/Runtime/Engine/Private/Materials/HLSLMaterialTranslator.cpp
@@ -13439,9 +13439,31 @@ int32 FHLSLMaterialTranslator::SubstrateHairBSDF(int32 BaseColor, int32 Scatter,
 
 int32 FHLSLMaterialTranslator::SubstrateEyeBSDF(int32 DiffuseAlbedo, int32 Roughness, int32 IrisMask, int32 IrisDistance, int32 IrisNormal, int32 IrisPlaneNormal, int32 SSSProfileId, int32 EmissiveColor, int32 CorneaNormal, const FString& SharedLocalBasisIndexMacro, FSubstrateOperator* PromoteToOperator)
 {
+	if (PromoteToOperator)
+	{
+		return AddCodeChunk(
+			MCT_Substrate, TEXT("Parameters.%s.PromoteParameterBlendedBSDFToOperator(GetSubstrateEyeBSDF(%s, %s, %s, %s, %s, %s, ExtractSubsurfaceProfileInt(%s), %s, %s), %u, %u, %u, %u) /* Cornea:%s Iris:%s */"),
+			*Substrate::GetParametersSubstrateTreeName(SubstrateTranslatorData.CurrentSubstrateCompilationContext),
+			*SubstrateGetCastParameterCode(DiffuseAlbedo, MCT_Float3),
+			*SubstrateGetCastParameterCode(Roughness, MCT_Float),
+			*SubstrateGetCastParameterCode(IrisMask, MCT_Float),
+			*SubstrateGetCastParameterCode(IrisDistance, MCT_Float),
+			*SubstrateGetCastParameterCode(IrisNormal, MCT_Float3),
+			*SubstrateGetCastParameterCode(IrisPlaneNormal, MCT_Float3),
+			*SubstrateGetCastParameterCode(SSSProfileId, MCT_Float),
+			*SubstrateGetCastParameterCode(EmissiveColor, MCT_Float3),
+			*SharedLocalBasisIndexMacro,
+			PromoteToOperator->Index,
+			PromoteToOperator->BSDFIndex,
+			PromoteToOperator->LayerDepth,
+			PromoteToOperator->bIsBottom ? 1 : 0,
+			*GetParameterCode(CorneaNormal),
+			*GetParameterCode(IrisNormal)
+		);
+	}
+
 	return AddCodeChunk(
-		MCT_Substrate, TEXT("Parameters.%s.PromoteParameterBlendedBSDFToOperator(GetSubstrateEyeBSDF(%s, %s, %s, %s, %s, %s, ExtractSubsurfaceProfileInt(%s), %s, %s), %u, %u, %u, %u) /* Cornea:%s Iris:%s */"),
-		*Substrate::GetParametersSubstrateTreeName(SubstrateTranslatorData.CurrentSubstrateCompilationContext),
+		MCT_Substrate, TEXT("GetSubstrateEyeBSDF(%s, %s, %s, %s, %s, %s, ExtractSubsurfaceProfileInt(%s), %s, %s)"),
 		*SubstrateGetCastParameterCode(DiffuseAlbedo, MCT_Float3),
 		*SubstrateGetCastParameterCode(Roughness, MCT_Float),
 		*SubstrateGetCastParameterCode(IrisMask, MCT_Float),
@@ -13450,13 +13472,7 @@ int32 FHLSLMaterialTranslator::SubstrateEyeBSDF(int32 DiffuseAlbedo, int32 Rough
 		*SubstrateGetCastParameterCode(IrisPlaneNormal, MCT_Float3),
 		*SubstrateGetCastParameterCode(SSSProfileId, MCT_Float),
 		*SubstrateGetCastParameterCode(EmissiveColor, MCT_Float3),
-		*SharedLocalBasisIndexMacro,
-		PromoteToOperator->Index,
-		PromoteToOperator->BSDFIndex,
-		PromoteToOperator->LayerDepth,
-		PromoteToOperator->bIsBottom ? 1 : 0,
-		*GetParameterCode(CorneaNormal),
-		*GetParameterCode(IrisNormal)
+		*SharedLocalBasisIndexMacro
 	);
 }
 
diff --git a/Engine/Source/Runtime/Engine/Private/Materials/MaterialExpressionSubstrate.cpp b/Engine/Source/Runtime/Engine/Private/Materials/MaterialExpressionSubstrate.cpp
--- a/Engine/Source/Runtime/Engine/Private/Materials/MaterialExpressionSubstrate.cpp
+++ b/Engine/Source/Runtime/Engine/Private/Materials/MaterialExpressionSubstrate.cpp
@@ -2251,16 +2251,7 @@ int32 UMaterialExpressionSubstrateEyeBSDF::Compile(class FMaterialCompiler* Comp
 
 	FSubstrateOperator& SubstrateOperator = Compiler->SubstrateCompilationGetOperator(Compiler->SubstrateTreeStackGetPathUniqueId());
 	SubstrateOperator.BSDFRegisteredSharedLocalBasis = NewRegisteredSharedLocalBasis;
-
-	if (SubstrateOperator.bUseParameterBlending)
-	{
-		return Compiler->Errorf(TEXT("Substrate Eye BSDF node cannot be used with parameter blending."));
-	}
-	else if (SubstrateOperator.bRootOfParameterBlendingSubTree)
-	{
-		return Compiler->Errorf(TEXT("Substrate Eye BSDF node cannot be the root of a parameter blending sub tree."));
-	}
-
+	 
 	int32 OutputCodeChunk = Compiler->SubstrateEyeBSDF(
 		CompileWithDefaultFloat3(Compiler, DiffuseColor, 0.0f, 0.0f, 0.0f),
 		CompileWithDefaultFloat1(Compiler, Roughness, 0.5f),
@@ -2272,7 +2263,7 @@ int32 UMaterialExpressionSubstrateEyeBSDF::Compile(class FMaterialCompiler* Comp
 		CompileWithDefaultFloat3(Compiler, EmissiveColor, 0.0f, 0.0f, 0.0f),
 		CorneaNormalCodeChunk,
 		Compiler->GetSubstrateSharedLocalBasisIndexMacro(NewRegisteredSharedLocalBasis),
-		&SubstrateOperator);
+		!SubstrateOperator.bUseParameterBlending || (SubstrateOperator.bUseParameterBlending && SubstrateOperator.bRootOfParameterBlendingSubTree) ? &SubstrateOperator : nullptr);
 
 	return OutputCodeChunk;
 }
diff --git a/Engine/Source/Runtime/Engine/Private/Materials/SubstrateTranslatorCommon.cpp b/Engine/Source/Runtime/Engine/Private/Materials/SubstrateTranslatorCommon.cpp
--- a/Engine/Source/Runtime/Engine/Private/Materials/SubstrateTranslatorCommon.cpp
+++ b/Engine/Source/Runtime/Engine/Private/Materials/SubstrateTranslatorCommon.cpp
@@ -385,9 +385,9 @@ bool FSubstrateCompilationContext::SubstrateGenerateDerivedMaterialOperatorData(
 		MTV_PostProcess = 0x200,
 		MTV_LightFunction = 0x400,
 
-		MTV_ValidMultipleBSDFs = (MTV_BSDF_Slab | MTV_BSDF_Hair | MTV_BSDF_Toon),	// SUBSTRATE_TODO Unlit cannot be in there yet due to the specific define for it (SUBSTRATE_OPTIMIZED_UNLIT)
+		MTV_SelectOp_ValidMultipleBSDFs = (MTV_BSDF_Slab | MTV_BSDF_Hair | MTV_BSDF_Eye | MTV_BSDF_Toon),	// Unlit cannot be in there yet due to the specific define for it (SUBSTRATE_OPTIMIZED_UNLIT).
 		MTV_WeightOp_AllowedList = (MTV_BSDF_Unlit | MTV_BSDF_Slab | MTV_BSDF_Toon),
-		MTV_SelectOp_AllowedList = (MTV_ValidMultipleBSDFs),							// This only works because parameter blending is enforced so there will always be a single BSDF of a single type only. And hair or toon BSDFs will never pass through other operator allowed list.
+		MTV_SelectOp_AllowedList = (MTV_SelectOp_ValidMultipleBSDFs),										// This only works because parameter blending is enforced so there will always be a single BSDF of a single type only. And hair or toon BSDFs will never pass through other operator allowed list.
 		MTV_HorizoOp_AllowedList = MTV_BSDF_Slab,
 		MTV_VerticOp_AllowedList = MTV_BSDF_Slab,
 		MTV_DecalOp_AllowedList = MTV_BSDF_Slab,
@@ -753,9 +753,9 @@ bool FSubstrateCompilationContext::SubstrateGenerateDerivedMaterialOperatorData(
 				const bool bRootOfParameterBlendingSubTree = bCurrentOpRequestParameterBlending && !bInsideParameterBlendingSubTree;
 				const bool bUseParameterBlending = bCurrentOpRequestParameterBlending || bInsideParameterBlendingSubTree;
 
-				if (CurrentOperator.BSDFType == SUBSTRATE_BSDF_TYPE_SLAB																// Any operators
-					|| CurrentOperator.BSDFType == SUBSTRATE_BSDF_TYPE_UNLIT															// Weight operator
-					|| CurrentOperator.BSDFType == SUBSTRATE_BSDF_TYPE_HAIR || CurrentOperator.BSDFType == SUBSTRATE_BSDF_TYPE_TOON)	// Select operator
+				if (CurrentOperator.BSDFType == SUBSTRATE_BSDF_TYPE_SLAB																													// Any operators
+					|| CurrentOperator.BSDFType == SUBSTRATE_BSDF_TYPE_UNLIT																												// Weight operator
+					|| CurrentOperator.BSDFType == SUBSTRATE_BSDF_TYPE_HAIR || CurrentOperator.BSDFType == SUBSTRATE_BSDF_TYPE_EYE || CurrentOperator.BSDFType == SUBSTRATE_BSDF_TYPE_TOON)	// Select operator
 				{
 					// Update the parameter blending data for BSDFs supporting operators.
 					// We also need to do this for UNLIT since it supports Coverage operator.
@@ -1354,7 +1354,7 @@ bool FSubstrateCompilationContext::SubstrateGenerateDerivedMaterialOperatorData(
 				UMaterialInterface* MaterialInterface = SubstrateTranslatorData.GetMaterialInterface();
 				SubstrateTranslatorData.GetMaterialCompilationOutput().SubstrateMaterialCompilationOutput.bIsThin = MaterialInterface && MaterialInterface->IsThinSurface() ? 1 : 0;
 
-				uint32 BSDFTypeCount = FMath::CountBits(SubstrateTreeValidation.MaterialType & MTV_ValidMultipleBSDFs);
+				uint32 BSDFTypeCount = FMath::CountBits(SubstrateTreeValidation.MaterialType & MTV_SelectOp_ValidMultipleBSDFs);
 
 				// The order of ifs here is important.
 				if (MaterialInterface && MaterialInterface->IsLightFunctionMaterial())
