[MIR] Expand texture type kinds and fix Custom HLSL function parameters

- Replace ETypeKind::Texture with distinct per-type kinds (Texture2D,                                                                                                                                                                                                                                                                               TextureCube, Texture2DArray, TextureCubeArray, Texture3D)
- Auto-generate SamplerState arguments for texture parameters in custom
HLSL function declarations and call sites
- Split LWC inputs into float + FWSType parameter pairs in custom HLSL
function declarations and call sites

#rnx
#Synced-CL 52251728

diff --git a/Engine/Source/Runtime/Engine/Private/Materials/MaterialIREmitter.cpp b/Engine/Source/Runtime/Engine/Private/Materials/MaterialIREmitter.cpp
--- a/Engine/Source/Runtime/Engine/Private/Materials/MaterialIREmitter.cpp
+++ b/Engine/Source/Runtime/Engine/Private/Materials/MaterialIREmitter.cpp
@@ -1168,7 +1168,7 @@ FValueRef FEmitter::TextureUniform(FName Name, UTexture* Texture, EMaterialSampl
 		return Poison();
 	}
 
-	FTextureUniform Proto = MakePrototype<FTextureUniform>(FType::MakeTexture());
+	FTextureUniform Proto = MakePrototype<FTextureUniform>(FType::FromMaterialValueType(Texture->GetMaterialType()));
 	Proto.Name = FScriptName{ Name };
 	Proto.Texture = Texture;
 	Proto.SamplerType = SamplerType;
@@ -3655,20 +3655,15 @@ static void CastTexCoord(FEmitter& Em, FValueRef Texture, FValueRef& TexCoord)
 		return;
 	}
 
-	// @massimo.tristano Instead of looking inside Texture's object, add various FTypes to match all texture types (right now we have a single "Texture" to capture them all).
-	if (const UTexture* TextureObject = Cast<UTexture>(Texture->GetTextureObject()))
+	MIR::ETypeKind Kind = Texture->Type.GetKind();
+	if (Kind == MIR::ETypeKind::TextureCube || Kind == MIR::ETypeKind::Texture2DArray || Kind == MIR::ETypeKind::TextureCubeArray || Kind == MIR::ETypeKind::Texture3D)
 	{
-		EMaterialValueType MaterialValueType = TextureObject->GetMaterialType();
-		if (MaterialValueType == MCT_TextureCube || MaterialValueType == MCT_Texture2DArray || MaterialValueType == MCT_VolumeTexture)
-		{
-			// These texture types require a float3 texture coordinate
-			TexCoord = Em.CastToFloat(TexCoord, 3);
-			return;
-		}
+		TexCoord = Em.CastToFloat(TexCoord, 3);
+	}
+	else
+	{
+		TexCoord = Em.CastToFloat(TexCoord, 2);
 	}
-
-	// By default, cast TexCoord to a float2
-	TexCoord = Em.CastToFloat(TexCoord, 2);
 }
 
 FValueRef FEmitter::TextureGather(FValueRef Texture, FValueRef TexCoord, ETextureReadMode GatherMode, const FTextureSampleAttributes& BaseAttributes)
diff --git a/Engine/Source/Runtime/Engine/Private/Materials/MaterialIRToHLSLTranslator.cpp b/Engine/Source/Runtime/Engine/Private/Materials/MaterialIRToHLSLTranslator.cpp
--- a/Engine/Source/Runtime/Engine/Private/Materials/MaterialIRToHLSLTranslator.cpp
+++ b/Engine/Source/Runtime/Engine/Private/Materials/MaterialIRToHLSLTranslator.cpp
@@ -257,6 +257,12 @@ static FPrinter& operator<<(FPrinter& Printer, const MIR::FType& Type)
 		switch (Type.GetKind())
 		{
 			case MIR::ETypeKind::Void: Printer << TEXTVIEW("void"); break;
+			case MIR::ETypeKind::Texture2D: Printer << TEXTVIEW("Texture2D"); break;
+			case MIR::ETypeKind::TextureCube: Printer << TEXTVIEW("TextureCube"); break;
+			case MIR::ETypeKind::Texture2DArray: Printer << TEXTVIEW("Texture2DArray"); break;
+			case MIR::ETypeKind::TextureCubeArray: Printer << TEXTVIEW("TextureCubeArray"); break;
+			case MIR::ETypeKind::Texture3D: Printer << TEXTVIEW("Texture3D"); break;
+			case MIR::ETypeKind::RuntimeVirtualTexture: Printer << TEXTVIEW("VirtualTexturePhysical"); break;
 			case MIR::ETypeKind::SubstrateData: Printer << TEXTVIEW("FSubstrateData"); break;
 			case MIR::ETypeKind::VTPageTableResult: Printer << TEXTVIEW("VTPageTableResult"); break;
 			default: UE_MIR_UNREACHABLE();
@@ -576,12 +582,12 @@ struct FEntryPointCodeGen
 
 	void LowerTextureUniform(const MIR::FTextureUniform* TextureUniform)
 	{
-		Printer << GetMaterialUBName() << TEXTVIEW(".") << LowerTextureMaterialType(TextureUniform->Texture->GetMaterialType(), true) << TEXTVIEW("_") << TextureUniform->Analysis_UniformIndex;
+		Printer << GetMaterialUBName() << TEXTVIEW(".") << LowerTextureTypeKind(TextureUniform->Type, true) << TEXTVIEW("_") << TextureUniform->Analysis_UniformIndex;
 	}
 
 	void LowerVirtualTextureUniform(const MIR::FVirtualTextureUniform* VirtualTextureUniform)
 	{
-		Printer << GetMaterialUBName() << TEXTVIEW(".") << LowerTextureMaterialType(MCT_TextureVirtual, true) << TEXTVIEW("_") << VirtualTextureUniform->Analysis_UniformIndex;
+		Printer << GetMaterialUBName() << TEXTVIEW(".") << LowerTextureTypeKind(MIR::FType::MakeRuntimeVirtualTexture(), true) << TEXTVIEW("_") << VirtualTextureUniform->Analysis_UniformIndex;
 	}
 
 	void EmitPreshaderBufferReadDoubleVector(uint32 GlobalComponentOffset, int32 SourceComponents, int32 UsedComponents)
@@ -1405,30 +1411,23 @@ struct FEntryPointCodeGen
 		}
 	}
 
-	ENoOp LowerTextureMaterialType(EMaterialValueType TextureType, bool bForResourceDeclarations)
+	ENoOp LowerTextureTypeKind(const MIR::FType& Type, bool bForResourceDeclarations)
 	{
-		switch (TextureType)
+	    // @massimo.tristano Double check we need this.
+		if (bForResourceDeclarations && Type.Is(MIR::ETypeKind::Texture3D))
 		{
-			case MCT_Texture2D: Printer << TEXTVIEW("Texture2D"); break;
-			case MCT_TextureCube: Printer << TEXTVIEW("TextureCube"); break;
-			case MCT_Texture2DArray: Printer << TEXTVIEW("Texture2DArray"); break;
-			case MCT_TextureCubeArray: Printer << TEXTVIEW("TextureCubeArray"); break;
-			case MCT_VolumeTexture: Printer << (bForResourceDeclarations ? TEXTVIEW("VolumeTexture") : TEXTVIEW("Texture3D")); break;
-			case MCT_TextureExternal: Printer << TEXTVIEW("Texture2D"); break;
-			case MCT_TextureVirtual: Printer << TEXTVIEW("VirtualTexturePhysical"); break;
-			default: UE_MIR_UNREACHABLE();
+			Printer << TEXTVIEW("VolumeTexture");
+		}
+		else
+		{
+			Printer << Type;
 		}
 		return NoOp;
 	}
 
 	void LowerStandardTextureRead(const MIR::FTextureRead* TextureRead)
 	{
-		switch (TextureRead->TextureUniform->Kind)
-		{
-			case MIR::EValueKind::VK_TextureUniform: 		LowerTextureMaterialType(TextureRead->TextureUniform->As<MIR::FTextureUniform>()->Texture->GetMaterialType(), false); break;
-			case MIR::EValueKind::VK_VirtualTextureUniform: LowerTextureMaterialType(MCT_TextureVirtual, false); break;
-			default: UE_MIR_UNREACHABLE();
-		}
+		LowerTextureTypeKind(TextureRead->TextureUniform->Type, false);
 
 		switch (TextureRead->Mode)
 		{
@@ -1795,7 +1794,21 @@ struct FEntryPointCodeGen
 			else
 			{
 				// Input-only parameters instead can inline their value instead
+				const MIR::FType& ParamType = Call->Function->Parameters[i].Type;
+
+				if (ParamType.IsDouble())
+				{
+					// LWC types get two values: first the value demoted to float, then the raw LWC value
+					Printer << TEXTVIEW("WSDemote(") << LowerValue(Call->Arguments[i]) << TEXTVIEW("), ");
+				}
+
 				Printer << LowerValue(Call->Arguments[i]);
+
+				// Texture parameters require an accompanying sampler argument
+				if (ParamType.IsTexture())
+				{
+					Printer << TEXTVIEW(", ") << LowerValue(Call->Arguments[i]) << TEXTVIEW("Sampler");
+				}
 			}
 		}
 		Printer << EndArgs;
@@ -2304,7 +2317,32 @@ static void GenerateCustomFunctionsHLSL(const FMaterialIRModule* Module, TMap<FS
 				: TEXTVIEW("out "));
 
 				// Type and name
-				Printer << Function->Parameters[i].Type << TEXT(' ') << *Function->Parameters[i].Name.ToString();
+				const MIR::FType& ParamType = Function->Parameters[i].Type;
+				const FString& ParamName = Function->Parameters[i].Name.ToString();
+
+				if (ParamType.IsDouble())
+				{
+					// LWC types get two parameters: first a demoted float version, then the raw LWC value.
+					// This way legacy custom expressions can continue to operate on regular float values.
+					int32 NumComp = ParamType.GetPrimitive().NumComponents();
+					Printer << TEXTVIEW("float");
+					if (NumComp > 1)
+					{
+						Printer << NumComp;
+					}
+					Printer << TEXT(' ') << *ParamName;
+					Printer << TEXTVIEW(", ") << ParamType << TEXTVIEW(" LWC") << *ParamName;
+				}
+				else
+				{
+					Printer << ParamType << TEXT(' ') << *ParamName;
+
+					// Texture parameters require an accompanying sampler argument
+					if (ParamType.IsTexture())
+					{
+						Printer << TEXTVIEW(", SamplerState ") << *ParamName << TEXTVIEW("Sampler");
+					}
+				}
 			}
 
 			Printer << EndArgs << OpenBlock;
diff --git a/Engine/Source/Runtime/Engine/Private/Materials/MaterialIRTypes.cpp b/Engine/Source/Runtime/Engine/Private/Materials/MaterialIRTypes.cpp
--- a/Engine/Source/Runtime/Engine/Private/Materials/MaterialIRTypes.cpp
+++ b/Engine/Source/Runtime/Engine/Private/Materials/MaterialIRTypes.cpp
@@ -10,19 +10,27 @@ namespace MIR
 
 const TCHAR* TypeKindToString(ETypeKind Kind)
 {
+    #define CASE(Name) case ETypeKind::Name: return TEXT(#Name); 
+    
 	switch (Kind)
 	{
-		case ETypeKind::Poison: return TEXT("Poison");
-		case ETypeKind::Void: return TEXT("Void");
-		case ETypeKind::Primitive: return TEXT("Primitive");
-		case ETypeKind::Aggregate: return TEXT("Aggregate");
-		case ETypeKind::Texture: return TEXT("Texture");
-		case ETypeKind::RuntimeVirtualTexture: return TEXT("RuntimeVirtualTexture");
-		case ETypeKind::ParameterCollection: return TEXT("ParameterCollection");
-		case ETypeKind::SubstrateData: return TEXT("SubstrateData");
-		case ETypeKind::VTPageTableResult: return TEXT("VTPageTableResult");
+		CASE(Poison)
+		CASE(Void)
+		CASE(Primitive)
+		CASE(Aggregate)
+		CASE(Texture2D)
+		CASE(TextureCube)
+		CASE(Texture2DArray)
+		CASE(TextureCubeArray)
+		CASE(Texture3D)
+		CASE(RuntimeVirtualTexture)
+		CASE(ParameterCollection)
+		CASE(SubstrateData)
+		CASE(VTPageTableResult)
 		default: UE_MIR_UNREACHABLE();
 	}
+	
+	#undef CASE
 }
 
 bool ScalarKindIsArithmetic(EScalarKind Kind)
@@ -96,11 +104,11 @@ FType FType::FromMaterialValueType(EMaterialValueType Type)
 		case MCT_Float2: return FType::MakeFloatVector(2);
 		case MCT_Float3: return FType::MakeFloatVector(3);
 		case MCT_Float4: return FType::MakeFloatVector(4);
-		case MCT_Texture2D: return FType::MakeTexture();
-		case MCT_TextureCube: return FType::MakeTexture();
-		case MCT_Texture2DArray: return FType::MakeTexture();
-		case MCT_TextureCubeArray: return FType::MakeTexture();
-		case MCT_VolumeTexture: return FType::MakeTexture();
+		case MCT_Texture2D: return FType::MakeTexture2D();
+		case MCT_TextureCube: return FType::MakeTextureCube();
+		case MCT_Texture2DArray: return FType::MakeTexture2DArray();
+		case MCT_TextureCubeArray: return FType::MakeTextureCubeArray();
+		case MCT_VolumeTexture: return FType::MakeTexture3D();
 		case MCT_ShadingModel: return FType::MakeIntScalar();
 		case MCT_StaticBool: return FType::MakeBoolScalar();
 		case MCT_UInt1: return FType::MakeIntScalar();
@@ -128,7 +136,7 @@ FType FType::FromMaterialParameterType(EMaterialParameterType Type)
 		case EMaterialParameterType::Scalar: return FType::MakeFloatScalar();
 		case EMaterialParameterType::Vector: return FType::MakeFloatVector(4);
 		case EMaterialParameterType::DoubleVector: return FType::MakeDoubleVector(4);
-		case EMaterialParameterType::Texture: return FType::MakeTexture();
+		case EMaterialParameterType::Texture: return FType::MakeTexture2D();
 		case EMaterialParameterType::TextureCollection: UE_MIR_TODO();
 		case EMaterialParameterType::ParameterCollection: UE_MIR_TODO();
 		case EMaterialParameterType::Font: UE_MIR_TODO();
@@ -182,10 +190,38 @@ FType FType::MakeParameterCollection()
 	return Type;
 }
 
-FType FType::MakeTexture()
+FType FType::MakeTexture2D()
+{
+	FType Type{};
+	Type.Kind = ETypeKind::Texture2D;
+	return Type;
+}
+
+FType FType::MakeTextureCube()
+{
+	FType Type{};
+	Type.Kind = ETypeKind::TextureCube;
+	return Type;
+}
+
+FType FType::MakeTexture2DArray()
+{
+	FType Type{};
+	Type.Kind = ETypeKind::Texture2DArray;
+	return Type;
+}
+
+FType FType::MakeTextureCubeArray()
+{
+	FType Type{};
+	Type.Kind = ETypeKind::TextureCubeArray;
+	return Type;
+}
+
+FType FType::MakeTexture3D()
 {
 	FType Type{};
-	Type.Kind = ETypeKind::Texture;
+	Type.Kind = ETypeKind::Texture3D;
 	return Type;
 }
 
diff --git a/Engine/Source/Runtime/Engine/Public/Materials/MaterialIRTypes.h b/Engine/Source/Runtime/Engine/Public/Materials/MaterialIRTypes.h
--- a/Engine/Source/Runtime/Engine/Public/Materials/MaterialIRTypes.h
+++ b/Engine/Source/Runtime/Engine/Public/Materials/MaterialIRTypes.h
@@ -20,8 +20,12 @@ enum class ETypeKind : uint8
 	Aggregate,	// A material aggregate.
 
 	// Opaque Objects
-	
-	Texture,				// A texture object.
+
+	Texture2D,				// A 2D texture object.
+	TextureCube,			// A cube texture object.
+	Texture2DArray,			// A 2D texture array object.
+	TextureCubeArray,		// A cube texture array object.
+	Texture3D,				// A 3D (volume) texture object.
 	RuntimeVirtualTexture,	// A runtime virtual texture object.
 	ParameterCollection,	// Material parameter collection object.
 	SubstrateData,			// Substrate-specific data structure.
@@ -189,8 +193,20 @@ struct FType
 	// Returns the Material Parameter Collection object type.
 	static FType MakeParameterCollection();
 
-	// Returns the Texture object type.
-	static FType MakeTexture();
+	// Returns the Texture2D object type.
+	static FType MakeTexture2D();
+
+	// Returns the TextureCube object type.
+	static FType MakeTextureCube();
+
+	// Returns the Texture2DArray object type.
+	static FType MakeTexture2DArray();
+
+	// Returns the TextureCubeArray object type.
+	static FType MakeTextureCubeArray();
+
+	// Returns the Texture3D object type.
+	static FType MakeTexture3D();
 
 	// Returns the RuntimeVirtualTexture object type.
 	static FType MakeRuntimeVirtualTexture();
@@ -255,8 +271,8 @@ struct FType
 	// Returns true if the given type is a matrix primitive.
 	bool IsMatrix() const { return Is(ETypeKind::Primitive) && Primitive.IsMatrix(); }
 
-	// Returns whether this type is a Texture object.
-	bool IsTexture() const { return Kind == ETypeKind::Texture; }
+	// Returns whether this type is any non-virtual texture object.
+	bool IsTexture() const { return Kind == ETypeKind::Texture2D || Kind == ETypeKind::TextureCube || Kind == ETypeKind::Texture2DArray || Kind == ETypeKind::TextureCubeArray || Kind == ETypeKind::Texture3D; }
 
 	// Returns whether this type is a RuntimeVirtualTexture object.
 	bool IsRuntimeVirtualTexture() const { return Kind == ETypeKind::RuntimeVirtualTexture; }
