Data Table - Get contents of given Row/Column
Provide a lambda for how to get the value from the property, e.g. Property->GetPropertyValue(CellData);

#Synced-CL 52105157

diff --git a/Engine/Source/Runtime/Engine/Classes/Engine/DataTable.h b/Engine/Source/Runtime/Engine/Classes/Engine/DataTable.h
--- a/Engine/Source/Runtime/Engine/Classes/Engine/DataTable.h
+++ b/Engine/Source/Runtime/Engine/Classes/Engine/DataTable.h
@@ -269,6 +269,40 @@ public:
 		return *RowDataPtr;
 	}
 
+	/** 
+	 * Get the value in the datatable at the given row and column. 
+	 * The provided function should get the value from the property type and the cell data
+	 */
+	template<typename TPropertyType, typename TReturnType>
+	const TReturnType* GetCellValue(const FName& RowName, const FName& ColumnName, TFunction<const TReturnType*(const TPropertyType*, const uint8*)> GetPropertyValueFromCell)
+	{
+		const FProperty* Property = FindTableProperty(ColumnName);
+		if (!Property)
+		{
+			return nullptr;
+		}
+		
+		const TPropertyType* PropertyType = CastField<TPropertyType>(Property);
+		if (!PropertyType)
+		{
+			return nullptr;
+		}
+		
+		const uint8* RowData = FindRowUnchecked(RowName);
+		if (!RowData)
+		{
+			return nullptr;
+		}
+
+		const uint8* CellData = Property->ContainerPtrToValuePtr<uint8>(RowData, 0);
+		if (!CellData)
+		{
+			return nullptr;
+		}
+
+		return GetPropertyValueFromCell(PropertyType, CellData);
+	}
+
 	/** Empty the table info (will not clear RowStruct) */
 	ENGINE_API virtual void EmptyTable();
 
