Fix race condition between GetCodec() and PostLoadSubobjects on UAnimBoneCompressionSettings

#Synced-CL 52292702

diff --git a/Engine/Source/Runtime/Engine/Classes/Animation/AnimBoneCompressionSettings.h b/Engine/Source/Runtime/Engine/Classes/Animation/AnimBoneCompressionSettings.h
--- a/Engine/Source/Runtime/Engine/Classes/Animation/AnimBoneCompressionSettings.h
+++ b/Engine/Source/Runtime/Engine/Classes/Animation/AnimBoneCompressionSettings.h
@@ -3,6 +3,7 @@
 #pragma once
 
 #include "CoreMinimal.h"
+#include "Async/SharedMutex.h"
 #include "AnimBoneCompressionSettings.generated.h"
 
 class UAnimBoneCompressionCodec;
@@ -40,6 +41,9 @@ class UAnimBoneCompressionSettings : public UObject
 	/** Allow us to convert DDC serialized path back into codec object */
 	ENGINE_API UAnimBoneCompressionCodec* GetCodec(const FString& DDCHandle);
 
+	// UObject overrides
+	ENGINE_API virtual void PostLoadSubobjects(FObjectInstancingGraph* OuterInstanceGraph) override;
+
 	//////////////////////////////////////////////////////////////////////////
 
 #if WITH_EDITORONLY_DATA
@@ -69,4 +73,8 @@ class UAnimBoneCompressionSettings : public UObject
 	UE_DEPRECATED(5.1, "This function has been deprecated. Override the one above instead.")
 	ENGINE_API void PopulateDDCKey(FArchive& Ar);
 #endif
+
+private:
+	/** Guards the Codecs array against concurrent read (GetCodec) and write (PostLoadSubobjects/InstanceSubobjects) access during async loading. */
+	mutable UE::FSharedMutex CodecsLock;
 };
diff --git a/Engine/Source/Runtime/Engine/Private/Animation/AnimBoneCompressionSettings.cpp b/Engine/Source/Runtime/Engine/Private/Animation/AnimBoneCompressionSettings.cpp
--- a/Engine/Source/Runtime/Engine/Private/Animation/AnimBoneCompressionSettings.cpp
+++ b/Engine/Source/Runtime/Engine/Private/Animation/AnimBoneCompressionSettings.cpp
@@ -6,6 +6,8 @@
 #include "AnimationUtils.h"
 #include "AnimationCompression.h"
 #include "Async/ParallelFor.h"
+#include "Async/UniqueLock.h"
+#include "Async/SharedLock.h"
 
 #include UE_INLINE_GENERATED_CPP_BY_NAME(AnimBoneCompressionSettings)
 
@@ -25,6 +27,8 @@ UAnimBoneCompressionSettings::UAnimBoneCompressionSettings(const FObjectInitiali
 
 UAnimBoneCompressionCodec* UAnimBoneCompressionSettings::GetCodec(const FString& DDCHandle)
 {
+	UE::TSharedLock<UE::FSharedMutex> ReadLock(CodecsLock);
+
 	UAnimBoneCompressionCodec* CodecMatch = nullptr;
 	for (UAnimBoneCompressionCodec* Codec : Codecs)
 	{
@@ -42,6 +46,14 @@ UAnimBoneCompressionCodec* UAnimBoneCompressionSettings::GetCodec(const FString&
 	return CodecMatch;
 }
 
+void UAnimBoneCompressionSettings::PostLoadSubobjects(FObjectInstancingGraph* OuterInstanceGraph)
+{
+	// The Codecs array is read by GetCodec() during UAnimSequence::Serialize on async loading threads.
+	// FArrayProperty::InstanceSubobjects memmoves array elements causing a race
+	UE::TUniqueLock<UE::FSharedMutex> WriteLock(CodecsLock);
+	Super::PostLoadSubobjects(OuterInstanceGraph);
+}
+
 #if WITH_EDITORONLY_DATA
 void UAnimBoneCompressionSettings::GetPreloadDependencies(TArray<UObject*>& OutDeps)
 {
