Fix UE::Anim::FNamedValueArrayUtils::Union(InOutValueArray0, InValueArray1, InPredicate) not processing elements that are only in InOutValueArray0 and whose sort value is greater than the last element of InValueArray1

#Synced-CL 52242993

diff --git a/Engine/Source/Runtime/Engine/Public/Animation/NamedValueArray.h b/Engine/Source/Runtime/Engine/Public/Animation/NamedValueArray.h
--- a/Engine/Source/Runtime/Engine/Public/Animation/NamedValueArray.h
+++ b/Engine/Source/Runtime/Engine/Public/Animation/NamedValueArray.h
@@ -346,8 +346,18 @@ struct FNamedValueArrayUtils
 			}
 		}
 
-		// Tape merge is done, copy anything that might be remaining
-		if (ElementPtr1 < ElementEndPtr1)
+		// Tape merge is done, process anything that might be remaining
+		if (ElementPtr0 < ElementEndPtr0)
+		{
+			// Reached end of ValueArray1 with remaining elements in ValueArray0, call predicate on each
+			for (; ElementPtr0 < ElementEndPtr0; ++ElementPtr0)
+			{
+				DefaultElement.Name = ElementPtr0->Name;
+
+				InPredicate(*ElementPtr0, DefaultElement, UE::Anim::ENamedValueUnionFlags::ValidArg0);
+			}
+		}
+		else if (ElementPtr1 < ElementEndPtr1)
 		{
 			// Reached end of ValueArray0 with remaining in ValueArray1, we can just copy the remainder of ValueArray1
 			const int32 NumResults = InOutValueArray0.Elements.Num();
