Fixed skeletal mesh build not reporting required memory due to integer overflow and added debug strings

#Synced-CL 52292524

diff --git a/Engine/Source/Runtime/Engine/Private/SkeletalMesh.cpp b/Engine/Source/Runtime/Engine/Private/SkeletalMesh.cpp
--- a/Engine/Source/Runtime/Engine/Private/SkeletalMesh.cpp
+++ b/Engine/Source/Runtime/Engine/Private/SkeletalMesh.cpp
@@ -2356,10 +2356,11 @@ void USkeletalMesh::Build()
 		FQueuedThreadPool* SkeletalMeshThreadPool = FSkinnedAssetCompilingManager::Get().GetThreadPool();
 		EQueuedWorkPriority BasePriority = FSkinnedAssetCompilingManager::Get().GetBasePriority(this);
 		check(AsyncTask == nullptr);
+		int64 RequiredMemory = 1024 * 1024 * static_cast<int64>(CVarMemoryForSkeletalMeshAssetCompile.GetValueOnAnyThread());
 		AsyncTask = MakeUnique<FSkinnedAssetAsyncBuildTask>(this, MoveTemp(Context));
 		if (!HasAnyDependenciesCompiling())
 		{
-			AsyncTask->StartBackgroundTask(SkeletalMeshThreadPool, BasePriority, EQueuedWorkFlags::DoNotRunInsideBusyWait);
+			AsyncTask->StartBackgroundTask(SkeletalMeshThreadPool, BasePriority, EQueuedWorkFlags::DoNotRunInsideBusyWait, RequiredMemory, TEXT("SkeletalMesh"));
 		}
 
 		if (HasCompileDependencies())
@@ -2480,7 +2481,7 @@ FEvent* USkeletalMesh::LockPropertiesUntil()
 	EQueuedWorkPriority BasePriority = FSkinnedAssetCompilingManager::Get().GetBasePriority(this);
 	check(AsyncTask == nullptr);
 
-	int64 RequiredMemory = 1024 * 1024 * CVarMemoryForSkeletalMeshAssetCompile.GetValueOnAnyThread();
+	int64 RequiredMemory = 1024 * 1024 * static_cast<int64>(CVarMemoryForSkeletalMeshAssetCompile.GetValueOnAnyThread());
 
 	AsyncTask = MakeUnique<FSkinnedAssetAsyncBuildTask>(this, MoveTemp(Context));
 	if (!HasAnyDependenciesCompiling())
diff --git a/Engine/Source/Runtime/Engine/Private/SkinnedAssetCompiler.cpp b/Engine/Source/Runtime/Engine/Private/SkinnedAssetCompiler.cpp
--- a/Engine/Source/Runtime/Engine/Private/SkinnedAssetCompiler.cpp
+++ b/Engine/Source/Runtime/Engine/Private/SkinnedAssetCompiler.cpp
@@ -315,7 +315,8 @@ void FSkinnedAssetCompilingManager::SchedulePendingCompilations()
 			GetThreadPool(),
 			GetBasePriority(SkinnedAsset),
 			EQueuedWorkFlags::DoNotRunInsideBusyWait,
-			SkinnedAsset->GetRequiredMemoryForCompile()
+			SkinnedAsset->GetRequiredMemoryForCompile(),
+			TEXT("SkinnedAsset")
 		);
 	}
 
