PR #14478
NotifyStates: Add the option to not merge notify states on concurrent plays of the same animation asset

#Synced-CL 52338451

diff --git a/Engine/Source/Runtime/Engine/Classes/Animation/AnimInstance.h b/Engine/Source/Runtime/Engine/Classes/Animation/AnimInstance.h
--- a/Engine/Source/Runtime/Engine/Classes/Animation/AnimInstance.h
+++ b/Engine/Source/Runtime/Engine/Classes/Animation/AnimInstance.h
@@ -1513,6 +1513,9 @@ public:
 
 	/** Pending teleport type, set in ResetDynamics and cleared in UpdateAnimation */
 	ETeleportType PendingDynamicResetTeleportType;
+	
+	/** ID for this anim instance */
+	uint32 AnimInstanceID = 0;
 
 	/** Animation Notifies that has been triggered in the latest tick **/
 	UPROPERTY(transient)
diff --git a/Engine/Source/Runtime/Engine/Classes/Animation/AnimNode_AssetPlayerBase.h b/Engine/Source/Runtime/Engine/Classes/Animation/AnimNode_AssetPlayerBase.h
--- a/Engine/Source/Runtime/Engine/Classes/Animation/AnimNode_AssetPlayerBase.h
+++ b/Engine/Source/Runtime/Engine/Classes/Animation/AnimNode_AssetPlayerBase.h
@@ -76,6 +76,15 @@ private:
 #endif
 	
 protected:
+	/** Track whether we have been full weight previously. Reset when we reach 0 weight*/
+	bool bHasBeenFullWeight = false;
+	
+	/** If true, an asset player instance ID has been successfully set */
+	bool bHasValidAssetPlayerInstanceID = false;
+	
+	/** Instance ID for this asset player node. Used to identify concurrently playing instances of the same animation during notify triggering. */
+	uint32 AssetPlayerInstanceID = 0;
+	
 	/** Store data about current marker position when using marker based syncing*/
 	FMarkerTickRecord MarkerTickRecord;
 
@@ -89,10 +98,7 @@ protected:
 	
 	/** Previous frame InternalTimeAccumulator value and effective delta time leading into the current frame */
 	FDeltaTimeRecord DeltaTimeRecord;
-
-	/** Track whether we have been full weight previously. Reset when we reach 0 weight*/
-	bool bHasBeenFullWeight = false;
-
+	
 	/** Get sync parameters for the tick record produced by this asset player */
 	ENGINE_API UE::Anim::FAnimSyncParams GetSyncParams(bool bRequestedInertialization) const;
 };
diff --git a/Engine/Source/Runtime/Engine/Classes/Animation/AnimNotifies/AnimNotifyState.h b/Engine/Source/Runtime/Engine/Classes/Animation/AnimNotifies/AnimNotifyState.h
--- a/Engine/Source/Runtime/Engine/Classes/Animation/AnimNotifies/AnimNotifyState.h
+++ b/Engine/Source/Runtime/Engine/Classes/Animation/AnimNotifies/AnimNotifyState.h
@@ -20,6 +20,16 @@ class FCanvas;
 class FSceneView;
 #endif
 
+/** Flags that control how notify states behave */
+UENUM(meta = (Bitflags, UseEnumValuesAsMaskValuesInEditor = "true"))
+enum class EAnimNotifyStateBehaviorFlags : uint8
+{
+	None					= 0 UMETA(Hidden),
+	// If set, each concurrently playing instance of the same animation maintains an independent notify state lifecycle rather than merging them.
+	NoMergeOnConcurrentPlay	= 1 << 0,
+};
+ENUM_CLASS_FLAGS(EAnimNotifyStateBehaviorFlags);
+
 UCLASS(abstract, editinlinenew, Blueprintable, const, hidecategories=Object, collapsecategories, meta=(ShowWorldContextPin), MinimalAPI)
 class UAnimNotifyState : public UObject
 {
@@ -97,6 +107,9 @@ class UAnimNotifyState : public UObject
 	/** This notify is always a branching point when used on Montages. */
 	bool bIsNativeBranchingPoint;
 
+	UPROPERTY(EditAnywhere, BlueprintReadWrite, AdvancedDisplay, Category = AnimNotify, meta=(Bitmask, BitmaskEnum="/Script/Engine.EAnimNotifyStateBehaviorFlags"))
+	uint8 NotifyStateBehaviorFlags;
+	
 protected:
 	ENGINE_API UObject* GetContainingAsset() const;
 };
diff --git a/Engine/Source/Runtime/Engine/Private/Animation/AnimInstance.cpp b/Engine/Source/Runtime/Engine/Private/Animation/AnimInstance.cpp
--- a/Engine/Source/Runtime/Engine/Private/Animation/AnimInstance.cpp
+++ b/Engine/Source/Runtime/Engine/Private/Animation/AnimInstance.cpp
@@ -39,6 +39,8 @@
 #include "ClothingSimulationInteractor.h"
 #endif
 
+#include "Animation/AnimationInstanceScope.h"
+
 #include UE_INLINE_GENERATED_CPP_BY_NAME(AnimInstance)
 
 #if WITH_EDITOR
@@ -141,6 +143,12 @@ namespace UE::Anim::Private
 			return NextNotifyInstanceID++;
 		}
 	}
+	
+	static uint32 NextAnimInstanceID = 0;
+	static uint32 GetNextAnimInstanceID()
+	{
+		return NextAnimInstanceID++;
+	}
 }
 
 /** Used to determine notify dispatch behaviour when sources (montage, anim graph) may tick at different times within a frame. */
@@ -338,6 +346,7 @@ UWorld* UAnimInstance::GetWorld() const
 void UAnimInstance::InitializeMontageOnly()
 {
 	bUninitialized = false;
+	AnimInstanceID = UE::Anim::Private::GetNextAnimInstanceID();
 	USkeletalMeshComponent* OwnerComponent = GetSkelMeshComponent();
 	if (OwnerComponent->GetSkeletalMeshAsset() != NULL)
 	{
@@ -362,6 +371,8 @@ void UAnimInstance::InitializeAnimation(bool bInDeferRootNodeInitialization)
 	
 	TRACE_OBJECT_LIFETIME_BEGIN(this);
 
+	AnimInstanceID = UE::Anim::Private::GetNextAnimInstanceID();
+	
 	// make sure your skeleton is initialized
 	// you can overwrite different skeleton
 	USkeletalMeshComponent* OwnerComponent = GetSkelMeshComponent();
@@ -1878,18 +1889,53 @@ void UAnimInstance::TriggerAnimNotifies(float DeltaSeconds, ENotifyTriggerMode M
 			// Filter any notify states that are still active from the old tracking arrays to avoid triggering their respective being/end events.
 			if (AnimNotifyEvent->NotifyStateClass)
 			{
-				int32 ExistingItemIndex = INDEX_NONE;
+				const UE::Anim::FAnimNotifyMontageInstanceContext* QueuedMontageNotifyContext = NotifyQueue.AnimNotifies[Index].GetContextData<UE::Anim::FAnimNotifyMontageInstanceContext>();
+				const UE::Anim::FAnimNotifyAssetPlayerInstanceContext* QueuedAssetPlayerNotifyContext = NotifyQueue.AnimNotifies[Index].GetContextData<UE::Anim::FAnimNotifyAssetPlayerInstanceContext>();
 				
-				if (const bool bIsNotifyStateAlreadyActive = ActiveAnimNotifyState.Find(*AnimNotifyEvent, ExistingItemIndex))
+				bool bFoundAlreadyActiveNotify = false;
+				for (int32 ActiveNotifyStateIndex = 0; ActiveNotifyStateIndex < ActiveAnimNotifyState.Num(); ++ActiveNotifyStateIndex)
 				{
-					check(ActiveAnimNotifyState.Num() == ActiveAnimNotifyEventReference.Num());
+					// Check if the queued notify matches an already-active notify state
+					const bool bIsMatchingNotifyEvent = ActiveAnimNotifyState[ActiveNotifyStateIndex] == *AnimNotifyEvent;
+					if (bIsMatchingNotifyEvent)
+					{
+						check(ActiveAnimNotifyState.Num() == ActiveAnimNotifyEventReference.Num());
+						const FAnimNotifyEventReference& ActiveEventReference = ActiveAnimNotifyEventReference[ActiveNotifyStateIndex];
+						
+						// When NoMergeOnConcurrentPlay is set, require a matching source instance ID to prevent concurrent animation plays from sharing a single notify state.
+						const bool bRequiresSourceInstanceIDMatching = (AnimNotifyEvent->NotifyStateClass->NotifyStateBehaviorFlags & static_cast<uint8>(EAnimNotifyStateBehaviorFlags::NoMergeOnConcurrentPlay)) != 0;
+						bool bHasMatchingSourceInstanceID = false;
+						if (bRequiresSourceInstanceIDMatching)
+						{
+							if (QueuedMontageNotifyContext)
+							{
+								if (const UE::Anim::FAnimNotifyMontageInstanceContext* ActiveMontageNotifyContext = ActiveEventReference.GetContextData<UE::Anim::FAnimNotifyMontageInstanceContext>())
+								{
+									bHasMatchingSourceInstanceID = QueuedMontageNotifyContext->MontageInstanceID == ActiveMontageNotifyContext->MontageInstanceID;
+								}
+							}
+							else if (QueuedAssetPlayerNotifyContext)
+							{
+								if (const UE::Anim::FAnimNotifyAssetPlayerInstanceContext* ActiveAssetInstanceNotifyContext = ActiveEventReference.GetContextData<UE::Anim::FAnimNotifyAssetPlayerInstanceContext>())
+								{
+									bHasMatchingSourceInstanceID = QueuedAssetPlayerNotifyContext->AssetPlayerInstanceID == ActiveAssetInstanceNotifyContext->AssetPlayerInstanceID;
+								}
+							}
+						}
 
-					// Notify already exists, so we should re-use the same Instance ID.
-					NotifyQueue.AnimNotifies[Index].SetNotifyInstanceID(ActiveAnimNotifyEventReference[ExistingItemIndex].GetNotifyInstanceID());
-					ActiveAnimNotifyState.RemoveAtSwap(ExistingItemIndex, EAllowShrinking::No);
-					ActiveAnimNotifyEventReference.RemoveAtSwap(ExistingItemIndex, EAllowShrinking::No);
+						if (!bRequiresSourceInstanceIDMatching || (bRequiresSourceInstanceIDMatching && bHasMatchingSourceInstanceID))
+						{
+							// Notify already exists, so we should re-use the same Instance ID.
+							NotifyQueue.AnimNotifies[Index].SetNotifyInstanceID(ActiveAnimNotifyEventReference[ActiveNotifyStateIndex].GetNotifyInstanceID());
+							ActiveAnimNotifyState.RemoveAtSwap(ActiveNotifyStateIndex, EAllowShrinking::No);
+							ActiveAnimNotifyEventReference.RemoveAtSwap(ActiveNotifyStateIndex, EAllowShrinking::No);
+							bFoundAlreadyActiveNotify = true;
+							break;
+						}
+					}
 				}
-				else
+				
+				if (!bFoundAlreadyActiveNotify)
 				{
 					// New notify. Generate unique ID for this instance.
 					NotifyQueue.AnimNotifies[Index].SetNotifyInstanceID(UE::Anim::Private::GetNextNotifyInstanceID());
diff --git a/Engine/Source/Runtime/Engine/Private/Animation/AnimNode_AssetPlayerBase.cpp b/Engine/Source/Runtime/Engine/Private/Animation/AnimNode_AssetPlayerBase.cpp
--- a/Engine/Source/Runtime/Engine/Private/Animation/AnimNode_AssetPlayerBase.cpp
+++ b/Engine/Source/Runtime/Engine/Private/Animation/AnimNode_AssetPlayerBase.cpp
@@ -1,9 +1,12 @@
 // Copyright Epic Games, Inc. All Rights Reserved.
 
 #include "Animation/AnimNode_AssetPlayerBase.h"
+
+#include "Animation/AnimationInstanceScope.h"
 #include "Animation/AnimSyncScope.h"
 #include "Animation/AnimTrace.h"
 #include "Animation/AnimInertializationSyncScope.h"
+#include "Animation/AnimInstance.h"
 
 #include UE_INLINE_GENERATED_CPP_BY_NAME(AnimNode_AssetPlayerBase)
 
@@ -13,6 +16,12 @@ void FAnimNode_AssetPlayerBase::Initialize_AnyThread(const FAnimationInitializeC
 
 	MarkerTickRecord.Reset();
 	bHasBeenFullWeight = false;
+	
+	if (const UAnimInstance* AnimInstance = Cast<UAnimInstance>(Context.GetAnimInstanceObject()))
+	{
+		AssetPlayerInstanceID = HashCombine(GetTypeHash(AnimInstance->AnimInstanceID), GetTypeHash(Context.GetCurrentNodeId()));
+		bHasValidAssetPlayerInstanceID = true;
+	}
 }
 
 void FAnimNode_AssetPlayerBase::Update_AnyThread(const FAnimationUpdateContext& Context)
@@ -36,8 +45,13 @@ void FAnimNode_AssetPlayerBase::CreateTickRecordForNode(const FAnimationUpdateCo
 	TickRecord.GatherContextData(Context);
 	TickRecord.RootMotionWeightModifier = Context.GetRootMotionWeightModifier();
 	TickRecord.DeltaTimeRecord = &DeltaTimeRecord;
-	TickRecord.bRequestedInertialization = Context.GetMessage<UE::Anim::FAnimInertializationSyncScope>() != nullptr;;
+	TickRecord.bRequestedInertialization = Context.GetMessage<UE::Anim::FAnimInertializationSyncScope>() != nullptr;
 
+	if (bHasValidAssetPlayerInstanceID)
+	{
+		TickRecord.MakeContextData<UE::Anim::FAnimNotifyAssetPlayerInstanceContext>(AssetPlayerInstanceID);
+	}
+	
 	// Add asset player to synchronizer
 	SyncScope.AddTickRecord(TickRecord, GetSyncParams(TickRecord.bRequestedInertialization), UE::Anim::FAnimSyncDebugInfo(Context));
 }
diff --git a/Engine/Source/Runtime/Engine/Private/Animation/AnimNotifyState.cpp b/Engine/Source/Runtime/Engine/Private/Animation/AnimNotifyState.cpp
--- a/Engine/Source/Runtime/Engine/Private/Animation/AnimNotifyState.cpp
+++ b/Engine/Source/Runtime/Engine/Private/Animation/AnimNotifyState.cpp
@@ -23,6 +23,7 @@ UAnimNotifyState::UAnimNotifyState(const FObjectInitializer& ObjectInitializer)
 #endif // WITH_EDITORONLY_DATA
 
 	bIsNativeBranchingPoint = false;
+	NotifyStateBehaviorFlags = 0;
 }
 
 void UAnimNotifyState::NotifyBegin(USkeletalMeshComponent* MeshComp, UAnimSequenceBase* Animation, float TotalDuration, const FAnimNotifyEventReference& EventReference)
diff --git a/Engine/Source/Runtime/Engine/Private/Animation/AnimationInstanceScope.cpp b/Engine/Source/Runtime/Engine/Private/Animation/AnimationInstanceScope.cpp
new file mode 100644
--- /dev/null
+++ b/Engine/Source/Runtime/Engine/Private/Animation/AnimationInstanceScope.cpp
@@ -0,0 +1,14 @@
+// Copyright Epic Games, Inc. All Rights Reserved.
+
+
+#include "Animation/AnimationInstanceScope.h"
+
+namespace UE::Anim
+{
+IMPLEMENT_NOTIFY_CONTEXT_INTERFACE(FAnimNotifyAssetPlayerInstanceContext)
+
+FAnimNotifyAssetPlayerInstanceContext::FAnimNotifyAssetPlayerInstanceContext(const uint32 InAssetPlayerInstanceID)
+	: AssetPlayerInstanceID(InAssetPlayerInstanceID)
+{
+}
+}
diff --git a/Engine/Source/Runtime/Engine/Public/Animation/AnimationInstanceScope.h b/Engine/Source/Runtime/Engine/Public/Animation/AnimationInstanceScope.h
new file mode 100644
--- /dev/null
+++ b/Engine/Source/Runtime/Engine/Public/Animation/AnimationInstanceScope.h
@@ -0,0 +1,20 @@
+// Copyright Epic Games, Inc. All Rights Reserved.
+
+#pragma once
+
+#include "AnimNodeMessages.h"
+
+namespace UE::Anim
+{
+	
+// Notify Context Data containing the Instance ID of the asset player that triggered this notify
+// This is used to determine if a notify state was triggered by different instances of the same animation 
+class FAnimNotifyAssetPlayerInstanceContext : public UE::Anim::IAnimNotifyEventContextDataInterface
+{
+	DECLARE_NOTIFY_CONTEXT_INTERFACE_API(FAnimNotifyAssetPlayerInstanceContext, ENGINE_API)
+public:
+	ENGINE_API FAnimNotifyAssetPlayerInstanceContext(const uint32 InAssetPlayerInstanceID);
+	const uint32 AssetPlayerInstanceID;
+};
+	
+}
