PR #14498
Montages: When ForcedNextFromPosition is set, use it as the substepping origin in FAnimMontageInstance::Advance instead of the previous internal position. Controllable via CVar a.Montage.UseForcedFromPositionAsSubStepOrigin, enabled by default.

[RN] fix
#Synced-CL 52338451

diff --git a/Engine/Source/Runtime/Engine/Private/Animation/AnimMontage.cpp b/Engine/Source/Runtime/Engine/Private/Animation/AnimMontage.cpp
--- a/Engine/Source/Runtime/Engine/Private/Animation/AnimMontage.cpp
+++ b/Engine/Source/Runtime/Engine/Private/Animation/AnimMontage.cpp
@@ -50,6 +50,13 @@ namespace MontageCVars
 		bEndSectionRequiresTimeRemaining,
 		TEXT("Montage EndOfSection is only checked if there is remaining time (default false)."));
 	extern bool bEarlyOutMontageWhenUninitialized;
+	
+	static bool bUseForcedFromPositionAsSubStepOrigin = true;
+	static FAutoConsoleVariableRef CVarUseForcedFromPositionAsSubStepOrigin(
+		TEXT("a.Montage.UseForcedFromPositionAsSubStepOrigin"),
+		bUseForcedFromPositionAsSubStepOrigin,
+		TEXT(" When true (default), if a ForcedNextFromPosition is set during Montage Advance, use that forced position as the sub-step start point rather than the stale cached Position."));
+	
 } // end namespace MontageCVars
 
 ///////////////////////////////////////////////////////////////////////////
@@ -2560,7 +2567,10 @@ void FAnimMontageInstance::Advance(float DeltaTime, struct FRootMotionMovementPa
 			{
 				SCOPE_CYCLE_COUNTER(STAT_AnimMontageInstance_Advance_Iteration);
 
-				const float PreviousSubStepPosition = Position;
+				const float PreviousSubStepPosition = (MontageCVars::bUseForcedFromPositionAsSubStepOrigin && ForcedNextFromPosition.IsSet())
+					? ForcedNextFromPosition.GetValue()
+					: Position;
+				
 				const FBranchingPointMarker* BranchingPointMarker = nullptr;
 				EMontageSubStepResult SubStepResult = MontageSubStepper.Advance(Position, &BranchingPointMarker);
 
