FSimpleStreamableAssetManager: Automated test suite

Adds a comprehensive test suite for SSAM covering registration, unregistration,
streaming queries, lifecycle/memory, concurrency, and GC safety. 36 tests total
across 7 groups:

- Registration (5): single/multi object, shared assets, deduplication, zero texel filter
- Unregistration (5): object removal, partial unregister, same-frame, asset removal,
InstanceRemovedTimestamp (disabled, pending fix)
- Fallback (2): component-based registration, double-registration bug (disabled, pending fix)
- Streaming Queries (11): math correctness, relative ordering, view movement, force mip,
visibility, view-inside-bounds, split-screen, max draw distance, scale/boost factors
- Lifecycle (5): empty process, index recycling, rapid create/destroy,
level-transition memory bloat (disabled, pending UE-363476 fix), bulk register/unregister
- Concurrency (7): parallel GT registration, concurrent Process+register/unregister/update,
triple-producer, GC during Process, stress test
- Race detector canary: deliberately racy test to validate instrumented build detection

Bug fix: FSimpleStreamableAssetManager::Shutdown() now sets Instance = nullptr after delete,
preventing use-after-free when tests re-initialize the singleton across test boundaries.

Test infrastructure (SSAMTestTypes.h/.cpp):
- USSAMTestStreamableAsset: minimal streamable asset with BeginDestroy override for GC tests
- USSAMTestComponent / FSSAMTestSceneProxy: configurable component+proxy for SSAM registration
- FSSAMTestEnvironment: RAII helper managing SSAM lifecycle, world, asset cleanup
- Helper methods: UpdateBoundSizes(), QueryScreenSize(), ReleaseAssetsForGC()

Run: UnrealEditor EngineTest -ExecCmds="Automation RunTests System.Engine.Streaming.SSAM;Quit"

#Synced-CL 52335816

diff --git a/Engine/Source/Runtime/Engine/Private/Streaming/SimpleStreamableAssetManager.cpp b/Engine/Source/Runtime/Engine/Private/Streaming/SimpleStreamableAssetManager.cpp
--- a/Engine/Source/Runtime/Engine/Private/Streaming/SimpleStreamableAssetManager.cpp
+++ b/Engine/Source/Runtime/Engine/Private/Streaming/SimpleStreamableAssetManager.cpp
@@ -71,6 +71,7 @@ void FSimpleStreamableAssetManager::Shutdown()
 {
 	check(Instance != nullptr);
 	delete Instance;
+	Instance = nullptr;
 }
 
 void FSimpleStreamableAssetManager::Process()
diff --git a/Engine/Source/Runtime/Engine/Private/Tests/Streaming/SSAMCoreTests.cpp b/Engine/Source/Runtime/Engine/Private/Tests/Streaming/SSAMCoreTests.cpp
new file mode 100644
--- /dev/null
+++ b/Engine/Source/Runtime/Engine/Private/Tests/Streaming/SSAMCoreTests.cpp
@@ -0,0 +1,2027 @@
+// Copyright Epic Games, Inc. All Rights Reserved.
+
+// =============================================================================================================================================
+// FSimpleStreamableAssetManager (SSAM) — Core Test Suite
+//
+// Tests for the lightweight streaming asset manager that replaces the legacy
+// FStreamingManagerTexture path for computing per-asset screen sizes.
+//
+// Tests:
+// ┌────────┬──────────────────────────────────────────────┬────────────────────────────────────────────────────────┬──────────────────────────┐
+// │ Group  │ Test                                         │ What it verifies                                       │ Registration path        │
+// ├────────┼──────────────────────────────────────────────┼────────────────────────────────────────────────────────┼──────────────────────────┤
+// │ 1      │ Registration.SingleObject                    │ 1 component, 2 assets → valid indices + bounds         │ Primary (proxy)          │
+// │ 1      │ Registration.MultipleObjectsSharingAsset     │ 2 components share 1 asset → 2 bounds                  │ Primary (proxy)          │
+// │ 1      │ Registration.MultipleAssetsPerObject         │ 1 component, 5 assets → unique indices                 │ Primary (proxy)          │
+// │ 1      │ Registration.AssetDeduplication              │ Same asset twice → merged to 1 bound (highest texel)   │ Primary (proxy)          │
+// │ 1      │ Registration.ZeroTexelFactorFiltered         │ TexelFactor=0 → filtered out, no registration          │ Primary (proxy)          │
+// ├────────┼──────────────────────────────────────────────┼────────────────────────────────────────────────────────┼──────────────────────────┤
+// │ 2      │ Unregistration.Object                        │ DestroyComponent → shared index reset, bounds cleared  │ Primary (proxy)          │
+// │ 2      │ Unregistration.OneOfTwoObjects               │ Destroy 1 of 2 → 1 bound remains                       │ Primary (proxy)          │
+// │ 2      │ Unregistration.SameFrame                     │ Register+destroy before Process → clean state          │ Primary (proxy)          │
+// │ 2      │ Unregistration.Asset                         │ UnregisterAsset → asset index freed                    │ Primary (proxy)          │
+// │ 2      │ Unregistration.InstanceRemovedTimestamp      │ (DISABLED) Unregister → InstanceRemovedTimestamp       │ Primary (proxy)          │
+// ├────────┼──────────────────────────────────────────────┼────────────────────────────────────────────────────────┼──────────────────────────┤
+// │ 3      │ Fallback.Registration                        │ Proxy w/o streaming → fallback component path          │ Fallback (component)     │
+// │ 3      │ Fallback.DoubleRegistrationBug               │ (DISABLED) RecreateRenderState double-register guard   │ Fallback (component)     │
+// ├────────┼──────────────────────────────────────────────┼────────────────────────────────────────────────────────┼──────────────────────────┤
+// │ 4      │ StreamingQuery.MathCorrectness               │ Exact screen size vs bUseNewMetrics formula            │ Primary (proxy)          │
+// │ 4      │ StreamingQuery.RelativeCorrectness           │ Nearer > farther, higher texel > lower, 2x = 2x        │ Primary (proxy)          │
+// │ 4      │ StreamingQuery.UpdatesOnMove                 │ Far view → near view → screen size increases           │ Primary (proxy)          │
+// │ 4      │ StreamingQuery.ForceMipStreaming             │ bForceMipStreaming → MaxSize = FLT_MAX                 │ Primary (proxy)          │
+// │ 4      │ StreamingQuery.VisibilityBased               │ ConsiderVisibility CVar → VisibleOnly = 0 if unrendered│ Primary (proxy)          │
+// │ 4      │ StreamingQuery.ViewInsideBounds              │ Camera inside box → dist=0, max quality                │ Primary (proxy)          │
+// │ 4      │ StreamingQuery.MultipleViews                 │ Split-screen: result = max across views                │ Primary (proxy)          │
+// │ 4      │ StreamingQuery.MaxDrawDistance               │ Beyond CachedMaxDrawDistance → VisibleOnly = 0         │ Primary (proxy)          │
+// │ 4      │ StreamingQuery.ScaleFactor                   │ 2x world scale → >2x screen size (scale + extents)     │ Primary (proxy)          │
+// │ 4      │ StreamingQuery.BoostFactor                   │ 2x ExtraBoost → 2x VisibleOnly                         │ Primary (proxy)          │
+// │ 4      │ StreamingQuery.ObjectMovement                │ SetWorldLocation + FUpdate path → bounds update        │ Primary (proxy) + RT     │
+// ├────────┼──────────────────────────────────────────────┼────────────────────────────────────────────────────────┼──────────────────────────┤
+// │ 5      │ Lifecycle.EmptyProcess                       │ Process() with no pending records → safe no-op         │ N/A                      │
+// │ 5      │ Lifecycle.IndexRecyclingNoLeak               │ Freed indices reused, memory stays bounded             │ Primary (proxy)          │
+// │ 5      │ Lifecycle.RapidCreateDestroy                 │ 50 create/destroy cycles → no leak, no lingering bounds│ Primary (proxy)          │
+// │ 5      │ Lifecycle.MemoryBloatOnLevelTransition       │ (DISABLED) UE-363476: inner sparse arrays freed        │ Primary (proxy)          │
+// │ 5      │ Lifecycle.BulkRegisterUnregister             │ 100 components, 3 assets → bulk add/remove correctness │ Primary (proxy)          │
+// ├────────┼──────────────────────────────────────────────┼────────────────────────────────────────────────────────┼──────────────────────────┤
+// │ 6      │ Concurrency.ParallelRegistration             │ ParallelFor + AddPrimitive via FPrimitiveSceneDesc     │ FPrimitiveSceneDesc (GT) │
+// │ 6      │ Concurrency.ConcurrentProcessAndRegister     │ BG Process() loop + GT register bursts (cooked build)  │ Primary (proxy)          │
+// │ 6      │ Concurrency.ConcurrentProcessAndUnregister   │ BG Process() loop + GT DestroyComponent (cooked build) │ Primary (proxy)          │
+// │ 6      │ Concurrency.ConcurrentProcessAndUpdate       │ BG Process() loop + GT moves via RT (cooked build)     │ Primary (proxy) + RT     │
+// │ 6      │ Concurrency.TripleProducerConcurrency        │ BG Process + parallel GT desc-register + GT unregister │ Both (desc + proxy)      │
+// │ 6      │ Concurrency.GarbageCollectionDuringProcess   │ BG Process() + GT CollectGarbage (real GC path)        │ Primary (proxy)          │
+// │ 6      │ Concurrency.ConcurrentOperationsStress       │ 5 iterations of parallel add/remove (stress test)      │ FPrimitiveSceneDesc (GT) │
+// └────────┴──────────────────────────────────────────────┴────────────────────────────────────────────────────────┴──────────────────────────┘
+//
+// Threading model:
+// ┌──────────────────┬────────────────────────────────────────────────────────────────────────────────────────────────────────────────────────┐
+// │ Thread           │ Role in SSAM                                                                                                           │
+// ├──────────────────┼────────────────────────────────────────────────────────────────────────────────────────────────────────────────────────┤
+// │ Game Thread      │ RegisterComponent/DestroyComponent push FRegister/FUnregister to MPMC queues.                                          │
+// │                  │ SetWorldLocation triggers SendRenderTransform_Concurrent which enqueues an RT command.                                 │
+// ├──────────────────┼────────────────────────────────────────────────────────────────────────────────────────────────────────────────────────┤
+// │ Parallel GT      │ FastGeo calls CreateRenderState from parallel workers (ParallelFor + EParallelGameThread),                             │
+// │ (EParallelGT)    │ which internally calls CreateSceneProxy → AddPrimitive → FRegister push.                                               │
+// ├──────────────────┼────────────────────────────────────────────────────────────────────────────────────────────────────────────────────────┤
+// │ Render Thread    │ Executes deferred commands from GT: UpdatePrimitiveTransform_RenderThread pushes FUpdate records to SSAM queues.       │
+// │                  │ FlushRenderingCommands() synchronizes.                                                                                 │
+// ├──────────────────┼────────────────────────────────────────────────────────────────────────────────────────────────────────────────────────┤
+// │ Async Worker     │ In cooked builds, FDoWorkAsyncTask calls Process() every frame. Process() acquires CriticalSection,                    │
+// │ (thread pool)    │ then drains all MPMC queues via ConsumeAll + updates bounds.                                                           │
+// └──────────────────┴────────────────────────────────────────────────────────────────────────────────────────────────────────────────────────┘
+//
+// =============================================================================================================================================
+
+#include "Misc/AutomationTest.h"
+
+#if WITH_DEV_AUTOMATION_TESTS
+
+#include "SSAMTestTypes.h"
+#include "Streaming/SimpleStreamableAssetManager.h"
+#include "ContentStreaming.h"
+#include "Streaming/TextureStreamingHelpers.h"
+#include "Streaming/TextureInstanceView.h"
+#include "Streaming/StreamingManagerTexture.h"
+#include "Streaming/StreamingTexture.h"
+#include "Async/ParallelFor.h"
+#include "Async/TaskGraphInterfaces.h"
+#include "SceneInterface.h"
+#include "PrimitiveSceneDesc.h"
+#include "PrimitiveSceneProxyDesc.h"
+#include "Tasks/Task.h"
+
+#define SSAM_TEST_NAME "System.Engine.Streaming.SSAM"
+
+namespace SSAMTests
+{
+	static const EAutomationTestFlags TestFlags =
+		EAutomationTestFlags::EditorContext
+		| EAutomationTestFlags::ClientContext
+		| EAutomationTestFlags::EngineFilter;
+
+	// ===================================================================
+	// GROUP 1: Registration (Primary Path)
+	//
+	// Tests the primary SSAM registration flow: component registration
+	// triggers CreateSceneProxy → AddPrimitive → BatchAddPrimitivesInternal
+	// → RegisterHelper → FRegister(proxy) → queue → UpdateTask_Async.
+	// All tests in this group use FSSAMTestSceneProxy with
+	// bImplementsStreamableAssetGathering = true.
+	// ===================================================================
+
+	// -----------------------------------------------------------------------
+	// Register one component with two assets. Verify that after Process(),
+	// both assets get valid SSAM indices and GetAssetReferenceBounds
+	// returns the component's bounds for each.
+	// -----------------------------------------------------------------------
+	IMPLEMENT_SIMPLE_AUTOMATION_TEST(FSSAMSingleObjectRegistration, SSAM_TEST_NAME ".Registration.SingleObject", TestFlags)
+
+	bool FSSAMSingleObjectRegistration::RunTest(const FString& Parameters)
+	{
+		FSSAMTestEnvironment Env;
+
+		USSAMTestStreamableAsset* AssetA = Env.CreateTestAsset();
+		USSAMTestStreamableAsset* AssetB = Env.CreateTestAsset();
+
+		const FBoxSphereBounds TestBounds(FVector(100.f, 200.f, 300.f), FVector(50.f), 50.f);
+
+		Env.SpawnTestComponent({ { AssetA, 512.f }, { AssetB, 256.f } }, TestBounds);
+		Env.ProcessSSAM();
+
+		TestTrue(TEXT("AssetA has valid SSAM index"), AssetA->GetSSAMIndexValue() != INDEX_NONE);
+		TestTrue(TEXT("AssetB has valid SSAM index"), AssetB->GetSSAMIndexValue() != INDEX_NONE);
+
+		TArray<FBox> BoxesA, BoxesB;
+		FSimpleStreamableAssetManager::GetAssetReferenceBounds(AssetA, BoxesA);
+		FSimpleStreamableAssetManager::GetAssetReferenceBounds(AssetB, BoxesB);
+		TestEqual(TEXT("AssetA has 1 reference bound"), BoxesA.Num(), 1);
+		TestEqual(TEXT("AssetB has 1 reference bound"), BoxesB.Num(), 1);
+
+		return true;
+	}
+
+	// -----------------------------------------------------------------------
+	// Two components reference the same asset. Verify that
+	// GetAssetReferenceBounds returns two boxes and the asset index
+	// is shared (assigned once, used by both bound elements).
+	// -----------------------------------------------------------------------
+	IMPLEMENT_SIMPLE_AUTOMATION_TEST(FSSAMMultipleObjectsSharingAsset, SSAM_TEST_NAME ".Registration.MultipleObjectsSharingAsset", TestFlags)
+
+	bool FSSAMMultipleObjectsSharingAsset::RunTest(const FString& Parameters)
+	{
+		FSSAMTestEnvironment Env;
+
+		USSAMTestStreamableAsset* SharedAsset = Env.CreateTestAsset();
+
+		Env.SpawnTestComponent({ { SharedAsset, 512.f } }, FBoxSphereBounds(FVector(0.f), FVector(100.f), 100.f));
+		Env.SpawnTestComponent({ { SharedAsset, 256.f } }, FBoxSphereBounds(FVector(500.f, 0.f, 0.f), FVector(100.f), 100.f));
+
+		Env.ProcessSSAM();
+
+		TestTrue(TEXT("SharedAsset has valid SSAM index"), SharedAsset->GetSSAMIndexValue() != INDEX_NONE);
+
+		TArray<FBox> Boxes;
+		FSimpleStreamableAssetManager::GetAssetReferenceBounds(SharedAsset, Boxes);
+		TestEqual(TEXT("SharedAsset has 2 reference bounds"), Boxes.Num(), 2);
+
+		return true;
+	}
+
+	// -----------------------------------------------------------------------
+	// Register a single component that references 5 different assets.
+	// Verify each asset gets its own unique SSAM index and exactly one
+	// reference bound.
+	// -----------------------------------------------------------------------
+	IMPLEMENT_SIMPLE_AUTOMATION_TEST(FSSAMMultipleAssetsPerObject, SSAM_TEST_NAME ".Registration.MultipleAssetsPerObject", TestFlags)
+
+	bool FSSAMMultipleAssetsPerObject::RunTest(const FString& Parameters)
+	{
+		FSSAMTestEnvironment Env;
+
+		constexpr int32 NumAssets = 5;
+		TArray<USSAMTestStreamableAsset*> Assets;
+		TArray<FSSAMTestAssetEntry> Entries;
+		for (int32 i = 0; i < NumAssets; ++i)
+		{
+			USSAMTestStreamableAsset* Asset = Env.CreateTestAsset();
+			Assets.Add(Asset);
+			Entries.Add({ Asset, (float)(100 * (i + 1)) });
+		}
+
+		Env.SpawnTestComponent(Entries);
+		Env.ProcessSSAM();
+
+		TSet<int32> SeenIndices;
+		for (int32 i = 0; i < NumAssets; ++i)
+		{
+			TestTrue(FString::Printf(TEXT("Asset[%d] has valid SSAM index"), i), Assets[i]->GetSSAMIndexValue() != INDEX_NONE);
+
+			bool bAlreadyInSet = false;
+			SeenIndices.Add(Assets[i]->GetSSAMIndexValue(), &bAlreadyInSet);
+			TestFalse(FString::Printf(TEXT("Asset[%d] index is unique"), i), bAlreadyInSet);
+
+			TArray<FBox> Boxes;
+			FSimpleStreamableAssetManager::GetAssetReferenceBounds(Assets[i], Boxes);
+			TestEqual(FString::Printf(TEXT("Asset[%d] has 1 reference bound"), i), Boxes.Num(), 1);
+		}
+
+		return true;
+	}
+
+	// -----------------------------------------------------------------------
+	// Register a component that references the SAME asset twice with
+	// different texel factors. With deduplication enabled (default CVar
+	// GSimpleStreamableAssetManagerEnsureAssetUniqueOnRegistration = 1),
+	// SSAM should merge them into a single bound element, keeping the
+	// highest texel factor.
+	// -----------------------------------------------------------------------
+	IMPLEMENT_SIMPLE_AUTOMATION_TEST(FSSAMAssetDeduplication, SSAM_TEST_NAME ".Registration.AssetDeduplication", TestFlags)
+
+	bool FSSAMAssetDeduplication::RunTest(const FString& Parameters)
+	{
+		FSSAMTestEnvironment Env;
+
+		// Guard: this test requires deduplication to be enabled (default CVar = 1).
+		IConsoleVariable* DedupCVar = IConsoleManager::Get().FindConsoleVariable( TEXT("s.StreamableAssets.SimpleStreamableAssetManager.SortAssetsOnRegistration"));
+		const int32 OriginalDedup = DedupCVar ? DedupCVar->GetInt() : 1;
+		if (DedupCVar)
+		{
+			DedupCVar->Set(1, ECVF_SetByCode);
+		}
+
+		USSAMTestStreamableAsset* Asset = Env.CreateTestAsset();
+
+		Env.SpawnTestComponent({ { Asset, 128.f }, { Asset, 512.f } });
+		Env.ProcessSSAM();
+
+		TestTrue(TEXT("Asset has valid index"), Asset->GetSSAMIndexValue() != INDEX_NONE);
+
+		TArray<FBox> Boxes;
+		FSimpleStreamableAssetManager::GetAssetReferenceBounds(Asset, Boxes);
+		TestEqual(TEXT("Deduplicated asset has exactly 1 reference bound"), Boxes.Num(), 1);
+
+		if (DedupCVar)
+		{
+			DedupCVar->Set(OriginalDedup, ECVF_SetByCode);
+		}
+		return true;
+	}
+
+	// -----------------------------------------------------------------------
+	// Register a component with TexelFactor = 0. The FRegister constructor
+	// filters assets via CanBeStreamedByDistance which requires
+	// TexelFactor > UE_SMALL_NUMBER. An asset with zero texel factor
+	// should NOT be registered (no bound elements, no index assigned).
+	// -----------------------------------------------------------------------
+	IMPLEMENT_SIMPLE_AUTOMATION_TEST(FSSAMZeroTexelFactorFiltered, SSAM_TEST_NAME ".Registration.ZeroTexelFactorFiltered", TestFlags)
+
+	bool FSSAMZeroTexelFactorFiltered::RunTest(const FString& Parameters)
+	{
+		FSSAMTestEnvironment Env;
+
+		USSAMTestStreamableAsset* Asset = Env.CreateTestAsset();
+
+		// Register with TexelFactor = 0 — should be filtered out
+		Env.SpawnTestComponent({ { Asset, 0.f } });
+		Env.ProcessSSAM();
+
+		// The asset should NOT have been assigned an index (no streamable assets
+		// passed the CanBeStreamedByDistance filter)
+		TArray<FBox> Boxes;
+		FSimpleStreamableAssetManager::GetAssetReferenceBounds(Asset, Boxes);
+		TestEqual(TEXT("Zero texel factor asset has 0 bounds (filtered out)"), Boxes.Num(), 0);
+
+		return true;
+	}
+
+	// ===================================================================
+	// GROUP 2: Unregistration
+	//
+	// Tests all unregistration paths: object removal (DestroyComponent →
+	// BatchRemovePrimitivesInternal → UnregisterHelper), same-frame
+	// register+unregister, partial unregister (one of two objects), and
+	// asset-level removal (UnregisterAsset, simulating texture unload).
+	// ===================================================================
+
+	// -----------------------------------------------------------------------
+	// Register a component, process, then destroy it and process again.
+	// Verify the component's shared index is reset and the asset's
+	// reference bounds are cleaned up.
+	// -----------------------------------------------------------------------
+	IMPLEMENT_SIMPLE_AUTOMATION_TEST(FSSAMUnregisterObject, SSAM_TEST_NAME ".Unregistration.Object", TestFlags)
+
+	bool FSSAMUnregisterObject::RunTest(const FString& Parameters)
+	{
+		FSSAMTestEnvironment Env;
+
+		USSAMTestStreamableAsset* Asset = Env.CreateTestAsset();
+		USSAMTestComponent* Comp = Env.SpawnTestComponent({ { Asset, 512.f } });
+		Env.ProcessSSAM();
+
+		TestTrue(TEXT("Asset registered"), Asset->GetSSAMIndexValue() != INDEX_NONE);
+
+		// Destroying the component triggers the full
+		// BatchRemovePrimitivesInternal → UnregisterHelper flow
+		Comp->DestroyComponent();
+		Env.ProcessSSAM();
+
+		TArray<FBox> Boxes;
+		FSimpleStreamableAssetManager::GetAssetReferenceBounds(Asset, Boxes);
+		TestEqual(TEXT("Asset has 0 bounds after unregister"), Boxes.Num(), 0);
+
+		return true;
+	}
+
+	// -----------------------------------------------------------------------
+	// Two components share one asset. Unregister one. Verify the asset
+	// still has one reference bound from the remaining component.
+	// -----------------------------------------------------------------------
+	IMPLEMENT_SIMPLE_AUTOMATION_TEST(FSSAMUnregisterOneOfTwoObjects, SSAM_TEST_NAME ".Unregistration.OneOfTwoObjects", TestFlags)
+
+	bool FSSAMUnregisterOneOfTwoObjects::RunTest(const FString& Parameters)
+	{
+		FSSAMTestEnvironment Env;
+
+		USSAMTestStreamableAsset* SharedAsset = Env.CreateTestAsset();
+
+		USSAMTestComponent* Comp1 = Env.SpawnTestComponent({ { SharedAsset, 512.f } }, FBoxSphereBounds(FVector::ZeroVector, FVector(100.f), 100.f));
+		USSAMTestComponent* Comp2 = Env.SpawnTestComponent({ { SharedAsset, 256.f } }, FBoxSphereBounds(FVector(1000.f, 0.f, 0.f), FVector(100.f), 100.f));
+
+		Env.ProcessSSAM();
+
+		{
+			TArray<FBox> Boxes;
+			FSimpleStreamableAssetManager::GetAssetReferenceBounds(SharedAsset, Boxes);
+			TestEqual(TEXT("Asset has 2 bounds before partial unregister"), Boxes.Num(), 2);
+		}
+
+		Comp1->DestroyComponent();
+		Env.ProcessSSAM();
+
+		{
+			TArray<FBox> Boxes;
+			FSimpleStreamableAssetManager::GetAssetReferenceBounds(SharedAsset, Boxes);
+			TestEqual(TEXT("Asset has 1 bound after partial unregister"), Boxes.Num(), 1);
+		}
+
+		TestTrue(TEXT("SharedAsset still has valid SSAM index"), SharedAsset->GetSSAMIndexValue() != INDEX_NONE);
+
+		return true;
+	}
+
+	// -----------------------------------------------------------------------
+	// Register a component, then immediately destroy it before calling
+	// Process(). Both FRegister and FUnregister are processed in the same
+	// UpdateTask_Async call. The processing order (assign index → unregister
+	// → register) means the unregister resets the index to INDEX_NONE, then
+	// RegisterRecord sees INDEX_NONE and early-returns. Net result: clean state.
+	// -----------------------------------------------------------------------
+	IMPLEMENT_SIMPLE_AUTOMATION_TEST(FSSAMRegisterUnregisterSameFrame, SSAM_TEST_NAME ".Unregistration.SameFrame", TestFlags)
+
+	bool FSSAMRegisterUnregisterSameFrame::RunTest(const FString& Parameters)
+	{
+		FSSAMTestEnvironment Env;
+
+		USSAMTestStreamableAsset* Asset = Env.CreateTestAsset();
+		USSAMTestComponent* Comp = Env.SpawnTestComponent({ { Asset, 512.f } });
+		Comp->DestroyComponent();
+		Env.ProcessSSAM();
+
+		TArray<FBox> Boxes;
+		FSimpleStreamableAssetManager::GetAssetReferenceBounds(Asset, Boxes);
+		TestEqual(TEXT("Asset has 0 bounds after same-frame register+unregister"), Boxes.Num(), 0);
+
+		return true;
+	}
+
+	// -----------------------------------------------------------------------
+	// Register an object referencing an asset. Then call UnregisterAsset
+	// on the asset itself (simulating texture unload). Verify the asset
+	// index is freed and gets a fresh INDEX_NONE pointer, while the
+	// object registration remains intact.
+	// -----------------------------------------------------------------------
+	IMPLEMENT_SIMPLE_AUTOMATION_TEST(FSSAMUnregisterAsset, SSAM_TEST_NAME ".Unregistration.Asset", TestFlags)
+
+	bool FSSAMUnregisterAsset::RunTest(const FString& Parameters)
+	{
+		FSSAMTestEnvironment Env;
+
+		USSAMTestStreamableAsset* Asset = Env.CreateTestAsset();
+		USSAMTestComponent* Comp = Env.SpawnTestComponent({ { Asset, 512.f } });
+		Env.ProcessSSAM();
+
+		TestTrue(TEXT("Asset has valid index"), Asset->GetSSAMIndexValue() != INDEX_NONE);
+
+		// Unregister the ASSET (not the object)
+		FSimpleStreamableAssetManager::UnregisterAsset(Asset);
+		Env.ProcessSSAM();
+
+		// After UnregisterAsset, the asset's shared index should have been
+		// replaced with a fresh INDEX_NONE pointer
+		TestTrue(TEXT("Asset gets INDEX_NONE after UnregisterAsset"), Asset->GetSSAMIndexValue() == INDEX_NONE);
+
+		TArray<FBox> Boxes;
+		FSimpleStreamableAssetManager::GetAssetReferenceBounds(Asset, Boxes);
+		TestEqual(TEXT("Asset has 0 bounds after UnregisterAsset"), Boxes.Num(), 0);
+
+		return true;
+	}
+
+// DISABLED - Currently fails - Fix pending
+#ifdef TEST_SSAM_INSTANCE_REMOVED_TIMESTAMP
+	// -----------------------------------------------------------------------
+	// When all SSAM references to an asset are removed (component destroyed),
+	// the streaming manager's InstanceRemovedTimestamp on the corresponding
+	// FStreamingRenderAsset must be updated to the current time. Without this,
+	// the unknown-ref heuristic incorrectly forces full mip loading for assets
+	// that have no streaming references but are still rendered (e.g. cloud
+	// textures used by volumetric rendering, which never go through component
+	// registration).
+	// -----------------------------------------------------------------------
+	IMPLEMENT_SIMPLE_AUTOMATION_TEST(FSSAMUnregisterInstanceRemovedTimestamp, SSAM_TEST_NAME ".Unregistration.InstanceRemovedTimestamp", TestFlags)
+
+	bool FSSAMUnregisterInstanceRemovedTimestamp::RunTest(const FString& Parameters)
+	{
+		// This test requires the real FRenderAssetStreamingManager, which is only
+		// created when r.TextureStreaming is enabled. Ensure it's on for this test.
+		IConsoleVariable* TextureStreamingCVar = IConsoleManager::Get().FindConsoleVariable(TEXT("r.TextureStreaming"));
+		const int32 OriginalTextureStreaming = TextureStreamingCVar ? TextureStreamingCVar->GetInt() : 0;
+		if (TextureStreamingCVar && !OriginalTextureStreaming)
+		{
+			TextureStreamingCVar->Set(1, ECVF_SetByCode);
+			IStreamingManager::Get().Tick(0.f, true);
+		}
+		ON_SCOPE_EXIT
+		{
+			if (TextureStreamingCVar && !OriginalTextureStreaming)
+			{
+				TextureStreamingCVar->Set(0, ECVF_SetByCode);
+				IStreamingManager::Get().Tick(0.f, true);
+			}
+		};
+
+		FSSAMTestEnvironment Env;
+
+		USSAMTestStreamableAsset* Asset = Env.CreateTestAsset();
+
+		// Set up fake streaming state so the streaming manager creates a
+		// valid FStreamingRenderAsset (needs non-zero mip counts).
+		Asset->SetupFakeStreamingState();
+		Asset->SetStreamingIndexForTest(INDEX_NONE);
+
+		FRenderAssetStreamingManager& StreamingMgr = static_cast<FRenderAssetStreamingManager&>(
+			IStreamingManager::Get().GetRenderAssetStreamingManager());
+		StreamingMgr.AddStreamingRenderAsset(Asset);
+
+		// Register a component that references the asset through SSAM.
+		USSAMTestComponent* Comp = Env.SpawnTestComponent({ { Asset, 512.f } });
+		Env.ProcessSSAM();
+
+		// A single UpdateResourceStreaming flushes the pending asset into
+		// StreamingRenderAssets (via ProcessAddedRenderAssets) and runs the
+		// mip calc. At this point the asset has SSAM references so it's fine.
+		StreamingMgr.UpdateResourceStreaming(0.f, true);
+
+		FStreamingRenderAsset* StreamingAsset = StreamingMgr.GetStreamingRenderAsset(Asset);
+		if (!TestNotNull(TEXT("FStreamingRenderAsset exists after flush"), StreamingAsset))
+		{
+			StreamingMgr.RemoveStreamingRenderAsset(Asset);
+			return false;
+		}
+
+		// Record the timestamp before unregistration.
+		const double TimestampBefore = StreamingAsset->InstanceRemovedTimestamp;
+
+		// Advance FApp::GetCurrentTime() — it's a per-frame cached value that
+		// doesn't advance with wall clock or Sleep. The mip calc uses it to
+		// set InstanceRemovedTimestamp, so we need a visible delta.
+		FAppTime::SetCurrentTime(FApp::GetCurrentTime() + 1.0);
+
+		// Destroy the component, removing all SSAM references.
+		Comp->DestroyComponent();
+		Env.ProcessSSAM();
+
+		// Run a second streaming update — the mip calc consumes the removed
+		// asset indices and updates InstanceRemovedTimestamp.
+		StreamingMgr.UpdateResourceStreaming(0.f, true);
+
+		// Re-fetch: UpdateResourceStreaming may have reallocated the array.
+		StreamingAsset = StreamingMgr.GetStreamingRenderAsset(Asset);
+		if (!TestNotNull(TEXT("FStreamingRenderAsset still valid after streaming update"), StreamingAsset))
+		{
+			StreamingMgr.RemoveStreamingRenderAsset(Asset);
+			return false;
+		}
+
+		// The core assertion: InstanceRemovedTimestamp must have been updated
+		// so that the unknown-ref heuristic has a correct time baseline.
+		TestTrue(TEXT("InstanceRemovedTimestamp updated after SSAM unregistration"),
+			StreamingAsset->InstanceRemovedTimestamp > TimestampBefore);
+
+		StreamingMgr.RemoveStreamingRenderAsset(Asset);
+		return true;
+	}
+#endif // DISABLED_TEST_SSAM_INSTANCE_REMOVED_TIMESTAMP
+
+	// ===================================================================
+	// GROUP 3: Fallback Path
+	//
+	// Tests the fallback registration path for proxies that do NOT
+	// implement GetStreamableRenderAssetInfo (bImplementsStreamableAssetGathering
+	// = false). RegisterHelper pushes the component to
+	// FallbackComponentRegistration, and ProcessFallbackRegistrations_GameThread
+	// gathers assets via the component's GetStreamingRenderAssetInfo override.
+	// ===================================================================
+
+	// -----------------------------------------------------------------------
+	// Use a proxy that does NOT implement GetStreamableRenderAssetInfo.
+	// RegisterHelper takes the fallback path, pushing the component to
+	// FallbackComponentRegistration. ProcessFallbackRegistrations_GameThread
+	// then gathers assets from the component (via GetStreamingRenderAssetInfo
+	// override) and creates an FRegister record. Verify full end-to-end
+	// registration including valid indices and asset bounds.
+	// -----------------------------------------------------------------------
+	IMPLEMENT_SIMPLE_AUTOMATION_TEST(FSSAMFallbackRegistration, SSAM_TEST_NAME ".Fallback.Registration", TestFlags)
+
+	bool FSSAMFallbackRegistration::RunTest(const FString& Parameters)
+	{
+		FSSAMTestEnvironment Env;
+
+		USSAMTestStreamableAsset* Asset = Env.CreateTestAsset();
+		USSAMTestComponent* Comp = Env.SpawnTestFallbackComponent({ { Asset, 512.f } });
+
+		// Process the fallback queue (now a no-op) and the main queue
+		Env.ProcessFallbackRegistrations();
+		Env.ProcessSSAM();
+
+		TestTrue(TEXT("Asset has valid index after registration"), Asset->GetSSAMIndexValue() != INDEX_NONE);
+
+		TArray<FBox> Boxes;
+		FSimpleStreamableAssetManager::GetAssetReferenceBounds(Asset, Boxes);
+		TestEqual(TEXT("Asset has 1 reference bound via fallback"), Boxes.Num(), 1);
+
+		return true;
+	}
+
+// DISABLED - Currently fails - Fix pending
+#ifdef TEST_SSAM_DOUBLE_REGISTRATION_BUG
+	// -----------------------------------------------------------------------
+	// Reproduces the double-registration bug in the fallback path:
+	// 1. Register component via fallback path (proxy A, pushed to queue)
+	// 2. RecreateRenderState → UnregisterHelper (resets component's shared
+	//    index, pushes FUnregister for proxy A) → RegisterHelper (creates
+	//    NEW shared index on component, pushes to fallback queue again)
+	// 3. BEFORE ProcessFallbackRegistrations runs, the fallback queue has
+	//    TWO TWeakObjectPtr entries for the same component
+	// 4. ProcessFallbackRegistrations processes both — the guard
+	//    if (Component->SimpleStreamableAssetManagerIndex) passes for both
+	//    because the component's CURRENT pointer (from step 2) is valid.
+	//    Both entries call Register(FRegister(Component, Assets)) with the
+	//    SAME shared index → double registration.
+	//
+	// The asset should end up with exactly 1 reference bound, not 2.
+	// This test currently FAILS, confirming the bug.
+	// -----------------------------------------------------------------------
+	IMPLEMENT_SIMPLE_AUTOMATION_TEST(FSSAMFallbackDoubleRegistrationBug, SSAM_TEST_NAME ".Fallback.DoubleRegistrationBug", TestFlags)
+
+	bool FSSAMFallbackDoubleRegistrationBug::RunTest(const FString& Parameters)
+	{
+		FSSAMTestEnvironment Env;
+
+		USSAMTestStreamableAsset* Asset = Env.CreateTestAsset();
+		USSAMTestComponent* Comp = Env.SpawnTestFallbackComponent({ { Asset, 512.f } });
+
+		// Re-create proxy before fallback queue is drained → two entries
+		Comp->RecreateRenderState_Concurrent();
+
+		// Now process both fallback entries + the main queue
+		Env.ProcessFallbackRegistrations();
+		Env.ProcessSSAM();
+
+		// THE CRITICAL CHECK: the asset should have exactly 1 reference bound.
+		// If the double-registration bug exists, it would have 2.
+		TArray<FBox> Boxes;
+		FSimpleStreamableAssetManager::GetAssetReferenceBounds(Asset, Boxes);
+		TestEqual(TEXT("Asset has exactly 1 bound (no double registration)"), Boxes.Num(), 1);
+
+		return true;
+	}
+#endif // DISABLED_TEST_SSAM_DOUBLE_REGISTRATION_BUG
+
+	// ===================================================================
+	// GROUP 4: Streaming Queries
+	//
+	// Tests the streaming query pipeline: UpdateBoundSizes computes
+	// per-object screen sizes from view info, then GetRenderAssetScreenSize
+	// multiplies by texel factor to produce the final streaming priority.
+	// These tests verify the math is correct, inputs respond in the right
+	// direction, and special modes (force mip, visibility) work.
+	// ===================================================================
+
+	// -----------------------------------------------------------------------
+	// Validate exact screen size against the bUseNewMetrics formula.
+	//
+	// With bUseNewMetrics (default), distance is computed as distance-to-box:
+	//   For each axis: dist = max(0, abs(ViewOrigin[i] - Origin[i]) - Extent[i])
+	//   DistSq = sum(dist[i]²)
+	//   ClampedDistSq = max(1, max(MinDistanceSq, DistSq))
+	//   NormalizedSize = ScreenSize / sqrt(ClampedDistSq)
+	//
+	// Then GetRenderAssetScreenSize multiplies by TexelFactor * ComponentScale:
+	//   MaxSize = TexelFactor * ComponentScale * NormalizedSize
+	// -----------------------------------------------------------------------
+	IMPLEMENT_SIMPLE_AUTOMATION_TEST(FSSAMScreenSizeMathCorrectness, SSAM_TEST_NAME ".StreamingQuery.MathCorrectness", TestFlags)
+
+	bool FSSAMScreenSizeMathCorrectness::RunTest(const FString& Parameters)
+	{
+		FSSAMTestEnvironment Env;
+
+		// Guard: this test's expected values assume bUseNewMetrics = true (default).
+		// The old-metrics path uses a different distance formula.
+		IConsoleVariable* MetricsCVar = IConsoleManager::Get().FindConsoleVariable(TEXT("r.Streaming.UseNewMetrics"));
+		const int32 OriginalMetrics = MetricsCVar ? MetricsCVar->GetInt() : 1;
+		if (MetricsCVar)
+		{
+			MetricsCVar->Set(1, ECVF_SetByCode);
+		}
+
+		USSAMTestStreamableAsset* Asset = Env.CreateTestAsset();
+
+		const FVector BoundsExtent(100.f, 100.f, 100.f);
+		const float TexelFactor = 512.f;
+		const FBoxSphereBounds Bounds(FVector::ZeroVector, BoundsExtent, BoundsExtent.Size());
+
+		Env.SpawnTestComponent({ { Asset, TexelFactor } }, Bounds);
+		Env.ProcessSSAM();
+
+		const int32 AssetIndex = Asset->GetSSAMIndexValue();
+		TestTrue(TEXT("Asset registered"), AssetIndex != INDEX_NONE);
+
+		// View at (1000, 0, 0), ScreenSize = 1920
+		const float ScreenSizeParam = 1920.f;
+		Env.UpdateBoundSizes(FVector(1000.f, 0.f, 0.f), ScreenSizeParam);
+		FSSAMTestEnvironment::FScreenSizeResult Result = Env.QueryScreenSize(AssetIndex);
+
+		// Expected: distance-to-box along X = |1000 - 0| - min(|1000|, 100) = 900
+		//           NormalizedSize = 1920 / 900 = 2.1333
+		//           MaxSize = 512 * 1.0 * 2.1333 = 1092.27
+		const float DistToBox = 1000.f - 100.f;
+		const float ExpectedMaxSize = TexelFactor * 1.0f * (ScreenSizeParam / DistToBox);
+
+		TestTrue(FString::Printf(TEXT("MaxSize (%.2f) matches expected (%.2f) within 1%%"), Result.MaxSize, ExpectedMaxSize),
+			FMath::IsNearlyEqual(Result.MaxSize, ExpectedMaxSize, ExpectedMaxSize * 0.01f));
+
+		if (MetricsCVar)
+		{
+			MetricsCVar->Set(OriginalMetrics, ECVF_SetByCode);
+		}
+		return true;
+	}
+
+	// -----------------------------------------------------------------------
+	// Verify that screen size responds correctly to each input dimension:
+	// - Nearer object → larger screen size than farther object (same texel)
+	// - Higher texel factor → larger screen size (same position)
+	// - 2x texel factor → exactly 2x screen size (proportionality)
+	// -----------------------------------------------------------------------
+	IMPLEMENT_SIMPLE_AUTOMATION_TEST(FSSAMScreenSizeRelativeCorrectness, SSAM_TEST_NAME ".StreamingQuery.RelativeCorrectness", TestFlags)
+
+	bool FSSAMScreenSizeRelativeCorrectness::RunTest(const FString& Parameters)
+	{
+		FSSAMTestEnvironment Env;
+
+		USSAMTestStreamableAsset* AssetNear = Env.CreateTestAsset();
+		USSAMTestStreamableAsset* AssetFar = Env.CreateTestAsset();
+		USSAMTestStreamableAsset* AssetHighTexel = Env.CreateTestAsset();
+		USSAMTestStreamableAsset* AssetLowTexel = Env.CreateTestAsset();
+
+		const float R = 100.f;
+		const FBoxSphereBounds NearBounds(FVector::ZeroVector, FVector(R), R);
+		const FBoxSphereBounds FarBounds(FVector(2000.f, 0.f, 0.f), FVector(R), R);
+
+		Env.SpawnTestComponent({ { AssetNear, 512.f } }, NearBounds);      // Object A: at origin
+		Env.SpawnTestComponent({ { AssetFar, 512.f } }, FarBounds);        // Object B: at (2000, 0, 0)
+		Env.SpawnTestComponent({ { AssetHighTexel, 1024.f } }, NearBounds); // Object C: at origin, 2x texel
+		Env.SpawnTestComponent({ { AssetLowTexel, 256.f } }, NearBounds);  // Object D: at origin, 0.5x texel
+		Env.ProcessSSAM();
+
+		// View at (500, 0, 0) — closer to origin objects than to far object
+		Env.UpdateBoundSizes(FVector(500.f, 0.f, 0.f));
+
+		const float SizeNear = Env.QueryScreenSize(AssetNear->GetSSAMIndexValue()).MaxSize;
+		const float SizeFar = Env.QueryScreenSize(AssetFar->GetSSAMIndexValue()).MaxSize;
+		const float SizeHighTexel = Env.QueryScreenSize(AssetHighTexel->GetSSAMIndexValue()).MaxSize;
+		const float SizeLowTexel = Env.QueryScreenSize(AssetLowTexel->GetSSAMIndexValue()).MaxSize;
+
+		// Distance test: nearer object → larger screen size
+		TestTrue(FString::Printf(TEXT("Nearer (%.2f) > Farther (%.2f)"), SizeNear, SizeFar), SizeNear > SizeFar);
+
+		// Texel factor test: higher texel factor → larger screen size (same position)
+		TestTrue(FString::Printf(TEXT("HighTexel (%.2f) > LowTexel (%.2f)"), SizeHighTexel, SizeLowTexel), SizeHighTexel > SizeLowTexel);
+
+		// Proportionality: double the texel factor → double the screen size
+		// AssetHighTexel (1024) vs AssetNear (512), both at origin
+		const float Ratio = SizeHighTexel / SizeNear;
+		TestTrue(FString::Printf(TEXT("Texel factor 2x ratio is ~2.0 (actual: %.3f)"), Ratio), FMath::IsNearlyEqual(Ratio, 2.0f, 0.01f));
+
+		return true;
+	}
+
+	// -----------------------------------------------------------------------
+	// Register a component, compute screen size with a far view, then
+	// recompute with a near view. Verify screen size increases when the
+	// view moves closer. This tests the UpdateBoundSizes → query round-trip.
+	// -----------------------------------------------------------------------
+	IMPLEMENT_SIMPLE_AUTOMATION_TEST(FSSAMScreenSizeUpdatesOnMove, SSAM_TEST_NAME ".StreamingQuery.UpdatesOnMove", TestFlags)
+
+	bool FSSAMScreenSizeUpdatesOnMove::RunTest(const FString& Parameters)
+	{
+		FSSAMTestEnvironment Env;
+
+		USSAMTestStreamableAsset* Asset = Env.CreateTestAsset();
+		Env.SpawnTestComponent({ { Asset, 512.f } }, FBoxSphereBounds(FVector::ZeroVector, FVector(100.f), 100.f));
+		Env.ProcessSSAM();
+
+		const int32 AssetIndex = Asset->GetSSAMIndexValue();
+		TestTrue(TEXT("Asset registered"), AssetIndex != INDEX_NONE);
+
+		Env.UpdateBoundSizes(FVector(5000.f, 0.f, 0.f));
+		const float SizeFar = Env.QueryScreenSize(AssetIndex).MaxSize;
+		Env.UpdateBoundSizes(FVector(500.f, 0.f, 0.f));
+		const float SizeNear = Env.QueryScreenSize(AssetIndex).MaxSize;
+
+		TestTrue(TEXT("Far screen size is positive"), SizeFar > 0.f);
+		TestTrue(TEXT("Near screen size is positive"), SizeNear > 0.f);
+		TestTrue(TEXT("Closer view produces larger screen size"), SizeNear > SizeFar);
+
+		return true;
+	}
+
+	// -----------------------------------------------------------------------
+	// Register a component with bForceMipStreaming = true. Verify that
+	// GetRenderAssetScreenSize returns MaxSize = FLT_MAX regardless of
+	// distance. This exercises the force-load code path in ProcessElement
+	// (line 487: if bForceLoad && NormalizedSize > 0 → MaxSize = FLT_MAX).
+	// -----------------------------------------------------------------------
+	IMPLEMENT_SIMPLE_AUTOMATION_TEST(FSSAMForceMipStreaming, SSAM_TEST_NAME ".StreamingQuery.ForceMipStreaming", TestFlags)
+
+	bool FSSAMForceMipStreaming::RunTest(const FString& Parameters)
+	{
+		FSSAMTestEnvironment Env;
+
+		USSAMTestStreamableAsset* Asset = Env.CreateTestAsset();
+		const FBoxSphereBounds Bounds(FVector::ZeroVector, FVector(100.f), 100.f);
+		USSAMTestComponent* Comp = Env.SpawnTestComponent({ { Asset, 512.f } }, Bounds);
+
+		// Set bForceMipStreaming and re-register so the proxy picks up the flag
+		Comp->bForceMipStreaming = true;
+		Comp->MarkRenderStateDirty();
+		Comp->RecreateRenderState_Concurrent();
+		Env.ProcessSSAM();
+
+		const int32 AssetIndex = Asset->GetSSAMIndexValue();
+		TestTrue(TEXT("Asset registered"), AssetIndex != INDEX_NONE);
+
+		// Place view very far away — force mip should still return FLT_MAX
+		Env.UpdateBoundSizes(FVector(100000.f, 0.f, 0.f));
+		FSSAMTestEnvironment::FScreenSizeResult Result = Env.QueryScreenSize(AssetIndex);
+
+		TestEqual(TEXT("ForceMipStreaming produces FLT_MAX screen size"), Result.MaxSize, FLT_MAX);
+
+		return true;
+	}
+
+	// -----------------------------------------------------------------------
+	// When GSimpleStreamableAssetManagerConsiderVisibility is enabled,
+	// MaxNormalizedSize_VisibleOnly is only non-zero for objects with
+	// a recent LastRenderTime (> LastUpdateTime). Verify that:
+	// - An object that was never rendered has MaxSize_VisibleOnly = 0
+	// - MaxSize (distance-based, not gated by visibility) is still positive
+	//
+	// The visibility check at UpdateBoundSizes line 428-431:
+	//   if (bConsiderVisibility)
+	//     ScreenSizeOverDistance = select(LastRenderTime > LastUpdateTime, ..., 0)
+	// Our object's LastRenderTime defaults to -1000 (from FUpdate), which
+	// is NOT greater than LastUpdateTime=0.
+	// -----------------------------------------------------------------------
+	IMPLEMENT_SIMPLE_AUTOMATION_TEST(FSSAMVisibilityBasedStreaming, SSAM_TEST_NAME ".StreamingQuery.VisibilityBased", TestFlags)
+
+	bool FSSAMVisibilityBasedStreaming::RunTest(const FString& Parameters)
+	{
+		FSSAMTestEnvironment Env;
+
+		// Enable visibility consideration for this test
+		IConsoleVariable* VisibilityCVar = IConsoleManager::Get().FindConsoleVariable( TEXT("s.StreamableAssets.SimpleStreamableAssetManager.ConsiderVisibility"));
+		TestNotNull(TEXT("ConsiderVisibility CVar exists"), VisibilityCVar);
+		const int32 OriginalValue = VisibilityCVar->GetInt();
+		VisibilityCVar->Set(1, ECVF_SetByCode);
+
+		USSAMTestStreamableAsset* Asset = Env.CreateTestAsset();
+		Env.SpawnTestComponent({ { Asset, 512.f } }, FBoxSphereBounds(FVector::ZeroVector, FVector(100.f), 100.f));
+		Env.ProcessSSAM();
+
+		const int32 AssetIndex = Asset->GetSSAMIndexValue();
+		TestTrue(TEXT("Asset registered"), AssetIndex != INDEX_NONE);
+
+		Env.UpdateBoundSizes(FVector(500.f, 0.f, 0.f));
+		FSSAMTestEnvironment::FScreenSizeResult Result = Env.QueryScreenSize(AssetIndex);
+
+		TestTrue(TEXT("MaxSize is positive for non-visible object"), Result.MaxSize > 0.f);
+		TestEqual(TEXT("MaxSize_VisibleOnly is 0 for never-rendered object"), Result.MaxSizeVisibleOnly, 0.f);
+
+		// Restore CVar
+		VisibilityCVar->Set(OriginalValue, ECVF_SetByCode);
+		return true;
+	}
+
+	// -----------------------------------------------------------------------
+	// When the camera is INSIDE the object's bounding box, distance-to-box
+	// is 0 for all axes. ClampedDistSq = max(1, 0) = 1, giving
+	// NormalizedSize = ScreenSize / 1 = ScreenSize (near-maximum).
+	// This is the first-person rendering edge case — objects the player is
+	// standing inside must stream at highest quality.
+	// -----------------------------------------------------------------------
+	IMPLEMENT_SIMPLE_AUTOMATION_TEST(FSSAMViewInsideBounds, SSAM_TEST_NAME ".StreamingQuery.ViewInsideBounds", TestFlags)
+
+	bool FSSAMViewInsideBounds::RunTest(const FString& Parameters)
+	{
+		FSSAMTestEnvironment Env;
+
+		USSAMTestStreamableAsset* Asset = Env.CreateTestAsset();
+
+		const float TexelFactor = 512.f;
+		const FBoxSphereBounds Bounds(FVector::ZeroVector, FVector(500.f), 500.f);
+		Env.SpawnTestComponent({ { Asset, TexelFactor } }, Bounds);
+		Env.ProcessSSAM();
+
+		const int32 AssetIndex = Asset->GetSSAMIndexValue();
+		TestTrue(TEXT("Asset registered"), AssetIndex != INDEX_NONE);
+
+		// View at origin — inside the 500-unit bounding box
+		const float ScreenSizeParam = 1920.f;
+		Env.UpdateBoundSizes(FVector::ZeroVector, ScreenSizeParam);
+		FSSAMTestEnvironment::FScreenSizeResult Result = Env.QueryScreenSize(AssetIndex);
+
+		// Distance-to-box = 0 on all axes → ClampedDistSq = max(1, 0) = 1
+		// NormalizedSize = ScreenSize / sqrt(1) = 1920
+		// MaxSize = TexelFactor * 1.0 * 1920 = 983040
+		const float ExpectedMaxSize = TexelFactor * ScreenSizeParam;
+		TestTrue(FString::Printf(TEXT("View inside bounds: MaxSize (%.2f) matches expected (%.2f) within 1%%"), Result.MaxSize, ExpectedMaxSize),
+			FMath::IsNearlyEqual(Result.MaxSize, ExpectedMaxSize, ExpectedMaxSize * 0.01f));
+
+		return true;
+	}
+
+	// -----------------------------------------------------------------------
+	// Multiple simultaneous views (split-screen scenario). The screen size
+	// should be the MAX across all views. If view A is far and view B is
+	// near, the result should match view B's (larger) screen size.
+	// -----------------------------------------------------------------------
+	IMPLEMENT_SIMPLE_AUTOMATION_TEST(FSSAMMultipleViews, SSAM_TEST_NAME ".StreamingQuery.MultipleViews", TestFlags)
+
+	bool FSSAMMultipleViews::RunTest(const FString& Parameters)
+	{
+		FSSAMTestEnvironment Env;
+
+		USSAMTestStreamableAsset* Asset = Env.CreateTestAsset();
+		const float TexelFactor = 512.f;
+		Env.SpawnTestComponent({ { Asset, TexelFactor } }, FBoxSphereBounds(FVector::ZeroVector, FVector(100.f), 100.f));
+		Env.ProcessSSAM();
+
+		const int32 AssetIndex = Asset->GetSSAMIndexValue();
+		TestTrue(TEXT("Asset registered"), AssetIndex != INDEX_NONE);
+
+		// Single near view for reference
+		Env.UpdateBoundSizes(FVector(500.f, 0.f, 0.f));
+		const float SizeNearOnly = Env.QueryScreenSize(AssetIndex).MaxSize;
+
+		// Two views — one far (10000), one near (500)
+		const FVector MultiViewOrigins[] = { FVector(10000.f, 0.f, 0.f), FVector(500.f, 0.f, 0.f) };
+		Env.UpdateBoundSizes(MultiViewOrigins);
+		const float MaxSizeMulti = Env.QueryScreenSize(AssetIndex).MaxSize;
+
+		// Multi-view result should match the near view (which dominates)
+		TestTrue(FString::Printf(TEXT("Multi-view (%.2f) matches near-only (%.2f) within 1%%"), MaxSizeMulti, SizeNearOnly),
+			FMath::IsNearlyEqual(MaxSizeMulti, SizeNearOnly, SizeNearOnly * 0.01f));
+
+		return true;
+	}
+
+	// -----------------------------------------------------------------------
+	// MaxDrawDistance: when the view is beyond an object's max draw distance,
+	// the InRangeMask check in UpdateBoundSizes zeros out
+	// MaxNormalizedSize_VisibleOnly (the object is out of range).
+	// Verify that MaxSize_VisibleOnly = 0 when the view is out of range.
+	// Note: MaxNormalizedSize (non-visibility-gated) is NOT range-checked
+	// in the current implementation — it always uses raw distance.
+	// -----------------------------------------------------------------------
+	IMPLEMENT_SIMPLE_AUTOMATION_TEST(FSSAMMaxDrawDistance, SSAM_TEST_NAME ".StreamingQuery.MaxDrawDistance", TestFlags)
+
+	bool FSSAMMaxDrawDistance::RunTest(const FString& Parameters)
+	{
+		FSSAMTestEnvironment Env;
+
+		USSAMTestStreamableAsset* Asset = Env.CreateTestAsset();
+
+		// Set max draw distance = 1000 on the component
+		const FBoxSphereBounds Bounds(FVector::ZeroVector, FVector(100.f), 100.f);
+		USSAMTestComponent* Comp = Env.SpawnTestComponent({ { Asset, 512.f } }, Bounds);
+		Comp->CachedMaxDrawDistance = 1000.f;
+
+		// Re-register so the proxy picks up the draw distance
+		Comp->MarkRenderStateDirty();
+		Comp->RecreateRenderState_Concurrent();
+		Env.ProcessSSAM();
+
+		const int32 AssetIndex = Asset->GetSSAMIndexValue();
+		TestTrue(TEXT("Asset registered"), AssetIndex != INDEX_NONE);
+
+		// View at 500 — within max draw distance (1000)
+		Env.UpdateBoundSizes(FVector(500.f, 0.f, 0.f));
+		TestTrue(TEXT("Within range: MaxSize_VisibleOnly is positive"), Env.QueryScreenSize(AssetIndex).MaxSizeVisibleOnly > 0.f);
+
+		// View at 5000 — beyond max draw distance (1000)
+		Env.UpdateBoundSizes(FVector(5000.f, 0.f, 0.f));
+		TestEqual(TEXT("Beyond range: MaxSize_VisibleOnly is 0"), Env.QueryScreenSize(AssetIndex).MaxSizeVisibleOnly, 0.f);
+
+		return true;
+	}
+
+	// -----------------------------------------------------------------------
+	// StreamingScaleFactor: when CanApplyStreamableRenderAssetScaleFactor
+	// returns true, the render matrix's max axis scale is used as
+	// ComponentScale. A 2x scaled object should produce 2x screen size
+	// compared to an unscaled one (both at the same position with the
+	// same texel factor).
+	// -----------------------------------------------------------------------
+	IMPLEMENT_SIMPLE_AUTOMATION_TEST(FSSAMStreamingScaleFactor, SSAM_TEST_NAME ".StreamingQuery.ScaleFactor", TestFlags)
+
+	bool FSSAMStreamingScaleFactor::RunTest(const FString& Parameters)
+	{
+		FSSAMTestEnvironment Env;
+
+		USSAMTestStreamableAsset* AssetUnscaled = Env.CreateTestAsset();
+		USSAMTestStreamableAsset* AssetScaled = Env.CreateTestAsset();
+
+		const FBoxSphereBounds Bounds(FVector::ZeroVector, FVector(100.f), 100.f);
+
+		// Unscaled component — scale factor = 1.0
+		{
+			USSAMTestComponent* Comp = Env.SpawnTestComponent({ { AssetUnscaled, 512.f } }, Bounds);
+			Comp->bSupportScaleFactor = true;
+			Comp->SetWorldScale3D(FVector(1.f));
+			Comp->MarkRenderStateDirty();
+			Comp->RecreateRenderState_Concurrent();
+		}
+
+		// 2x scaled component — scale factor = 2.0
+		{
+			USSAMTestComponent* Comp = Env.SpawnTestComponent({ { AssetScaled, 512.f } }, Bounds);
+			Comp->bSupportScaleFactor = true;
+			Comp->SetWorldScale3D(FVector(2.f));
+			Comp->MarkRenderStateDirty();
+			Comp->RecreateRenderState_Concurrent();
+		}
+
+		Env.ProcessSSAM();
+
+		TestTrue(TEXT("Unscaled asset registered"), AssetUnscaled->GetSSAMIndexValue() != INDEX_NONE);
+		TestTrue(TEXT("Scaled asset registered"), AssetScaled->GetSSAMIndexValue() != INDEX_NONE);
+
+		// Same view for both
+		Env.UpdateBoundSizes(FVector(1000.f, 0.f, 0.f));
+
+		const float SizeUnscaled = Env.QueryScreenSize(AssetUnscaled->GetSSAMIndexValue()).MaxSize;
+		const float SizeScaled = Env.QueryScreenSize(AssetScaled->GetSSAMIndexValue()).MaxSize;
+
+		TestTrue(TEXT("Unscaled screen size is positive"), SizeUnscaled > 0.f);
+		TestTrue(TEXT("Scaled screen size is positive"), SizeScaled > 0.f);
+
+		// 2x scale increases ComponentScale (2x) AND reduces distance-to-box
+		// (larger extents), so the ratio should be > 2.0
+		const float Ratio = SizeScaled / SizeUnscaled;
+		TestTrue(FString::Printf(TEXT("2x scale produces ratio > 2.0 (actual: %.3f)"), Ratio), Ratio > 2.0f);
+
+		return true;
+	}
+
+	// -----------------------------------------------------------------------
+	// BoostFactor: FStreamingViewInfo::BoostFactor affects
+	// ExtraBoostForVisiblePrimitive in ViewInfoExtras. Verify that a 2x
+	// boost doubles MaxSize_VisibleOnly compared to 1x boost.
+	// Note: BoostFactor affects the VisibleOnly path only (multiplied by
+	// ExtraBoostForVisiblePrimitive at line 425).
+	// -----------------------------------------------------------------------
+	IMPLEMENT_SIMPLE_AUTOMATION_TEST(FSSAMBoostFactor, SSAM_TEST_NAME ".StreamingQuery.BoostFactor", TestFlags)
+
+	bool FSSAMBoostFactor::RunTest(const FString& Parameters)
+	{
+		FSSAMTestEnvironment Env;
+
+		USSAMTestStreamableAsset* Asset = Env.CreateTestAsset();
+		Env.SpawnTestComponent({ { Asset, 512.f } }, FBoxSphereBounds(FVector::ZeroVector, FVector(100.f), 100.f));
+		Env.ProcessSSAM();
+
+		const int32 AssetIndex = Asset->GetSSAMIndexValue();
+		TestTrue(TEXT("Asset registered"), AssetIndex != INDEX_NONE);
+
+		Env.UpdateBoundSizes(FVector(1000.f, 0.f, 0.f), 1920.f, 1.0f);
+		const float VisibleOnly_1x = Env.QueryScreenSize(AssetIndex).MaxSizeVisibleOnly;
+
+		Env.UpdateBoundSizes(FVector(1000.f, 0.f, 0.f), 1920.f, 2.0f);
+		const float VisibleOnly_2x = Env.QueryScreenSize(AssetIndex).MaxSizeVisibleOnly;
+
+		TestTrue(TEXT("VisibleOnly with 1x boost is positive"), VisibleOnly_1x > 0.f);
+		TestTrue(TEXT("VisibleOnly with 2x boost is positive"), VisibleOnly_2x > 0.f);
+
+		const float Ratio = VisibleOnly_2x / VisibleOnly_1x;
+		TestTrue(FString::Printf(TEXT("2x boost ratio is ~2.0 (actual: %.3f)"), Ratio), FMath::IsNearlyEqual(Ratio, 2.0f, 0.1f));
+
+		return true;
+	}
+
+	// -----------------------------------------------------------------------
+	// Object movement via FUpdate path: register an object, then move the
+	// component (triggering a transform update → FUpdate record). Verify
+	// that GetAssetReferenceBounds reflects the new position and screen
+	// size changes accordingly.
+	// In production, UpdatePrimitiveTransform_RenderThread pushes FUpdate
+	// records when a component's transform changes.
+	// -----------------------------------------------------------------------
+	IMPLEMENT_SIMPLE_AUTOMATION_TEST(FSSAMObjectMovementUpdate, SSAM_TEST_NAME ".StreamingQuery.ObjectMovement", TestFlags)
+
+	bool FSSAMObjectMovementUpdate::RunTest(const FString& Parameters)
+	{
+		FSSAMTestEnvironment Env;
+
+		USSAMTestStreamableAsset* Asset = Env.CreateTestAsset();
+
+		// Object starts at origin
+		USSAMTestComponent* Comp = Env.SpawnTestComponent({ { Asset, 512.f } }, FBoxSphereBounds(FVector::ZeroVector, FVector(100.f), 100.f));
+		Env.ProcessSSAM();
+
+		const int32 AssetIndex = Asset->GetSSAMIndexValue();
+		TestTrue(TEXT("Asset registered"), AssetIndex != INDEX_NONE);
+
+		// View at (1000, 0, 0)
+		Env.UpdateBoundSizes(FVector(1000.f, 0.f, 0.f));
+		const float SizeAtOrigin = Env.QueryScreenSize(AssetIndex).MaxSize;
+		TestTrue(TEXT("Screen size at origin is positive"), SizeAtOrigin > 0.f);
+
+		// Move the component closer to the view (to 800, 0, 0).
+		// SetWorldLocation updates ComponentToWorld. SendRenderTransform_Concurrent
+		// then recalculates bounds (CalcBounds with new transform) and calls
+		// Scene->UpdatePrimitiveTransform, which ENQUEUES a render command
+		// (RendererScene.cpp line 1687). The RT command pushes an FUpdate record
+		// to the SSAM MPMC queue (line 1702). FlushRenderingCommands is REQUIRED
+		// to ensure the RT executes that command before ProcessSSAM drains the queue.
+		// Without the flush, the FUpdate would never reach SSAM.
+		//
+		// Note: SendRenderTransform_Concurrent is called directly here (not through
+		// the engine's end-of-frame update path). This bypasses the
+		// FPrimitiveTransformUpdater TLS accumulation (RendererScene.cpp line 1675)
+		// which is only active during SendAllEndOfFrameUpdates.
+		Comp->SetWorldLocation(FVector(800.f, 0.f, 0.f));
+		Comp->SendRenderTransform_Concurrent();
+		FlushRenderingCommands();
+		Env.ProcessSSAM();
+
+		Env.UpdateBoundSizes(FVector(1000.f, 0.f, 0.f));
+		const float SizeAfterMove = Env.QueryScreenSize(AssetIndex).MaxSize;
+		TestTrue(TEXT("Screen size after move is positive"), SizeAfterMove > 0.f);
+
+		// Object moved from origin to (800,0,0), view is at (1000,0,0)
+		// Distance decreased from 900 (1000-100) to 100 (200-100)
+		// Screen size should be much larger
+		TestTrue(FString::Printf(TEXT("Screen size increased after moving closer: %.2f > %.2f"), SizeAfterMove, SizeAtOrigin), SizeAfterMove > SizeAtOrigin);
+
+		return true;
+	}
+
+	// ===================================================================
+	// GROUP 5: Lifecycle / Stress
+	//
+	// Tests that exercise SSAM's internal sparse array management under
+	// repeated registration/unregistration cycles. Verify index recycling,
+	// no memory leaks, and correct behavior under bulk operations.
+	// ===================================================================
+
+	// -----------------------------------------------------------------------
+	// Call Process() when there are no pending records. Verify it's a safe
+	// no-op that doesn't corrupt internal state. Tests that a subsequent
+	// registration still works correctly after empty processing cycles.
+	// -----------------------------------------------------------------------
+	IMPLEMENT_SIMPLE_AUTOMATION_TEST(FSSAMEmptyProcess, SSAM_TEST_NAME ".Lifecycle.EmptyProcess", TestFlags)
+
+	bool FSSAMEmptyProcess::RunTest(const FString& Parameters)
+	{
+		FSSAMTestEnvironment Env;
+
+		// Call Process() multiple times with nothing pending
+		Env.ProcessSSAM();
+		Env.ProcessSSAM();
+		Env.ProcessSSAM();
+
+		// Now register something and verify it still works
+		USSAMTestStreamableAsset* Asset = Env.CreateTestAsset();
+		Env.SpawnTestComponent({ { Asset, 512.f } });
+		Env.ProcessSSAM();
+
+		TestTrue(TEXT("Registration works after empty Process() calls"), Asset->GetSSAMIndexValue() != INDEX_NONE);
+
+		TArray<FBox> Boxes;
+		FSimpleStreamableAssetManager::GetAssetReferenceBounds(Asset, Boxes);
+		TestEqual(TEXT("Asset has 1 bound after empty Process() calls"), Boxes.Num(), 1);
+
+		return true;
+	}
+
+	// -----------------------------------------------------------------------
+	// Register N components, unregister all, then register N new ones.
+	// Verify that freed indices are reused (allocated size doesn't grow
+	// significantly) and that the new components are correctly registered.
+	// -----------------------------------------------------------------------
+	IMPLEMENT_SIMPLE_AUTOMATION_TEST(FSSAMIndexRecyclingNoLeak, SSAM_TEST_NAME ".Lifecycle.IndexRecyclingNoLeak", TestFlags)
+
+	bool FSSAMIndexRecyclingNoLeak::RunTest(const FString& Parameters)
+	{
+		FSSAMTestEnvironment Env;
+
+		USSAMTestStreamableAsset* Asset = Env.CreateTestAsset();
+		constexpr int32 N = 10;
+
+		// Phase 1: Register N components
+		TArray<USSAMTestComponent*> Components;
+		for (int32 i = 0; i < N; ++i)
+		{
+			Components.Add(Env.SpawnTestComponent({ { Asset, 512.f } }, FBoxSphereBounds(FVector::ZeroVector, FVector(100.f), 100.f),
+				FVector(i * 500.f, 0.f, 0.f)));
+		}
+		Env.ProcessSSAM();
+		const uint32 SizeAfterFirstBatch = FSimpleStreamableAssetManager::GetAllocatedSize();
+
+		// Phase 2: Unregister all
+		for (USSAMTestComponent* Comp : Components)
+		{
+			Comp->DestroyComponent();
+		}
+		Components.Empty();
+		Env.ProcessSSAM();
+
+		// Phase 3: Register N new components (should reuse freed indices)
+		for (int32 i = 0; i < N; ++i)
+		{
+			Components.Add(Env.SpawnTestComponent({ { Asset, 512.f } }, FBoxSphereBounds(FVector::ZeroVector, FVector(100.f), 100.f),
+				FVector(i * 500.f, 0.f, 0.f)));
+		}
+		Env.ProcessSSAM();
+		const uint32 SizeAfterSecondBatch = FSimpleStreamableAssetManager::GetAllocatedSize();
+
+		TestTrue(TEXT("Allocated size did not grow significantly after index recycling"), SizeAfterSecondBatch <= SizeAfterFirstBatch * 2);
+
+		TArray<FBox> Boxes;
+		FSimpleStreamableAssetManager::GetAssetReferenceBounds(Asset, Boxes);
+		TestEqual(TEXT("Asset has N bounds after re-registration"), Boxes.Num(), N);
+
+		for (USSAMTestComponent* Comp : Components)
+		{
+			Comp->DestroyComponent();
+		}
+		Env.ProcessSSAM();
+		return true;
+	}
+
+	// -----------------------------------------------------------------------
+	// Rapidly create and destroy components in a loop (50 cycles). Verify
+	// no crash, no memory leak (allocated size stays bounded), and no
+	// lingering bounds after all cycles complete.
+	// -----------------------------------------------------------------------
+	IMPLEMENT_SIMPLE_AUTOMATION_TEST(FSSAMRapidCreateDestroy, SSAM_TEST_NAME ".Lifecycle.RapidCreateDestroy", TestFlags)
+
+	bool FSSAMRapidCreateDestroy::RunTest(const FString& Parameters)
+	{
+		FSSAMTestEnvironment Env;
+
+		USSAMTestStreamableAsset* Asset = Env.CreateTestAsset();
+
+		// Run one cycle first to establish a warm baseline (sparse arrays
+		// allocate on first use, so comparing against a cold start is misleading)
+		{
+			USSAMTestComponent* Comp = Env.SpawnTestComponent({ { Asset, 512.f } });
+			Env.ProcessSSAM();
+			Comp->DestroyComponent();
+			Env.ProcessSSAM();
+		}
+		const uint32 BaselineSize = FSimpleStreamableAssetManager::GetAllocatedSize();
+
+		for (int32 i = 0; i < 50; ++i)
+		{
+			USSAMTestComponent* Comp = Env.SpawnTestComponent({ { Asset, 512.f } });
+			Env.ProcessSSAM();
+			Comp->DestroyComponent();
+			Env.ProcessSSAM();
+		}
+
+		TestTrue(TEXT("Allocated size stays bounded after rapid create/destroy"), FSimpleStreamableAssetManager::GetAllocatedSize() <= BaselineSize * 2);
+
+		TArray<FBox> Boxes;
+		FSimpleStreamableAssetManager::GetAssetReferenceBounds(Asset, Boxes);
+		TestEqual(TEXT("No lingering bounds"), Boxes.Num(), 0);
+
+		return true;
+	}
+
+// DISABLED - Currently fails - Fix pending
+#ifdef TEST_SSAM_MEMORY_BLOAT_ON_LEVEL_TRANSITION
+	// -----------------------------------------------------------------------
+	// UE-363476: Simulates level transitions where each "level" introduces
+	// unique assets with many objects. Objects are unregistered between
+	// levels but assets stay registered (cached). Verifies that the inner
+	// sparse arrays in AssetIndexToBounds4Index release memory when their
+	// element count drops to zero, preventing unbounded growth across
+	// transitions.
+	// -----------------------------------------------------------------------
+	IMPLEMENT_SIMPLE_AUTOMATION_TEST(FSSAMMemoryBloatOnLevelTransition, SSAM_TEST_NAME ".Lifecycle.MemoryBloatOnLevelTransition", TestFlags)
+
+	bool FSSAMMemoryBloatOnLevelTransition::RunTest(const FString& Parameters)
+	{
+		FSSAMTestEnvironment Env;
+
+		// Use many levels with many objects and unique assets per level.
+		// If inner sparse arrays are not freed when empty, they retain
+		// peak capacity after all objects are destroyed, causing linear
+		// memory growth across levels.
+		constexpr int32 NumLevels = 10;
+		constexpr int32 AssetsPerLevel = 10;
+		constexpr int32 ObjectsPerLevel = 50;
+
+		// Pre-create all assets (they persist across "level transitions")
+		TArray<USSAMTestStreamableAsset*> AllAssets;
+		for (int32 i = 0; i < NumLevels * AssetsPerLevel; ++i)
+		{
+			AllAssets.Add(Env.CreateTestAsset());
+		}
+
+		// Lambda to simulate one level transition: spawn objects, process,
+		// destroy objects, process. Assets stay alive.
+		auto SimulateLevel = [&](int32 Level)
+		{
+			TArray<USSAMTestComponent*> LevelComponents;
+			for (int32 i = 0; i < ObjectsPerLevel; ++i)
+			{
+				const int32 AssetBase = Level * AssetsPerLevel;
+				LevelComponents.Add(Env.SpawnTestComponent( { { AllAssets[AssetBase + (i % AssetsPerLevel)], 256.f },
+					  { AllAssets[AssetBase + ((i + 1) % AssetsPerLevel)], 128.f } },
+					FBoxSphereBounds(FVector(i * 200.f, Level * 1000.f, 0.f), FVector(50.f), 50.f)));
+			}
+			Env.ProcessSSAM();
+
+			for (USSAMTestComponent* Comp : LevelComponents)
+			{
+				Comp->DestroyComponent();
+			}
+			Env.ProcessSSAM();
+		};
+
+		// Run the first level to warm up all fixed-overhead structures
+		// (ObjectBounds4, ObjectRegistrationIndexToAssetProperty, etc.)
+		SimulateLevel(0);
+		const uint32 SizeAfterLevel1 = FSimpleStreamableAssetManager::GetAllocatedSize();
+
+		// Run remaining levels — each introduces new unique assets
+		for (int32 Level = 1; Level < NumLevels; ++Level)
+		{
+			SimulateLevel(Level);
+		}
+		const uint32 SizeAfterAllLevels = FSimpleStreamableAssetManager::GetAllocatedSize();
+
+		// Expected growth: the outer TArray<TSimpleSparseArray> gains 90 new
+		// entries (9 levels × 10 assets). Each TSimpleSparseArray is ~24 bytes
+		// + outer TArray slack ≈ ~8KB total. If inner sparse arrays are not
+		// freed, each retains its peak capacity (TBitArray + TArray of
+		// FAssetBoundElement for ~10+ elements each), adding ~75KB.
+		//
+		// Threshold: 16KB allows for outer array growth + alignment but
+		// catches the ~75KB+ bloat from unfixed inner array retention.
+		const int64 MemoryGrowth = (int64)SizeAfterAllLevels - (int64)SizeAfterLevel1;
+
+		UE_LOG(LogTemp, Display, TEXT("UE-363476 test: SizeAfterLevel1=%u, SizeAfterAllLevels=%u, Growth=%lld bytes"),
+			SizeAfterLevel1, SizeAfterAllLevels, MemoryGrowth);
+
+		TestTrue(TEXT("Memory growth across level transitions stays bounded (UE-363476)"), MemoryGrowth < 16384);
+
+		// Verify all asset bounds are actually empty
+		for (USSAMTestStreamableAsset* Asset : AllAssets)
+		{
+			TArray<FBox> Boxes;
+			FSimpleStreamableAssetManager::GetAssetReferenceBounds(Asset, Boxes);
+			TestEqual(TEXT("Asset has 0 bounds after all objects destroyed"), Boxes.Num(), 0);
+		}
+
+		return true;
+	}
+#endif // DISABLED_TEST_SSAM_MEMORY_BLOAT_ON_LEVEL_TRANSITION
+
+	// -----------------------------------------------------------------------
+	// Register 100 components, each referencing 1-3 of 3 shared assets.
+	// Verify all components and assets get valid indices, correct bound
+	// counts, and clean state after bulk unregister. Tests the MPMC queue
+	// and sparse array capacity under load.
+	// -----------------------------------------------------------------------
+	IMPLEMENT_SIMPLE_AUTOMATION_TEST(FSSAMBulkRegisterUnregister, SSAM_TEST_NAME ".Lifecycle.BulkRegisterUnregister", TestFlags)
+
+	bool FSSAMBulkRegisterUnregister::RunTest(const FString& Parameters)
+	{
+		FSSAMTestEnvironment Env;
+
+		constexpr int32 NumAssets = 3;
+		constexpr int32 NumComponents = 100;
+
+		TArray<USSAMTestStreamableAsset*> Assets;
+		for (int32 i = 0; i < NumAssets; ++i)
+		{
+			Assets.Add(Env.CreateTestAsset());
+		}
+
+		// Register 100 components, each referencing (i % NumAssets + 1) assets
+		TArray<USSAMTestComponent*> Components;
+		for (int32 i = 0; i < NumComponents; ++i)
+		{
+			TArray<FSSAMTestAssetEntry> Entries;
+			for (int32 j = 0; j <= (i % NumAssets); ++j)
+			{
+				Entries.Add({ Assets[j], 256.f * (j + 1) });
+			}
+			Components.Add(Env.SpawnTestComponent(Entries, FBoxSphereBounds(FVector::ZeroVector, FVector(100.f), 100.f),
+				FVector(i * 200.f, 0.f, 0.f)));
+		}
+		Env.ProcessSSAM();
+
+		for (int32 i = 0; i < NumAssets; ++i)
+		{
+			TestTrue(FString::Printf(TEXT("Asset[%d] has valid index"), i), Assets[i]->GetSSAMIndexValue() != INDEX_NONE);
+		}
+
+		// Asset[0] is referenced by all 100 components
+		{
+			TArray<FBox> Boxes;
+			FSimpleStreamableAssetManager::GetAssetReferenceBounds(Assets[0], Boxes);
+			TestEqual(TEXT("Asset[0] referenced by all components"), Boxes.Num(), NumComponents);
+		}
+
+		// Bulk unregister
+		for (USSAMTestComponent* Comp : Components)
+		{
+			Comp->DestroyComponent();
+		}
+		Env.ProcessSSAM();
+
+		for (int32 i = 0; i < NumAssets; ++i)
+		{
+			TArray<FBox> Boxes;
+			FSimpleStreamableAssetManager::GetAssetReferenceBounds(Assets[i], Boxes);
+			TestEqual(FString::Printf(TEXT("Asset[%d] has 0 bounds after bulk unregister"), i), Boxes.Num(), 0);
+		}
+
+		return true;
+	}
+
+	// ===================================================================
+	// GROUP 6: Concurrency
+	//
+	// Tests concurrent access to SSAM's MPMC queues. Some tests use the
+	// FPrimitiveSceneDesc path to add/remove primitives from parallel GT
+	// workers (exercising the lockless queues under contention). Others
+	// run Process() on a background thread while the GT produces records
+	// (mirroring the cooked build's FDoWorkAsyncTask pattern).
+	// ===================================================================
+
+	// -----------------------------------------------------------------------
+	// 50 primitives are added/removed from parallel game thread workers
+	// via FPrimitiveSceneDesc + ParallelFor + EParallelGameThread tag.
+	// This exercises TLocklessGrowingStorage / TConsumeAllMpmcQueue under
+	// real thread contention from multiple concurrent producers.
+	// -----------------------------------------------------------------------
+	IMPLEMENT_SIMPLE_AUTOMATION_TEST(FSSAMConcurrentRegistration, SSAM_TEST_NAME ".Concurrency.ParallelRegistration", TestFlags)
+
+	bool FSSAMConcurrentRegistration::RunTest(const FString& Parameters)
+	{
+		FSSAMTestEnvironment Env;
+
+		constexpr int32 NumPrimitives = 50;
+		USSAMTestStreamableAsset* Asset = Env.CreateTestAsset();
+
+		UWorld* World = Env.GetWorld();
+		FSceneInterface* Scene = World->Scene;
+		TestNotNull(TEXT("World has a scene"), Scene);
+
+		// Pre-build all data structures on the game thread.
+		// Heap-allocated because the render thread holds references that may
+		// outlive the synchronous test scope (deferred cleanup).
+		struct FTestPrimitive
+		{
+			FSSAMTestProxyDesc ProxyDesc;
+			FPrimitiveSceneInfoData SceneData;
+			FPrimitiveSceneDesc SceneDesc;
+			FSSAMTestSceneProxy* Proxy = nullptr;
+		};
+
+		TArray<TUniquePtr<FTestPrimitive>> Primitives;
+		Primitives.Reserve(NumPrimitives);
+
+		for (int32 i = 0; i < NumPrimitives; ++i)
+		{
+			auto P = MakeUnique<FTestPrimitive>();
+			const FBoxSphereBounds Bounds(FVector(i * 200.f, 0.f, 0.f), FVector(100.f), 100.f);
+
+			TArray<FStreamingRenderAssetPrimitiveInfo> AssetInfos;
+			AssetInfos.Emplace(Asset, Bounds, 256.f, PackedRelativeBox_Identity, false, true);
+
+			P->ProxyDesc.Scene = Scene;
+			P->Proxy = new FSSAMTestSceneProxy(&P->ProxyDesc, MoveTemp(AssetInfos));
+
+			P->SceneData.SceneProxy = P->Proxy;
+			P->SceneDesc.SceneProxy = P->Proxy;
+			P->SceneDesc.ProxyDesc = &P->ProxyDesc;
+			P->SceneDesc.PrimitiveSceneData = &P->SceneData;
+			P->SceneDesc.Bounds = Bounds;
+			P->SceneDesc.LocalBounds = Bounds;
+			P->SceneDesc.RenderMatrix = FMatrix::Identity;
+			P->SceneDesc.AttachmentRootPosition = Bounds.Origin;
+			P->SceneDesc.Mobility = EComponentMobility::Movable;
+
+			Primitives.Add(MoveTemp(P));
+		}
+
+		// Concurrent registration from parallel game thread workers
+		ParallelFor(NumPrimitives, [&Primitives, Scene](int32 Index)
+		{
+			FTaskTagScope Scope(ETaskTag::EParallelGameThread);
+			Scene->AddPrimitive(&Primitives[Index]->SceneDesc);
+		});
+		FlushRenderingCommands();
+		Env.ProcessSSAM();
+
+		{
+			TArray<FBox> Boxes;
+			FSimpleStreamableAssetManager::GetAssetReferenceBounds(Asset, Boxes);
+			TestEqual(TEXT("Asset has N bounds from concurrent registration"), Boxes.Num(), NumPrimitives);
+		}
+
+		// Concurrent unregistration from parallel game thread workers
+		ParallelFor(NumPrimitives, [&Primitives, Scene](int32 Index)
+		{
+			FTaskTagScope Scope(ETaskTag::EParallelGameThread);
+			Scene->RemovePrimitive(&Primitives[Index]->SceneDesc);
+		});
+		FlushRenderingCommands();
+		Env.ProcessSSAM();
+
+		{
+			TArray<FBox> Boxes;
+			FSimpleStreamableAssetManager::GetAssetReferenceBounds(Asset, Boxes);
+			TestEqual(TEXT("No bounds after concurrent cleanup"), Boxes.Num(), 0);
+		}
+
+		// Do NOT delete proxies — the scene takes ownership after AddPrimitive
+		// and deletes them during RemovePrimitive's deferred cleanup or scene teardown.
+
+		return true;
+	}
+	// -----------------------------------------------------------------------
+	// Cooked build pattern: Process() loops continuously on a background
+	// worker thread while the game thread registers and unregisters
+	// components. This mirrors the real streaming system where
+	// FDoWorkAsyncTask runs Process() every frame on the thread pool
+	// while the game thread and parallel GT workers are simultaneously
+	// pushing records to the MPMC queues.
+	//
+	// The consumer loop runs continuously while producers register objects.
+	//
+	// Note: FPlatformProcess::Yield() does not guarantee thread interleaving.
+	// All registrations may accumulate and be consumed in the final
+	// ProcessSSAM() call rather than incrementally by the background loop.
+	// The test still has value as a stress test for the MPMC queue under
+	// concurrent access, and the race detector (in instrumented builds)
+	// will flag any data races regardless of scheduling order.
+	// -----------------------------------------------------------------------
+	IMPLEMENT_SIMPLE_AUTOMATION_TEST(FSSAMConcurrentProcessAndRegister, SSAM_TEST_NAME ".Concurrency.ConcurrentProcessAndRegister", TestFlags)
+
+	bool FSSAMConcurrentProcessAndRegister::RunTest(const FString& Parameters)
+	{
+		FSSAMTestEnvironment Env;
+
+		USSAMTestStreamableAsset* Asset = Env.CreateTestAsset();
+
+		std::atomic<bool> bStopProcessing{ false };
+
+		// Launch a background worker that loops Process() continuously,
+		// just like the cooked build's thread pool worker does every frame.
+		UE::Tasks::FTask ProcessTask = UE::Tasks::Launch(UE_SOURCE_LOCATION, [&Env, &bStopProcessing]()
+		{
+			while (!bStopProcessing.load(std::memory_order_relaxed))
+			{
+				Env.ProcessSSAM();
+				FPlatformProcess::Yield();
+			}
+		});
+
+		// Game thread: register components in bursts while the consumer is running
+		constexpr int32 NumBursts = 10;
+		constexpr int32 NumPerBurst = 5;
+		TArray<USSAMTestComponent*> AllComponents;
+		for (int32 Burst = 0; Burst < NumBursts; ++Burst)
+		{
+			for (int32 i = 0; i < NumPerBurst; ++i)
+			{
+				AllComponents.Add(Env.SpawnTestComponent({ { Asset, 512.f } }, FBoxSphereBounds(FVector::ZeroVector, FVector(100.f), 100.f),
+					FVector(float(Burst * NumPerBurst + i) * 200.f, 0.f, 0.f)));
+			}
+			// Small yield to give the consumer a chance to run between bursts
+			FPlatformProcess::Yield();
+		}
+
+		// Stop the consumer loop
+		bStopProcessing.store(true, std::memory_order_relaxed);
+		ProcessTask.Wait();
+
+		// Final Process to drain anything remaining
+		Env.ProcessSSAM();
+
+		TArray<FBox> Boxes;
+		FSimpleStreamableAssetManager::GetAssetReferenceBounds(Asset, Boxes);
+		TestEqual(TEXT("All components registered after concurrent process+register"), Boxes.Num(), NumBursts * NumPerBurst);
+
+		return true;
+	}
+
+	// -----------------------------------------------------------------------
+	// Cooked build pattern: Process() loops on a background worker while
+	// the game thread destroys components, pushing FUnregister records
+	// concurrently with the consumer draining queues.
+	//
+	// Same Yield() scheduling caveat as ConcurrentProcessAndRegister.
+	// -----------------------------------------------------------------------
+	IMPLEMENT_SIMPLE_AUTOMATION_TEST(FSSAMConcurrentProcessAndUnregister, SSAM_TEST_NAME ".Concurrency.ConcurrentProcessAndUnregister", TestFlags)
+
+	bool FSSAMConcurrentProcessAndUnregister::RunTest(const FString& Parameters)
+	{
+		FSSAMTestEnvironment Env;
+
+		USSAMTestStreamableAsset* Asset = Env.CreateTestAsset();
+
+		// Register a batch of components first
+		constexpr int32 NumComponents = 30;
+		TArray<USSAMTestComponent*> Components;
+		for (int32 i = 0; i < NumComponents; ++i)
+		{
+			Components.Add(Env.SpawnTestComponent({ { Asset, 512.f } }, FBoxSphereBounds(FVector::ZeroVector, FVector(100.f), 100.f),
+				FVector(i * 200.f, 0.f, 0.f)));
+		}
+		Env.ProcessSSAM();
+
+		{
+			TArray<FBox> Boxes;
+			FSimpleStreamableAssetManager::GetAssetReferenceBounds(Asset, Boxes);
+			TestEqual(TEXT("All registered before concurrent unregister"), Boxes.Num(), NumComponents);
+		}
+
+		std::atomic<bool> bStopProcessing{ false };
+
+		// Background consumer loop
+		UE::Tasks::FTask ProcessTask = UE::Tasks::Launch(UE_SOURCE_LOCATION, [&Env, &bStopProcessing]()
+		{
+			while (!bStopProcessing.load(std::memory_order_relaxed))
+			{
+				Env.ProcessSSAM();
+				FPlatformProcess::Yield();
+			}
+		});
+
+		// Game thread: destroy components one by one while consumer runs
+		for (USSAMTestComponent* Comp : Components)
+		{
+			Comp->DestroyComponent();
+			FPlatformProcess::Yield();
+		}
+
+		bStopProcessing.store(true, std::memory_order_relaxed);
+		ProcessTask.Wait();
+
+		Env.ProcessSSAM();
+
+		TArray<FBox> Boxes;
+		FSimpleStreamableAssetManager::GetAssetReferenceBounds(Asset, Boxes);
+		TestEqual(TEXT("All unregistered after concurrent process+unregister"), Boxes.Num(), 0);
+
+		return true;
+	}
+
+	// -----------------------------------------------------------------------
+	// Cooked build pattern: Process() loops on a background worker while
+	// components move (pushing FUpdate records via the render thread).
+	// The worker is consuming FUpdate records while the RT is producing them.
+	// -----------------------------------------------------------------------
+	IMPLEMENT_SIMPLE_AUTOMATION_TEST(FSSAMConcurrentProcessAndUpdate, SSAM_TEST_NAME ".Concurrency.ConcurrentProcessAndUpdate", TestFlags)
+
+	bool FSSAMConcurrentProcessAndUpdate::RunTest(const FString& Parameters)
+	{
+		FSSAMTestEnvironment Env;
+
+		USSAMTestStreamableAsset* Asset = Env.CreateTestAsset();
+
+		constexpr int32 NumComponents = 10;
+		TArray<USSAMTestComponent*> Components;
+		for (int32 i = 0; i < NumComponents; ++i)
+		{
+			Components.Add(Env.SpawnTestComponent({ { Asset, 512.f } }, FBoxSphereBounds(FVector::ZeroVector, FVector(100.f), 100.f),
+				FVector(i * 200.f, 0.f, 0.f)));
+		}
+		Env.ProcessSSAM();
+
+		std::atomic<bool> bStopProcessing{ false };
+
+		// Background consumer loop
+		UE::Tasks::FTask ProcessTask = UE::Tasks::Launch(UE_SOURCE_LOCATION, [&Env, &bStopProcessing]()
+		{
+			while (!bStopProcessing.load(std::memory_order_relaxed))
+			{
+				Env.ProcessSSAM();
+				FPlatformProcess::Yield();
+			}
+		});
+
+		// Game thread: move components repeatedly while consumer runs.
+		// SendRenderTransform_Concurrent enqueues a render command (RendererScene.cpp
+		// line 1687). The actual FUpdate push to SSAM happens on the RENDER THREAD
+		// when that command executes. FlushRenderingCommands between move batches
+		// ensures the RT processes the commands, pushing FUpdate records into the
+		// SSAM queue where the background Process() loop can consume them.
+		constexpr int32 NumMoveIterations = 20;
+		for (int32 Iter = 0; Iter < NumMoveIterations; ++Iter)
+		{
+			for (int32 i = 0; i < NumComponents; ++i)
+			{
+				Components[i]->SetWorldLocation(FVector(i * 200.f + Iter * 10.f, 0.f, 0.f));
+				Components[i]->SendRenderTransform_Concurrent();
+			}
+			// Flush RT commands so FUpdate records are actually pushed to SSAM
+			// queues, giving the background Process() loop a chance to consume them.
+			FlushRenderingCommands();
+		}
+
+		bStopProcessing.store(true, std::memory_order_relaxed);
+		ProcessTask.Wait();
+		FlushRenderingCommands();
+		Env.ProcessSSAM();
+
+		// All components should still be registered (moves don't unregister)
+		TArray<FBox> Boxes;
+		FSimpleStreamableAssetManager::GetAssetReferenceBounds(Asset, Boxes);
+		TestEqual(TEXT("All components still registered after concurrent moves"), Boxes.Num(), NumComponents);
+
+		return true;
+	}
+
+	// -----------------------------------------------------------------------
+	// Worst-case cooked build scenario: Process() loops on a background
+	// worker while three producer threads operate simultaneously:
+	// (a) Parallel GT workers registering new primitives via FPrimitiveSceneDesc
+	// (b) Game thread unregistering existing components
+	// (c) Game thread registering new components
+	// All producers push to SSAM's MPMC queues while the consumer drains them.
+	// -----------------------------------------------------------------------
+	IMPLEMENT_SIMPLE_AUTOMATION_TEST(FSSAMTripleProducerConcurrency, SSAM_TEST_NAME ".Concurrency.TripleProducerConcurrency", TestFlags)
+
+	bool FSSAMTripleProducerConcurrency::RunTest(const FString& Parameters)
+	{
+		FSSAMTestEnvironment Env;
+
+		USSAMTestStreamableAsset* Asset = Env.CreateTestAsset();
+
+		UWorld* World = Env.GetWorld();
+		FSceneInterface* Scene = World->Scene;
+		TestNotNull(TEXT("World has a scene"), Scene);
+
+		// Pre-register components that we'll unregister during the test
+		constexpr int32 NumToUnregister = 15;
+		TArray<USSAMTestComponent*> ComponentsToUnregister;
+		for (int32 i = 0; i < NumToUnregister; ++i)
+		{
+			ComponentsToUnregister.Add(Env.SpawnTestComponent({ { Asset, 512.f } }, FBoxSphereBounds(FVector::ZeroVector, FVector(100.f), 100.f),
+				FVector(i * 200.f, 0.f, 0.f)));
+		}
+		Env.ProcessSSAM();
+
+		// Prepare FPrimitiveSceneDesc primitives for parallel registration
+		constexpr int32 NumDescPrimitives = 15;
+		struct FTestPrimitive
+		{
+			FSSAMTestProxyDesc ProxyDesc;
+			FPrimitiveSceneInfoData SceneData;
+			FPrimitiveSceneDesc SceneDesc;
+		};
+
+		TArray<TUniquePtr<FTestPrimitive>> Primitives;
+		Primitives.Reserve(NumDescPrimitives);
+		for (int32 i = 0; i < NumDescPrimitives; ++i)
+		{
+			auto P = MakeUnique<FTestPrimitive>();
+			const FBoxSphereBounds Bounds(FVector((NumToUnregister + i) * 200.f, 0.f, 0.f), FVector(100.f), 100.f);
+
+			TArray<FStreamingRenderAssetPrimitiveInfo> AssetInfos;
+			AssetInfos.Emplace(Asset, Bounds, 256.f, PackedRelativeBox_Identity, false, true);
+
+			P->ProxyDesc.Scene = Scene;
+			FSSAMTestSceneProxy* Proxy = new FSSAMTestSceneProxy(&P->ProxyDesc, MoveTemp(AssetInfos));
+
+			P->SceneData.SceneProxy = Proxy;
+			P->SceneDesc.SceneProxy = Proxy;
+			P->SceneDesc.ProxyDesc = &P->ProxyDesc;
+			P->SceneDesc.PrimitiveSceneData = &P->SceneData;
+			P->SceneDesc.Bounds = Bounds;
+			P->SceneDesc.LocalBounds = Bounds;
+			P->SceneDesc.RenderMatrix = FMatrix::Identity;
+			P->SceneDesc.AttachmentRootPosition = Bounds.Origin;
+			P->SceneDesc.Mobility = EComponentMobility::Movable;
+
+			Primitives.Add(MoveTemp(P));
+		}
+
+		std::atomic<bool> bStopProcessing{ false };
+
+		// (a) Background consumer loop — Process() running continuously
+		UE::Tasks::FTask ProcessTask = UE::Tasks::Launch(UE_SOURCE_LOCATION, [&Env, &bStopProcessing]()
+		{
+			while (!bStopProcessing.load(std::memory_order_relaxed))
+			{
+				Env.ProcessSSAM();
+				FPlatformProcess::Yield();
+			}
+		});
+
+		// (b) Parallel GT workers registering new primitives via FPrimitiveSceneDesc
+		UE::Tasks::FTask RegisterTask = UE::Tasks::Launch(UE_SOURCE_LOCATION, [&Primitives, Scene, NumDescPrimitives]()
+		{
+			ParallelFor(NumDescPrimitives, [&Primitives, Scene](int32 Index)
+			{
+				FTaskTagScope Scope(ETaskTag::EParallelGameThread);
+				Scene->AddPrimitive(&Primitives[Index]->SceneDesc);
+			});
+		});
+
+		// (c) Game thread unregistering existing components — concurrent with (a) and (b)
+		for (USSAMTestComponent* Comp : ComponentsToUnregister)
+		{
+			Comp->DestroyComponent();
+			FPlatformProcess::Yield();
+		}
+
+		// Wait for parallel registration to finish
+		RegisterTask.Wait();
+
+		// Stop the consumer and drain
+		bStopProcessing.store(true, std::memory_order_relaxed);
+		ProcessTask.Wait();
+		FlushRenderingCommands();
+		Env.ProcessSSAM();
+
+		// Verify: unregistered components are gone, desc-based primitives are registered
+		{
+			TArray<FBox> Boxes;
+			FSimpleStreamableAssetManager::GetAssetReferenceBounds(Asset, Boxes);
+			TestEqual(TEXT("Only desc-based primitives remain after triple-producer"), Boxes.Num(), NumDescPrimitives);
+		}
+
+		// Cleanup desc-based primitives
+		ParallelFor(NumDescPrimitives, [&Primitives, Scene](int32 Index)
+		{
+			FTaskTagScope Scope(ETaskTag::EParallelGameThread);
+			Scene->RemovePrimitive(&Primitives[Index]->SceneDesc);
+		});
+		FlushRenderingCommands();
+		Env.ProcessSSAM();
+
+		{
+			TArray<FBox> Boxes;
+			FSimpleStreamableAssetManager::GetAssetReferenceBounds(Asset, Boxes);
+			TestEqual(TEXT("Clean state after triple-producer cleanup"), Boxes.Num(), 0);
+		}
+
+		return true;
+	}
+
+	// -----------------------------------------------------------------------
+	// Cooked build + GC scenario: Process() loops on a background worker
+	// while the game thread garbage-collects UStreamableRenderAsset objects.
+	// GC's BeginDestroy → UnlinkStreaming → UnregisterAsset pushes
+	// FRemovedAssetRecord to the MPMC queue while the background worker
+	// is draining it. Both paths acquire the same FCriticalSection.
+	// -----------------------------------------------------------------------
+	IMPLEMENT_SIMPLE_AUTOMATION_TEST(FSSAMGarbageCollectionDuringProcess, SSAM_TEST_NAME ".Concurrency.GarbageCollectionDuringProcess", TestFlags)
+
+	bool FSSAMGarbageCollectionDuringProcess::RunTest(const FString& Parameters)
+	{
+		FSSAMTestEnvironment Env;
+
+		constexpr int32 NumAssets = 10;
+		constexpr int32 NumComponentsPerAsset = 5;
+
+		TArray<USSAMTestStreamableAsset*> Assets;
+		for (int32 i = 0; i < NumAssets; ++i)
+		{
+			Assets.Add(Env.CreateTestAsset());
+		}
+
+		// Register components referencing each asset
+		TArray<USSAMTestComponent*> Components;
+		for (int32 i = 0; i < NumAssets; ++i)
+		{
+			for (int32 j = 0; j < NumComponentsPerAsset; ++j)
+			{
+				Components.Add(Env.SpawnTestComponent({ { Assets[i], 256.f } }, FBoxSphereBounds(FVector(i * 1000.f + j * 200.f, 0.f, 0.f), FVector(100.f), 100.f)));
+			}
+		}
+		Env.ProcessSSAM();
+
+		// Save shared pointers to each asset's SSAM index. These survive GC
+		// because TSharedPtr ownership is independent of the UObject lifetime.
+		// After Process() consumes the FRemovedAssetRecord, it sets the
+		// pointed-to int32 to INDEX_NONE — which we can observe through our copy.
+		TArray<TSharedPtr<int32, ESPMode::ThreadSafe>> SavedIndices;
+		for (USSAMTestStreamableAsset* Asset : Assets)
+		{
+			SavedIndices.Add(Asset->GetSSAMIndex());
+			TestTrue(TEXT("Asset has valid SSAM index before GC"), SavedIndices.Last().IsValid() && *SavedIndices.Last() != INDEX_NONE);
+		}
+
+		// Unregister all objects (but keep assets registered in SSAM)
+		for (USSAMTestComponent* Comp : Components)
+		{
+			Comp->DestroyComponent();
+		}
+		Env.ProcessSSAM();
+
+		// Start background Process() loop — mirrors the cooked build async worker
+		std::atomic<bool> bStopProcessing{ false };
+		UE::Tasks::FTask ProcessTask = UE::Tasks::Launch(UE_SOURCE_LOCATION, [&Env, &bStopProcessing]()
+		{
+			while (!bStopProcessing.load(std::memory_order_relaxed))
+			{
+				Env.ProcessSSAM();
+				FPlatformProcess::Yield();
+			}
+		});
+
+		// Release assets for GC without setting StreamingIndex = INDEX_NONE.
+		// This preserves the real BeginDestroy → UnlinkStreaming → UnregisterAsset path.
+		Env.ReleaseAssetsForGC(Assets);
+
+		// Force garbage collection — assets are now unreferenced and will be collected.
+		// BeginDestroy → UnlinkStreaming → UnregisterAsset fires on the game thread
+		// WHILE the background Process() loop is running.
+		CollectGarbage(GARBAGE_COLLECTION_KEEPFLAGS, true);
+
+		// Stop background loop and drain any remaining records
+		bStopProcessing.store(true, std::memory_order_relaxed);
+		ProcessTask.Wait();
+		Env.ProcessSSAM();
+
+		// Verify: Process() consumed all FRemovedAssetRecords and set indices to INDEX_NONE
+		for (int32 i = 0; i < NumAssets; ++i)
+		{
+			TestTrue(FString::Printf(TEXT("Asset[%d] SSAM index is INDEX_NONE after GC"), i), *SavedIndices[i] == INDEX_NONE);
+		}
+
+		return true;
+	}
+
+	// -----------------------------------------------------------------------
+	// Deliberately racy test to validate that the race detector catches
+	// unsynchronized concurrent access. Two threads read/write a plain
+	// int32 without any synchronization. Expected to FAIL on instrumented
+	// builds. Remove or disable after validation.
+	// -----------------------------------------------------------------------
+	IMPLEMENT_SIMPLE_AUTOMATION_TEST(FSSAMRaceDetectorCanary, SSAM_TEST_NAME ".Concurrency.RaceDetectorCanary", TestFlags)
+
+	bool FSSAMRaceDetectorCanary::RunTest(const FString& Parameters)
+	{
+		FSSAMTestEnvironment Env;
+
+		volatile int32 SharedCounter = 0;
+
+		UE::Tasks::FTask WriterTask = UE::Tasks::Launch(UE_SOURCE_LOCATION, [&SharedCounter]()
+		{
+			for (int32 i = 0; i < 1000; ++i)
+			{
+				++SharedCounter;
+			}
+		});
+
+		// Read from the main thread while the writer is running
+		for (int32 i = 0; i < 1000; ++i)
+		{
+			if (SharedCounter > 500)
+			{
+				break;
+			}
+		}
+
+		WriterTask.Wait();
+		TestTrue(TEXT("Counter was incremented"), SharedCounter > 0);
+		return true;
+	}
+
+	// -----------------------------------------------------------------------
+	// 30 primitives added/removed via FPrimitiveSceneDesc from parallel
+	// game thread workers, repeated 5 times. Maximizes contention on
+	// SSAM's MPMC queues through repeated concurrent cycles.
+	// -----------------------------------------------------------------------
+	IMPLEMENT_SIMPLE_AUTOMATION_TEST(FSSAMConcurrentOperationsStress, SSAM_TEST_NAME ".Concurrency.ConcurrentOperationsStress", TestFlags)
+
+	bool FSSAMConcurrentOperationsStress::RunTest(const FString& Parameters)
+	{
+		FSSAMTestEnvironment Env;
+
+		constexpr int32 NumPrimitives = 30;
+		USSAMTestStreamableAsset* Asset = Env.CreateTestAsset();
+
+		UWorld* World = Env.GetWorld();
+		FSceneInterface* Scene = World->Scene;
+		TestNotNull(TEXT("World has a scene"), Scene);
+
+		struct FTestPrimitive
+		{
+			FSSAMTestProxyDesc ProxyDesc;
+			FPrimitiveSceneInfoData SceneData;
+			FPrimitiveSceneDesc SceneDesc;
+			FSSAMTestSceneProxy* Proxy = nullptr;
+		};
+
+		TArray<TUniquePtr<FTestPrimitive>> Primitives;
+		Primitives.Reserve(NumPrimitives);
+
+		for (int32 i = 0; i < NumPrimitives; ++i)
+		{
+			auto P = MakeUnique<FTestPrimitive>();
+			const FBoxSphereBounds Bounds(FVector(i * 200.f, 0.f, 0.f), FVector(100.f), 100.f);
+
+			TArray<FStreamingRenderAssetPrimitiveInfo> AssetInfos;
+			AssetInfos.Emplace(Asset, Bounds, 256.f, PackedRelativeBox_Identity, false, true);
+
+			P->ProxyDesc.Scene = Scene;
+			P->Proxy = new FSSAMTestSceneProxy(&P->ProxyDesc, MoveTemp(AssetInfos));
+
+			P->SceneData.SceneProxy = P->Proxy;
+			P->SceneDesc.SceneProxy = P->Proxy;
+			P->SceneDesc.ProxyDesc = &P->ProxyDesc;
+			P->SceneDesc.PrimitiveSceneData = &P->SceneData;
+			P->SceneDesc.Bounds = Bounds;
+			P->SceneDesc.LocalBounds = Bounds;
+			P->SceneDesc.RenderMatrix = FMatrix::Identity;
+			P->SceneDesc.AttachmentRootPosition = Bounds.Origin;
+			P->SceneDesc.Mobility = EComponentMobility::Movable;
+
+			Primitives.Add(MoveTemp(P));
+		}
+
+		// Run multiple iterations to increase chances of catching races
+		constexpr int32 NumIterations = 5;
+		for (int32 Iter = 0; Iter < NumIterations; ++Iter)
+		{
+			// Concurrent registration
+			ParallelFor(NumPrimitives, [&Primitives, Scene](int32 Index)
+			{
+				FTaskTagScope Scope(ETaskTag::EParallelGameThread);
+				Scene->AddPrimitive(&Primitives[Index]->SceneDesc);
+			});
+			FlushRenderingCommands();
+			Env.ProcessSSAM();
+
+			// Verify registration worked
+			{
+				TArray<FBox> Boxes;
+				FSimpleStreamableAssetManager::GetAssetReferenceBounds(Asset, Boxes);
+				TestEqual(FString::Printf(TEXT("Iter %d: Asset has N bounds after registration"), Iter), Boxes.Num(), NumPrimitives);
+			}
+
+			// Concurrent unregistration
+			ParallelFor(NumPrimitives, [&Primitives, Scene](int32 Index)
+			{
+				FTaskTagScope Scope(ETaskTag::EParallelGameThread);
+				Scene->RemovePrimitive(&Primitives[Index]->SceneDesc);
+			});
+			FlushRenderingCommands();
+			Env.ProcessSSAM();
+
+			// Verify clean state
+			{
+				TArray<FBox> Boxes;
+				FSimpleStreamableAssetManager::GetAssetReferenceBounds(Asset, Boxes);
+				TestEqual(FString::Printf(TEXT("Iter %d: No bounds after unregistration"), Iter), Boxes.Num(), 0);
+			}
+
+			// Re-create proxies for next iteration (scene deleted them on remove)
+			if (Iter < NumIterations - 1)
+			{
+				for (int32 i = 0; i < NumPrimitives; ++i)
+				{
+					FTestPrimitive& P = *Primitives[i];
+					const FBoxSphereBounds Bounds(FVector(i * 200.f, 0.f, 0.f), FVector(100.f), 100.f);
+
+					TArray<FStreamingRenderAssetPrimitiveInfo> AssetInfos;
+					AssetInfos.Emplace(Asset, Bounds, 256.f, PackedRelativeBox_Identity, false, true);
+
+					P.Proxy = new FSSAMTestSceneProxy(&P.ProxyDesc, MoveTemp(AssetInfos));
+					P.SceneData.SceneProxy = P.Proxy;
+					P.SceneDesc.SceneProxy = P.Proxy;
+				}
+			}
+		}
+
+		return true;
+	}
+}
+
+#undef SSAM_TEST_NAME
+
+#endif // WITH_DEV_AUTOMATION_TESTS
diff --git a/Engine/Source/Runtime/Engine/Private/Tests/Streaming/SSAMTestTypes.cpp b/Engine/Source/Runtime/Engine/Private/Tests/Streaming/SSAMTestTypes.cpp
new file mode 100644
--- /dev/null
+++ b/Engine/Source/Runtime/Engine/Private/Tests/Streaming/SSAMTestTypes.cpp
@@ -0,0 +1,336 @@
+// Copyright Epic Games, Inc. All Rights Reserved.
+
+#include "SSAMTestTypes.h"
+
+#include UE_INLINE_GENERATED_CPP_BY_NAME(SSAMTestTypes)
+
+#if WITH_DEV_AUTOMATION_TESTS
+
+#include "ContentStreaming.h"
+#include "Streaming/TextureStreamingHelpers.h"
+#include "Streaming/TextureInstanceView.h"
+
+// ---------------------------------------------------------------------------
+// USSAMTestStreamableAsset
+// ---------------------------------------------------------------------------
+
+void USSAMTestStreamableAsset::BeginDestroy()
+{
+	if (SimpleStreamableAssetManagerIndex.IsValid() && *SimpleStreamableAssetManagerIndex != INDEX_NONE)
+	{
+		FSimpleStreamableAssetManager::UnregisterAsset(this);
+	}
+	// Prevent UnlinkStreaming from calling into the real FRenderAssetStreamingManager
+	// (our test assets were never registered there via LinkStreaming).
+	StreamingIndex = INDEX_NONE;
+	Super::BeginDestroy();
+}
+
+// ---------------------------------------------------------------------------
+// USSAMTestComponent
+// ---------------------------------------------------------------------------
+
+FPrimitiveSceneProxy* USSAMTestComponent::CreateSceneProxy()
+{
+	if (bUseFallbackProxy)
+	{
+		return new FSSAMTestFallbackSceneProxy(this);
+	}
+
+	TArray<FStreamingRenderAssetPrimitiveInfo> AssetInfos;
+	AssetInfos.Reserve(TestAssets.Num());
+
+	for (const FAssetEntry& Entry : TestAssets)
+	{
+		if (Entry.Asset)
+		{
+			FStreamingRenderAssetPrimitiveInfo Info(
+				Entry.Asset,
+				TestBounds,
+				Entry.TexelFactor,
+				PackedRelativeBox_Identity,
+				false,	// bAllowInvalidTexelFactorWhenUnregistered
+				true	// bAffectedByComponentScale
+			);
+			AssetInfos.Add(MoveTemp(Info));
+		}
+	}
+
+	return new FSSAMTestSceneProxy(this, MoveTemp(AssetInfos), bSupportScaleFactor);
+}
+
+void USSAMTestComponent::GetStreamingRenderAssetInfo(FStreamingTextureLevelContext& LevelContext, TArray<FStreamingRenderAssetPrimitiveInfo>& OutStreamingRenderAssets) const
+{
+	if (bUseFallbackProxy)
+	{
+		for (const FAssetEntry& Entry : TestAssets)
+		{
+			if (Entry.Asset)
+			{
+				OutStreamingRenderAssets.Emplace(
+					Entry.Asset,
+					GetBounds(),
+					Entry.TexelFactor,
+					PackedRelativeBox_Identity,
+					false,	// bAllowInvalidTexelFactorWhenUnregistered
+					true	// bAffectedByComponentScale
+				);
+			}
+		}
+		return;
+	}
+
+	Super::GetStreamingRenderAssetInfo(LevelContext, OutStreamingRenderAssets);
+}
+
+// ---------------------------------------------------------------------------
+// FSSAMTestProxyDesc
+// ---------------------------------------------------------------------------
+
+FSSAMTestProxyDesc::FSSAMTestProxyDesc()
+{
+	CustomPrimitiveData = &TestCustomPrimitiveData;
+	Mobility = EComponentMobility::Movable;
+	bIsVisible = true;
+}
+
+// ---------------------------------------------------------------------------
+// FSSAMTestSceneProxy
+// ---------------------------------------------------------------------------
+
+SIZE_T FSSAMTestSceneProxy::GetTypeHash() const
+{
+	static size_t UniquePointer;
+	return reinterpret_cast<size_t>(&UniquePointer);
+}
+
+FSSAMTestSceneProxy::FSSAMTestSceneProxy(const UPrimitiveComponent* InComponent, TArray<FStreamingRenderAssetPrimitiveInfo>&& InAssets, bool bInSupportScaleFactor)
+	: FPrimitiveSceneProxy(InComponent)
+	, CapturedAssets(MoveTemp(InAssets))
+	, bSupportScaleFactor(bInSupportScaleFactor)
+{
+	bImplementsStreamableAssetGathering = true;
+}
+
+FSSAMTestSceneProxy::FSSAMTestSceneProxy(const FPrimitiveSceneProxyDesc* InDesc, TArray<FStreamingRenderAssetPrimitiveInfo>&& InAssets, bool bInSupportScaleFactor)
+	: FPrimitiveSceneProxy(*InDesc)
+	, CapturedAssets(MoveTemp(InAssets))
+	, bSupportScaleFactor(bInSupportScaleFactor)
+{
+	bImplementsStreamableAssetGathering = true;
+}
+
+bool FSSAMTestSceneProxy::CanApplyStreamableRenderAssetScaleFactor() const
+{
+	return bSupportScaleFactor;
+}
+
+void FSSAMTestSceneProxy::GetStreamableRenderAssetInfo(const FBoxSphereBounds& PrimitiveBounds, TArray<FStreamingRenderAssetPrimitiveInfo>& OutStreamableRenderAssets) const
+{
+	OutStreamableRenderAssets = CapturedAssets;
+}
+
+FPrimitiveViewRelevance FSSAMTestSceneProxy::GetViewRelevance(const FSceneView* View) const
+{
+	FPrimitiveViewRelevance Relevance;
+	Relevance.bDrawRelevance = false;
+	return Relevance;
+}
+
+// ---------------------------------------------------------------------------
+// FSSAMTestFallbackSceneProxy
+// ---------------------------------------------------------------------------
+
+SIZE_T FSSAMTestFallbackSceneProxy::GetTypeHash() const
+{
+	static size_t UniquePointer;
+	return reinterpret_cast<size_t>(&UniquePointer);
+}
+
+FSSAMTestFallbackSceneProxy::FSSAMTestFallbackSceneProxy(const UPrimitiveComponent* InComponent)
+	: FPrimitiveSceneProxy(InComponent)
+{
+}
+
+FPrimitiveViewRelevance FSSAMTestFallbackSceneProxy::GetViewRelevance(const FSceneView* View) const
+{
+	FPrimitiveViewRelevance Relevance;
+	Relevance.bDrawRelevance = false;
+	return Relevance;
+}
+
+// ---------------------------------------------------------------------------
+// FSSAMTestEnvironment
+// ---------------------------------------------------------------------------
+
+FSSAMTestEnvironment::FSSAMTestEnvironment()
+	: bDidInit(false)
+	, bWasEnabled(FSimpleStreamableAssetManager::IsEnabled())
+{
+	if (!bWasEnabled)
+	{
+		IConsoleVariable* CVar = IConsoleManager::Get().FindConsoleVariable(TEXT("s.StreamableAssets.UseSimpleStreamableAssetManager"));
+		if (CVar)
+		{
+			CVar->Set(1, ECVF_SetByCode);
+		}
+	}
+
+	if (!FSimpleStreamableAssetManager::GetCriticalSection())
+	{
+		FSimpleStreamableAssetManager::Init();
+		bDidInit = true;
+	}
+
+	WorldWrapper.CreateTestWorld(EWorldType::Game);
+	WorldWrapper.BeginPlayInTestWorld();
+}
+
+FSSAMTestEnvironment::~FSSAMTestEnvironment()
+{
+	CleanupAssets();
+	WorldWrapper.DestroyTestWorld(true);
+
+	if (bDidInit)
+	{
+		FSimpleStreamableAssetManager::Shutdown();
+	}
+
+	if (!bWasEnabled)
+	{
+		IConsoleVariable* CVar = IConsoleManager::Get().FindConsoleVariable(TEXT("s.StreamableAssets.UseSimpleStreamableAssetManager"));
+		if (CVar)
+		{
+			CVar->Set(0, ECVF_SetByCode);
+		}
+	}
+}
+
+USSAMTestStreamableAsset* FSSAMTestEnvironment::CreateTestAsset()
+{
+	USSAMTestStreamableAsset* Asset = NewObject<USSAMTestStreamableAsset>();
+	Asset->SetStreamingIndexForTest(0);
+	Asset->AddToRoot();
+	CreatedAssets.Add(Asset);
+	return Asset;
+}
+
+USSAMTestComponent* FSSAMTestEnvironment::SpawnTestComponent(const TArray<FSSAMTestAssetEntry>& Assets, const FBoxSphereBounds& Bounds, const FVector& Location)
+{
+	UWorld* World = GetWorld();
+	check(World);
+
+	FActorSpawnParameters SpawnParams;
+	SpawnParams.SpawnCollisionHandlingOverride = ESpawnActorCollisionHandlingMethod::AlwaysSpawn;
+	AActor* Actor = World->SpawnActor<AActor>(SpawnParams);
+	check(Actor);
+
+	USSAMTestComponent* Comp = NewObject<USSAMTestComponent>(Actor);
+	Comp->TestAssets = Assets;
+	Comp->TestBounds = Bounds;
+	if (Actor->GetRootComponent())
+	{
+		Comp->AttachToComponent(Actor->GetRootComponent(), FAttachmentTransformRules::KeepRelativeTransform);
+	}
+	Actor->SetActorLocation(Location);
+	Comp->RegisterComponent();
+
+	return Comp;
+}
+
+USSAMTestComponent* FSSAMTestEnvironment::SpawnTestFallbackComponent(const TArray<FSSAMTestAssetEntry>& Assets, const FBoxSphereBounds& Bounds, const FVector& Location)
+{
+	UWorld* World = GetWorld();
+	check(World);
+
+	FActorSpawnParameters SpawnParams;
+	SpawnParams.SpawnCollisionHandlingOverride = ESpawnActorCollisionHandlingMethod::AlwaysSpawn;
+	AActor* Actor = World->SpawnActor<AActor>(SpawnParams);
+	check(Actor);
+
+	USSAMTestComponent* Comp = NewObject<USSAMTestComponent>(Actor);
+	Comp->TestAssets = Assets;
+	Comp->TestBounds = Bounds;
+	Comp->bUseFallbackProxy = true;
+	if (Actor->GetRootComponent())
+	{
+		Comp->AttachToComponent(Actor->GetRootComponent(), FAttachmentTransformRules::KeepRelativeTransform);
+	}
+	Actor->SetActorLocation(Location);
+	Comp->RegisterComponent();
+
+	return Comp;
+}
+
+void FSSAMTestEnvironment::CleanupAssets()
+{
+	for (USSAMTestStreamableAsset* Asset : CreatedAssets)
+	{
+		if (Asset)
+		{
+			Asset->SetStreamingIndexForTest(INDEX_NONE);
+			Asset->RemoveFromRoot();
+		}
+	}
+	CreatedAssets.Empty();
+}
+
+void FSSAMTestEnvironment::ReleaseAssetsForGC(TArrayView<USSAMTestStreamableAsset*> Assets)
+{
+	for (USSAMTestStreamableAsset* Asset : Assets)
+	{
+		CreatedAssets.Remove(Asset);
+		Asset->RemoveFromRoot();
+	}
+}
+
+void FSSAMTestEnvironment::UpdateBoundSizes(const FVector& ViewOrigin, float ScreenSize, float BoostFactor, float LastUpdateTime)
+{
+	UpdateBoundSizes(MakeArrayView(&ViewOrigin, 1), ScreenSize, BoostFactor, LastUpdateTime);
+}
+
+void FSSAMTestEnvironment::UpdateBoundSizes(TArrayView<const FVector> ViewOrigins, float ScreenSize, float BoostFactor, float LastUpdateTime)
+{
+	TArray<FStreamingViewInfo> ViewInfos;
+	FStreamingViewInfoExtraArray ViewInfoExtras;
+	ViewInfos.Reserve(ViewOrigins.Num());
+	ViewInfoExtras.Reserve(ViewOrigins.Num());
+
+	for (const FVector& Origin : ViewOrigins)
+	{
+		ViewInfos.Emplace(Origin, ScreenSize, ScreenSize, BoostFactor, false, 0.f, TWeakObjectPtr<AActor>(), TWeakObjectPtr<UWorld>());
+		FStreamingViewInfoExtra Extra;
+		Extra.ScreenSizeFloat = ScreenSize;
+		Extra.ExtraBoostForVisiblePrimitiveFloat = BoostFactor;
+		ViewInfoExtras.Add(Extra);
+	}
+
+	FRenderAssetStreamingSettings Settings;
+	FSimpleStreamableAssetManager::UpdateBoundSizes(ViewInfos, ViewInfoExtras, LastUpdateTime, Settings);
+}
+
+FSSAMTestEnvironment::FScreenSizeResult FSSAMTestEnvironment::QueryScreenSize(int32 AssetIndex)
+{
+	FScreenSizeResult Result;
+	FSimpleStreamableAssetManager::GetRenderAssetScreenSize(EStreamableRenderAssetType::Texture, AssetIndex, Result.MaxSize, Result.MaxSizeVisibleOnly, Result.MaxNumForcedLODs, FLT_MAX, MAX_TEXTURE_MIP_COUNT, TEXT("Test"));
+	return Result;
+}
+
+#else
+
+void USSAMTestStreamableAsset::BeginDestroy()
+{
+	Super::BeginDestroy();
+}
+
+FPrimitiveSceneProxy* USSAMTestComponent::CreateSceneProxy()
+{
+	return nullptr;
+}
+
+void USSAMTestComponent::GetStreamingRenderAssetInfo(FStreamingTextureLevelContext& LevelContext, TArray<FStreamingRenderAssetPrimitiveInfo>& OutStreamingRenderAssets) const
+{
+	Super::GetStreamingRenderAssetInfo(LevelContext, OutStreamingRenderAssets);
+}
+
+#endif // WITH_DEV_AUTOMATION_TESTS
diff --git a/Engine/Source/Runtime/Engine/Private/Tests/Streaming/SSAMTestTypes.h b/Engine/Source/Runtime/Engine/Private/Tests/Streaming/SSAMTestTypes.h
new file mode 100644
--- /dev/null
+++ b/Engine/Source/Runtime/Engine/Private/Tests/Streaming/SSAMTestTypes.h
@@ -0,0 +1,286 @@
+// Copyright Epic Games, Inc. All Rights Reserved.
+
+#pragma once
+
+#include "CoreMinimal.h"
+#include "Engine/StreamableRenderAsset.h"
+#include "Engine/TextureStreamingTypes.h"
+#include "Components/PrimitiveComponent.h"
+
+#if WITH_DEV_AUTOMATION_TESTS
+#include "PrimitiveSceneProxy.h"
+#include "PrimitiveSceneProxyDesc.h"
+#include "PrimitiveSceneDesc.h"
+#include "Streaming/SimpleStreamableAssetManager.h"
+#include "Tests/AutomationCommon.h"
+#if USING_INSTRUMENTATION
+#include "Sanitizer/RaceDetector.h"
+#endif
+#endif // WITH_DEV_AUTOMATION_TESTS
+
+#include "SSAMTestTypes.generated.h"
+
+// ---------------------------------------------------------------------------
+// USSAMTestStreamableAsset
+//
+// A minimal UStreamableRenderAsset subclass that SSAM treats as a real
+// streamable texture. No GPU resources, no render data — just the shared
+// index pointer that SSAM needs for tracking.
+//
+// Note: UCLASS cannot be inside #if WITH_DEV_AUTOMATION_TESTS (UHT restriction).
+// The class is harmless in non-test builds — it's Transient and never instantiated.
+// ---------------------------------------------------------------------------
+UCLASS(Transient)
+class USSAMTestStreamableAsset : public UStreamableRenderAsset
+{
+	GENERATED_BODY()
+
+public:
+	virtual EStreamableRenderAssetType GetRenderAssetType() const override
+	{
+		return EStreamableRenderAssetType::Texture;
+	}
+
+	// StreamingIndex is protected in the base class. Setting it to a
+	// non-INDEX_NONE value makes IsStreamable() return true, without
+	// calling LinkStreaming() (which would register with the real
+	// streaming manager and cause side effects).
+	void SetStreamingIndexForTest(int32 InIndex)
+	{
+		StreamingIndex = InIndex;
+	}
+
+	// SimpleStreamableAssetManagerIndex is protected on UStreamableRenderAsset.
+	// Expose read access for test verification.
+	const TSharedPtr<int32, ESPMode::ThreadSafe>& GetSSAMIndex() const
+	{
+		return SimpleStreamableAssetManagerIndex;
+	}
+
+	int32 GetSSAMIndexValue() const
+	{
+		return SimpleStreamableAssetManagerIndex.IsValid() ? *SimpleStreamableAssetManagerIndex : INDEX_NONE;
+	}
+
+	// Override BeginDestroy to route asset cleanup through SSAM's UnregisterAsset
+	// directly, bypassing UnlinkStreaming which calls into the real
+	// FRenderAssetStreamingManager (our test assets were never registered there).
+	virtual void BeginDestroy() override;
+
+	virtual int32 CalcCumulativeLODSize(int32 NumLODs) const override
+	{
+		// Return a fake size: 1KB per mip, doubling each level.
+		int32 Size = 0;
+		for (int32 i = 0; i < NumLODs; ++i)
+		{
+			Size += 1024 << i;
+		}
+		return Size;
+	}
+
+	// Set up minimal CachedSRRState so the streaming manager can create a
+	// valid FStreamingRenderAsset for this test asset (needs non-zero mip counts).
+	void SetupFakeStreamingState(int32 NumMips = 10, int32 NumNonStreaming = 1)
+	{
+		CachedSRRState.bSupportsStreaming = true;
+		CachedSRRState.MaxNumLODs = NumMips;
+		CachedSRRState.NumNonStreamingLODs = NumNonStreaming;
+		CachedSRRState.NumNonOptionalLODs = NumMips;
+		CachedSRRState.NumResidentLODs = NumMips;
+		CachedSRRState.NumRequestedLODs = NumMips;
+		CachedSRRState.AssetLODBias = 0;
+	}
+};
+
+// ---------------------------------------------------------------------------
+// USSAMTestComponent
+//
+// Primitive component that creates an FSSAMTestSceneProxy (defined below,
+// guarded by WITH_DEV_AUTOMATION_TESTS). Configure TestAssets and TestBounds
+// before registering.
+// ---------------------------------------------------------------------------
+UCLASS(Transient)
+class USSAMTestComponent : public UPrimitiveComponent
+{
+	GENERATED_BODY()
+
+public:
+	struct FAssetEntry
+	{
+		UStreamableRenderAsset* Asset = nullptr;
+		float TexelFactor = 1.0f;
+	};
+
+	TArray<FAssetEntry> TestAssets;
+	FBoxSphereBounds TestBounds = FBoxSphereBounds(FVector::ZeroVector, FVector(100.f), 100.f);
+	bool bUseFallbackProxy = false;
+	bool bSupportScaleFactor = false;
+
+	void AddTestAsset(USSAMTestStreamableAsset* Asset, float TexelFactor = 1.0f)
+	{
+		TestAssets.Add(FAssetEntry{ Asset, TexelFactor });
+	}
+
+	virtual FPrimitiveSceneProxy* CreateSceneProxy() override;
+
+	virtual FBoxSphereBounds CalcBounds(const FTransform& LocalToWorld) const override
+	{
+		return TestBounds.TransformBy(LocalToWorld);
+	}
+
+	// Override the streaming asset gathering used by the FALLBACK path.
+	// When using the fallback proxy, ProcessFallbackRegistrations_GameThread
+	// calls FStreamingTextureLevelContext::GetStreamingRenderAssetInfoWithNULLRemoval
+	// which calls Component->GetStreamingRenderAssetInfo. This override
+	// returns our test assets directly, bypassing the material system.
+	virtual void GetStreamingRenderAssetInfo(FStreamingTextureLevelContext& LevelContext, TArray<FStreamingRenderAssetPrimitiveInfo>& OutStreamingRenderAssets) const override;
+};
+
+// ===========================================================================
+// Everything below requires the test framework and is stripped from shipping.
+// ===========================================================================
+#if WITH_DEV_AUTOMATION_TESTS
+
+// Alias for convenience in test code
+using FSSAMTestAssetEntry = USSAMTestComponent::FAssetEntry;
+
+// ---------------------------------------------------------------------------
+// FSSAMTestProxyDesc
+//
+// Minimal FPrimitiveSceneProxyDesc for creating proxies without a UPrimitiveComponent.
+// Overrides GetUsedMaterials to avoid the check(Component) in the base class,
+// which is called unconditionally in editor builds from FPrimitiveSceneProxy's constructor.
+// ---------------------------------------------------------------------------
+struct FSSAMTestProxyDesc : public FPrimitiveSceneProxyDesc
+{
+	FCustomPrimitiveData TestCustomPrimitiveData;
+
+	FSSAMTestProxyDesc();
+
+	virtual void GetUsedMaterials(TArray<UMaterialInterface*>& OutMaterials, bool bGetDebugMaterials = false) const override
+	{
+	}
+};
+
+// ---------------------------------------------------------------------------
+// FSSAMTestSceneProxy
+//
+// Scene proxy that implements GetStreamableRenderAssetInfo, taking the
+// PRIMARY SSAM registration path (not the fallback component path).
+// The asset list is captured from the component at construction time.
+// ---------------------------------------------------------------------------
+class FSSAMTestSceneProxy : public FPrimitiveSceneProxy
+{
+public:
+	SIZE_T GetTypeHash() const override;
+	uint32 GetMemoryFootprint() const override { return sizeof(*this); }
+
+	FSSAMTestSceneProxy(const UPrimitiveComponent* InComponent, TArray<FStreamingRenderAssetPrimitiveInfo>&& InAssets, bool bInSupportScaleFactor = false);
+	FSSAMTestSceneProxy(const FPrimitiveSceneProxyDesc* InDesc, TArray<FStreamingRenderAssetPrimitiveInfo>&& InAssets, bool bInSupportScaleFactor = false);
+
+	virtual bool CanApplyStreamableRenderAssetScaleFactor() const override;
+
+	virtual void GetStreamableRenderAssetInfo(const FBoxSphereBounds& PrimitiveBounds, TArray<FStreamingRenderAssetPrimitiveInfo>& OutStreamableRenderAssets) const override;
+
+	virtual FPrimitiveViewRelevance GetViewRelevance(const FSceneView* View) const override;
+
+private:
+	TArray<FStreamingRenderAssetPrimitiveInfo> CapturedAssets;
+	bool bSupportScaleFactor = false;
+};
+
+// ---------------------------------------------------------------------------
+// FSSAMTestFallbackSceneProxy
+//
+// Scene proxy that does NOT implement GetStreamableRenderAssetInfo.
+// bImplementsStreamableAssetGathering stays false (default), forcing
+// RegisterHelper to take the FALLBACK path through
+// ProcessFallbackRegistrations_GameThread.
+// ---------------------------------------------------------------------------
+class FSSAMTestFallbackSceneProxy : public FPrimitiveSceneProxy
+{
+public:
+	SIZE_T GetTypeHash() const override;
+	uint32 GetMemoryFootprint() const override { return sizeof(*this); }
+
+	FSSAMTestFallbackSceneProxy(const UPrimitiveComponent* InComponent);
+
+	virtual FPrimitiveViewRelevance GetViewRelevance(const FSceneView* View) const override;
+};
+
+// ---------------------------------------------------------------------------
+// FSSAMTestEnvironment
+//
+// RAII helper that enables SSAM for the duration of a test. Handles:
+// - Setting the CVar (GUseSimpleStreamableAssetManager)
+// - Initializing / shutting down the SSAM singleton if needed
+// - Creating and managing a test world via FTestWorldWrapper
+// ---------------------------------------------------------------------------
+class FSSAMTestEnvironment
+{
+#if USING_INSTRUMENTATION
+	UE::Sanitizer::RaceDetector::FRaceDetectorScope RaceDetectorScope;
+#endif
+
+public:
+	FSSAMTestEnvironment();
+	~FSSAMTestEnvironment();
+
+	UWorld* GetWorld() const { return WorldWrapper.GetTestWorld(); }
+
+	void Tick(float DeltaTime = 0.01f)
+	{
+		WorldWrapper.TickTestWorld(DeltaTime);
+	}
+
+	void ProcessSSAM()
+	{
+		FSimpleStreamableAssetManager::Process();
+	}
+
+	// Convenience: create a test asset that SSAM will treat as streamable
+	USSAMTestStreamableAsset* CreateTestAsset();
+
+	// Convenience: spawn actor, attach a configured test component, register it.
+	// Component registration triggers the full engine flow:
+	// RegisterComponent -> CreateSceneProxy -> AddPrimitive -> BatchAddPrimitivesInternal -> RegisterHelper
+	USSAMTestComponent* SpawnTestComponent(const TArray<FSSAMTestAssetEntry>& Assets, const FBoxSphereBounds& Bounds = FBoxSphereBounds(FVector::ZeroVector, FVector(100.f), 100.f), const FVector& Location = FVector::ZeroVector);
+
+	// Spawn a component that uses the FALLBACK proxy (bImplementsStreamableAssetGathering = false).
+	// Registration goes through RegisterHelper's fallback path -> FallbackComponentRegistration queue.
+	// ProcessFallbackRegistrations_GameThread() must be called to complete registration.
+	USSAMTestComponent* SpawnTestFallbackComponent(const TArray<FSSAMTestAssetEntry>& Assets, const FBoxSphereBounds& Bounds = FBoxSphereBounds(FVector::ZeroVector, FVector(100.f), 100.f), const FVector& Location = FVector::ZeroVector);
+
+	void CleanupAssets();
+
+	// Fallback path was removed from SSAM — this is now a no-op.
+	void ProcessFallbackRegistrations()
+	{
+	}
+
+	// Release assets so GC can collect them through the normal BeginDestroy →
+	// UnlinkStreaming → UnregisterAsset path. Unlike CleanupAssets(), this does
+	// NOT set StreamingIndex = INDEX_NONE, preserving the real teardown flow.
+	// Released assets are removed from CreatedAssets so the destructor skips them.
+	void ReleaseAssetsForGC(TArrayView<USSAMTestStreamableAsset*> Assets);
+
+	// Streaming query helpers — eliminate repeated view setup boilerplate.
+	struct FScreenSizeResult
+	{
+		float MaxSize = 0.f;
+		float MaxSizeVisibleOnly = 0.f;
+		int32 MaxNumForcedLODs = 0;
+	};
+
+	void UpdateBoundSizes(const FVector& ViewOrigin, float ScreenSize = 1920.f, float BoostFactor = 1.f, float LastUpdateTime = 0.f);
+	void UpdateBoundSizes(TArrayView<const FVector> ViewOrigins, float ScreenSize = 1920.f, float BoostFactor = 1.f, float LastUpdateTime = 0.f);
+	FScreenSizeResult QueryScreenSize(int32 AssetIndex);
+
+private:
+	FTestWorldWrapper WorldWrapper;
+	TArray<USSAMTestStreamableAsset*> CreatedAssets;
+	bool bDidInit;
+	bool bWasEnabled;
+};
+
+#endif // WITH_DEV_AUTOMATION_TESTS
