Fix early return in the in-place version of UE::Anim::FNamedValueArrayUtils::Union so that it doesn't skip the predicate on existing elements

#Synced-CL 52242993

diff --git a/Engine/Source/Runtime/Engine/Public/Animation/NamedValueArray.h b/Engine/Source/Runtime/Engine/Public/Animation/NamedValueArray.h
--- a/Engine/Source/Runtime/Engine/Public/Animation/NamedValueArray.h
+++ b/Engine/Source/Runtime/Engine/Public/Animation/NamedValueArray.h
@@ -268,9 +268,16 @@ struct FNamedValueArrayUtils
 		const int32 NumElements0 = InOutValueArray0.Num();	// ValueArray1 elements remain constant, but ValueArray0 can have entries added.
 		const int32 NumElements1 = InValueArray1.Num();
 
-		// Early out if we have no elements to union
+		// Nothing to merge, just run the predicate on every elements.
 		if (NumElements1 == 0)
 		{
+			InOutValueArray0.SortElementsIfRequired();
+			ElementTypeParam DefaultElement;
+			for (auto& Element : InOutValueArray0.Elements)
+			{
+				DefaultElement.Name = Element.Name;
+				InPredicate(Element, DefaultElement, UE::Anim::ENamedValueUnionFlags::ValidArg0);
+			}
 			return;
 		}
 
