[GC] Add the option to GC based on pending actors to unload (in addition to existing mechanism that uses number of levels) for finer control (disabled by default). Number of levels can be inaccurate given the discrepancies between level counts and memory impact. across main grid / data layers

#tests ReplayRun with various values for wp.Runtime.LevelStreamingContinuouslyIncrementalGCWhileActorsPendingPurgeForWP. Disabled by default in this submit. Will continue monitoring PT data and make the call on if we want to enable this / what values we should go with on different platforms.
[RN] Add the option to GC based on pending actors to unload (in addition to existing mechanism that uses number of levels) for finer control (disabled by default). Controlled via `s.ContinuouslyIncrementalGCWhileActorsPendingPurge`/ `s.LevelStreamingLowMemoryActorsPendingPurgeCount` & `wp.Runtime.LevelStreamingContinuouslyIncrementalGCWhileActorsPendingPurgeForWP` (for partitioned worlds).
#Synced-CL 52194623

diff --git a/Engine/Source/Editor/UnrealEd/Private/WorldPartition/WorldPartitionBuilder.cpp b/Engine/Source/Editor/UnrealEd/Private/WorldPartition/WorldPartitionBuilder.cpp
--- a/Engine/Source/Editor/UnrealEd/Private/WorldPartition/WorldPartitionBuilder.cpp
+++ b/Engine/Source/Editor/UnrealEd/Private/WorldPartition/WorldPartitionBuilder.cpp
@@ -180,7 +180,9 @@ bool UWorldPartitionBuilder::Run(UWorld* World, FPackageSourceControlHelper& Pac
 {
 	// As we manage GC frequency ourselves during builds, make sure to turn off automated GC based on other settings
 	TGuardValue<int32> Guard_GLevelStreamingContinuouslyIncrementalGCWhileLevelsPendingPurge(GLevelStreamingContinuouslyIncrementalGCWhileLevelsPendingPurge, 0);
+	TGuardValue<int32> Guard_GLevelStreamingContinuouslyIncrementalGCWhileActorsPendingPurge(GLevelStreamingContinuouslyIncrementalGCWhileActorsPendingPurge, 0);
 	TGuardValue<int32> Guard_GLevelStreamingLowMemoryPendingPurgeCount(GLevelStreamingLowMemoryPendingPurgeCount, MAX_int32);
+	TGuardValue<int32> Guard_GLevelStreamingLowMemoryActorsPurgeCount(GLevelStreamingLowMemoryActorsPendingPurgeCount, MAX_int32);
 	TGuardValue<int32> Guard_GLevelStreamingForceGCAfterLevelStreamedOut(GLevelStreamingForceGCAfterLevelStreamedOut, 0);
 
 	UWorldPartition* WorldPartition = World->GetWorldPartition();
diff --git a/Engine/Source/Runtime/Engine/Classes/Engine/CoreSettings.h b/Engine/Source/Runtime/Engine/Classes/Engine/CoreSettings.h
--- a/Engine/Source/Runtime/Engine/Classes/Engine/CoreSettings.h
+++ b/Engine/Source/Runtime/Engine/Classes/Engine/CoreSettings.h
@@ -151,6 +151,8 @@ extern ENGINE_API float GLevelStreamingUnregisterComponentsTimeLimit;
 extern ENGINE_API int32 GLevelStreamingForceGCAfterLevelStreamedOut;
 /** Whether to kick off incremental GC when there are over the specified amount of levels still waiting to be purged. */
 extern ENGINE_API int32 GLevelStreamingContinuouslyIncrementalGCWhileLevelsPendingPurge;
+/** Whether to kick off incremental GC when there are over the specified amount of actors (from levels) still waiting to be purged. */
+extern ENGINE_API int32 GLevelStreamingContinuouslyIncrementalGCWhileActorsPendingPurge;
 /** Enables level streaming requests while async loading (of anything) while the match is already in progress and no loading screen is up. */
 extern ENGINE_API int32 GLevelStreamingAllowLevelRequestsWhileAsyncLoadingInMatch;
 /** When we're already loading this many levels and actively in match, don't allow any more requests until one of those completes.  Set to zero to disable. */
@@ -161,6 +163,8 @@ extern ENGINE_API int32 GLevelStreamingForceVerifyLevelsGotRemovedByGC;
 extern ENGINE_API int32 GLevelStreamingForceRouteActorInitializeNextFrame;
 /** In a low memory situation, will kick off a 'soft' GC if the number of streaming levels pending purge meets or exceeds this count. */
 extern ENGINE_API int32 GLevelStreamingLowMemoryPendingPurgeCount;
+/** In a low memory situation, will kick off a 'soft' GC if the number of actors (from streaming levels) pending purge meets or exceeds this count. */
+extern ENGINE_API int32 GLevelStreamingLowMemoryActorsPendingPurgeCount;
 
 /**
 * Implements the settings for garbage collection.
diff --git a/Engine/Source/Runtime/Engine/Classes/Engine/Level.h b/Engine/Source/Runtime/Engine/Classes/Engine/Level.h
--- a/Engine/Source/Runtime/Engine/Classes/Engine/Level.h
+++ b/Engine/Source/Runtime/Engine/Classes/Engine/Level.h
@@ -431,6 +431,9 @@ public:
 	/** Array of actors to be exposed to GC in this level. All other actors will be referenced through ULevelActorContainer */
 	TArray<TObjectPtr<AActor>> ActorsForGC;
 
+	/** Number of actors marked for pending unload by GC */
+	int32 NumActorsMarkedForPendingUnloadByGC = 0;
+
 	/** Add actor to list */
 	ENGINE_API bool TryAddActorToList(AActor* InActor, bool bAddUnique);
 
diff --git a/Engine/Source/Runtime/Engine/Classes/Engine/LevelStreamingGCHelper.h b/Engine/Source/Runtime/Engine/Classes/Engine/LevelStreamingGCHelper.h
--- a/Engine/Source/Runtime/Engine/Classes/Engine/LevelStreamingGCHelper.h
+++ b/Engine/Source/Runtime/Engine/Classes/Engine/LevelStreamingGCHelper.h
@@ -40,6 +40,11 @@ struct FLevelStreamingGCHelper
 	 */
 	static ENGINE_API int32 GetNumLevelsPendingPurge();
 
+	/**
+	 * @return	The number of actors (from levels) pending a purge by the garbage collector
+	 */
+	static ENGINE_API int32 GetNumActorsPendingPurge();
+
 	/**
 	 * Trash a package, making sure it will be reloaded if needed again by level streaming
 	 */
@@ -81,4 +86,10 @@ private:
 
 	/** The number of unloaded levels prepared for the next GC (used by GetNumLevelsPendingPurge). */
 	static int32 NumberOfPreparedStreamedOutLevelsForGC;
-};
+
+	/** The number of unloaded actors prepared for the next GC (used by GetNumActorsPendingPurge). */
+	static int32 NumberOfPreparedStreamedOutActorsForGC;
+
+	/** The number of actors that should be unloaded */
+	static int32 ActorsPendingUnload;
+};	
diff --git a/Engine/Source/Runtime/Engine/Private/CoreSettings.cpp b/Engine/Source/Runtime/Engine/Private/CoreSettings.cpp
--- a/Engine/Source/Runtime/Engine/Private/CoreSettings.cpp
+++ b/Engine/Source/Runtime/Engine/Private/CoreSettings.cpp
@@ -24,11 +24,13 @@ int32 GLevelStreamingRouteActorInitializationGranularity = 10;
 int32 GLevelStreamingRouteActorEndPlayForRemoveFromWorldGranularity = 0;
 int32 GLevelStreamingForceGCAfterLevelStreamedOut = 1;
 int32 GLevelStreamingContinuouslyIncrementalGCWhileLevelsPendingPurge = 1;
+int32 GLevelStreamingContinuouslyIncrementalGCWhileActorsPendingPurge = 0; // Disabled by default
 int32 GLevelStreamingAllowLevelRequestsWhileAsyncLoadingInMatch = 1;
 int32 GLevelStreamingMaxLevelRequestsAtOnceWhileInMatch = 0;
 int32 GLevelStreamingForceVerifyLevelsGotRemovedByGC = 0;
 int32 GLevelStreamingForceRouteActorInitializeNextFrame = 0;
 int32 GLevelStreamingLowMemoryPendingPurgeCount = MAX_int32;
+int32 GLevelStreamingLowMemoryActorsPendingPurgeCount = MAX_int32;
 
 static FAutoConsoleVariableRef CVarUseUnifiedTimeBudgetForStreaming(
 	TEXT("s.UseUnifiedTimeBudgetForStreaming"),
@@ -151,6 +153,13 @@ static FAutoConsoleVariableRef CVarContinuouslyIncrementalGCWhileLevelsPendingPu
 	ECVF_Default
 );
 
+static FAutoConsoleVariableRef CVarContinuouslyIncrementalGCWhileActorsPendingPurge(
+	TEXT("s.ContinuouslyIncrementalGCWhileActorsPendingPurge"),
+	GLevelStreamingContinuouslyIncrementalGCWhileActorsPendingPurge,
+	TEXT("Whether to repeatedly kick off incremental GC when there are actors (from levels) still waiting to be purged."),
+	ECVF_Default
+);
+
 static FAutoConsoleVariableRef CVarAllowLevelRequestsWhileAsyncLoadingInMatch(
 	TEXT("s.AllowLevelRequestsWhileAsyncLoadingInMatch"),
 	GLevelStreamingAllowLevelRequestsWhileAsyncLoadingInMatch,
@@ -172,6 +181,12 @@ static FAutoConsoleVariableRef CVarLevelStreamingLowMemoryPendingPurgeCount(
 	ECVF_Default
 );
 
+static FAutoConsoleVariableRef CVarLevelStreamingLowMemoryActorsPendingPurgeCount(
+	TEXT("s.LevelStreamingLowMemoryActorsPendingPurgeCount"),
+	GLevelStreamingLowMemoryActorsPendingPurgeCount,
+	TEXT("When system is in low memory state, if the number of actors from streaming levels to purge meets or exceeds this value, perform a 'soft' GC."),
+	ECVF_Default
+);
 
 UStreamingSettings::UStreamingSettings()
 	: Super()
diff --git a/Engine/Source/Runtime/Engine/Private/LevelStreamingGCHelper.cpp b/Engine/Source/Runtime/Engine/Private/LevelStreamingGCHelper.cpp
--- a/Engine/Source/Runtime/Engine/Private/LevelStreamingGCHelper.cpp
+++ b/Engine/Source/Runtime/Engine/Private/LevelStreamingGCHelper.cpp
@@ -17,6 +17,9 @@
 /************************************************************************/
 FLevelStreamingGCHelper::FOnGCStreamedOutLevelsEvent FLevelStreamingGCHelper::OnGCStreamedOutLevels;
 TArray<TWeakObjectPtr<ULevel> > FLevelStreamingGCHelper::LevelsPendingUnload;
+int32 FLevelStreamingGCHelper::NumberOfPreparedStreamedOutActorsForGC = 0;
+int32 FLevelStreamingGCHelper::ActorsPendingUnload = 0;
+
 TSet<FName> FLevelStreamingGCHelper::LevelPackageNames;
 bool FLevelStreamingGCHelper::bEnabledForCommandlet = false;
 bool FLevelStreamingGCHelper::bIsPrepareStreamedOutLevelForGCDelayedToWorldTickEnd = false;
@@ -73,19 +76,34 @@ void FLevelStreamingGCHelper::RequestUnload( ULevel* InLevel )
 				// Since PrepareStreamedOutLevelForGC cannot be called inside a world tick, delay it to the next OnWorldTickEnd
 				// TODO: Investigate why NotifyStreamingLevelUnload can't be called inside world's tick
 				bIsPrepareStreamedOutLevelForGCDelayedToWorldTickEnd = true;
-				LevelsPendingUnload.AddUnique(InLevel);
+				if (LevelsPendingUnload.Find(InLevel) == INDEX_NONE)
+				{
+					LevelsPendingUnload.Add(InLevel);
+					InLevel->NumActorsMarkedForPendingUnloadByGC = InLevel->Actors.Num();
+					ActorsPendingUnload += InLevel->NumActorsMarkedForPendingUnloadByGC;
+				}
 			}
 		}
 		else
 		{
-			LevelsPendingUnload.AddUnique(InLevel);
+			if (LevelsPendingUnload.Find(InLevel) == INDEX_NONE)
+			{
+				LevelsPendingUnload.Add(InLevel);
+				InLevel->NumActorsMarkedForPendingUnloadByGC = InLevel->Actors.Num();
+				ActorsPendingUnload += InLevel->NumActorsMarkedForPendingUnloadByGC;
+			}
 		}
 	}
 }
 
 void FLevelStreamingGCHelper::CancelUnloadRequest( ULevel* InLevel )
 {
-	LevelsPendingUnload.Remove( InLevel );
+	if (InLevel && LevelsPendingUnload.Remove(InLevel) && InLevel->NumActorsMarkedForPendingUnloadByGC > 0)
+	{
+		ActorsPendingUnload -= InLevel->NumActorsMarkedForPendingUnloadByGC;
+		check(ActorsPendingUnload >= 0);
+		InLevel->NumActorsMarkedForPendingUnloadByGC = 0;
+	}
 }
 
 void FLevelStreamingGCHelper::OnWorldTickEnd(UWorld* InWorld, ELevelTick InTickType, float InDeltaSeconds)
@@ -117,11 +135,18 @@ void FLevelStreamingGCHelper::PrepareStreamedOutLevelsForGC()
 	for( int32 LevelIndex=0; LevelIndex < LevelsPendingUnload.Num(); LevelIndex++ )
 	{
 		ULevel* Level = LevelsPendingUnload[LevelIndex].Get();
+		if (Level && Level->NumActorsMarkedForPendingUnloadByGC > 0)
+		{
+			// Reset cached NumActorsMarkedForPendingUnloadByGC for level.
+			Level->NumActorsMarkedForPendingUnloadByGC = 0;
+		}
 		PrepareStreamedOutLevelForGC(Level);
 	}
 
 	LevelsPendingUnload.Empty();
 	NumberOfPreparedStreamedOutLevelsForGC = 0;
+	NumberOfPreparedStreamedOutActorsForGC = 0;
+	ActorsPendingUnload = 0;
 }
 
 void FLevelStreamingGCHelper::PrepareStreamedOutLevelForGC(ULevel* InLevel)
@@ -129,6 +154,7 @@ void FLevelStreamingGCHelper::PrepareStreamedOutLevelForGC(ULevel* InLevel)
 	if (InLevel && ((!GIsEditor || bEnabledForCommandlet) || InLevel->GetPackage()->HasAnyPackageFlags(PKG_PlayInEditor)))
 	{
 		++NumberOfPreparedStreamedOutLevelsForGC;
+		NumberOfPreparedStreamedOutActorsForGC += InLevel->Actors.Num();
 
 		UPackage* LevelPackage = InLevel->GetOutermost();
 		UE_LOG(LogStreaming, Log, TEXT("PrepareStreamedOutLevelsForGC called on '%s'"), *LevelPackage->GetName());
@@ -258,6 +284,11 @@ int32 FLevelStreamingGCHelper::GetNumLevelsPendingPurge()
 	return LevelsPendingUnload.Num() + NumberOfPreparedStreamedOutLevelsForGC;
 }
 
+int32 FLevelStreamingGCHelper::GetNumActorsPendingPurge()
+{
+	return ActorsPendingUnload + NumberOfPreparedStreamedOutActorsForGC;
+}
+
 void FLevelStreamingGCHelper::TrashPackage(UPackage* InPackage)
 {
 	TrashObject(InPackage);
diff --git a/Engine/Source/Runtime/Engine/Private/ProfilingDebugging/LevelStreamingProfilingSubsystem.cpp b/Engine/Source/Runtime/Engine/Private/ProfilingDebugging/LevelStreamingProfilingSubsystem.cpp
--- a/Engine/Source/Runtime/Engine/Private/ProfilingDebugging/LevelStreamingProfilingSubsystem.cpp
+++ b/Engine/Source/Runtime/Engine/Private/ProfilingDebugging/LevelStreamingProfilingSubsystem.cpp
@@ -388,6 +388,7 @@ void ULevelStreamingProfilingSubsystem::RecordLevelStreamingCountsToCSV()
 		CSV_CUSTOM_STAT(LevelStreaming, NumLevelsMakingVisible, NumLevelsMakingVisible, ECsvCustomStatOp::Set);
 		CSV_CUSTOM_STAT(LevelStreaming, NumLevelsMakingInvisible, NumLevelsMakingInvisible, ECsvCustomStatOp::Set);
 		CSV_CUSTOM_STAT(LevelStreaming, NumLevelsPendingPurge, FLevelStreamingGCHelper::GetNumLevelsPendingPurge(), ECsvCustomStatOp::Set);
+		CSV_CUSTOM_STAT(LevelStreaming, NumActorsPendingPurge, FLevelStreamingGCHelper::GetNumActorsPendingPurge(), ECsvCustomStatOp::Set);
 	}
 }
 #endif
diff --git a/Engine/Source/Runtime/Engine/Private/World.cpp b/Engine/Source/Runtime/Engine/Private/World.cpp
--- a/Engine/Source/Runtime/Engine/Private/World.cpp
+++ b/Engine/Source/Runtime/Engine/Private/World.cpp
@@ -5037,8 +5037,9 @@ void UWorld::UpdateLevelStreaming(const TOptional<const UE::FTimeout>& ExternalT
 		GAdaptiveAddToWorld.BeginUpdate();
 	}
 
-	// Store current number of pending unload levels, it may change in loop bellow
+	// Store current number of pending unload levels (or actors), it may change in loop bellow
 	const int32 NumLevelsPendingPurge = FLevelStreamingGCHelper::GetNumLevelsPendingPurge();
+	const int32 NumActorsPendingPurge = FLevelStreamingGCHelper::GetNumActorsPendingPurge();
 
 	if (GLevelVisibilityPrioritySort >= 0)
 	{
@@ -5141,10 +5142,17 @@ void UWorld::UpdateLevelStreaming(const TOptional<const UE::FTimeout>& ExternalT
 	}
 
 	const int32 CurrentNumLevelsPendingPurge = FLevelStreamingGCHelper::GetNumLevelsPendingPurge();
-	const bool bShouldPurgeLevels = GLevelStreamingContinuouslyIncrementalGCWhileLevelsPendingPurge && CurrentNumLevelsPendingPurge >= GLevelStreamingContinuouslyIncrementalGCWhileLevelsPendingPurge;
+	const int32 CurrentNumActorsPendingPurge = FLevelStreamingGCHelper::GetNumActorsPendingPurge();
+	const bool bShouldPurgeLevels = 
+		(GLevelStreamingContinuouslyIncrementalGCWhileLevelsPendingPurge && CurrentNumLevelsPendingPurge >= GLevelStreamingContinuouslyIncrementalGCWhileLevelsPendingPurge) ||
+		(GLevelStreamingContinuouslyIncrementalGCWhileActorsPendingPurge && CurrentNumActorsPendingPurge >= GLevelStreamingContinuouslyIncrementalGCWhileActorsPendingPurge);
 
-	// Are we currently in a low memory situation and number of pending levels to purge meets or exceeds our threshold?
-	const bool bShouldDoLowMemoryGC = GIsLowMemory && CurrentNumLevelsPendingPurge >= GLevelStreamingLowMemoryPendingPurgeCount;
+
+	// Are we currently in a low memory situation and number of pending levels (or actors) to purge meets or exceeds our threshold?
+	const bool bShouldDoLowMemoryGC = GIsLowMemory && (
+		(CurrentNumLevelsPendingPurge >= GLevelStreamingLowMemoryPendingPurgeCount) ||
+		(CurrentNumActorsPendingPurge >= GLevelStreamingLowMemoryActorsPendingPurgeCount)
+		);
 
 	// Low memory GC takes precedence over Continuous GC condition
 	if (bShouldDoLowMemoryGC)
@@ -5159,13 +5167,17 @@ void UWorld::UpdateLevelStreaming(const TOptional<const UE::FTimeout>& ExternalT
 		GEngine->ForceGarbageCollection(false);
 	}
 
-	// In case more levels has been requested to unload, force GC on next tick 
+	// In case more levels (or actors) have been requested to unload, force GC on next tick 
 	if (GLevelStreamingForceGCAfterLevelStreamedOut != 0)
 	{
 		if (NumLevelsPendingPurge < FLevelStreamingGCHelper::GetNumLevelsPendingPurge())
 		{
 			GEngine->ForceGarbageCollection(true); 
 		}
+		else if (NumActorsPendingPurge < FLevelStreamingGCHelper::GetNumActorsPendingPurge())
+		{
+			GEngine->ForceGarbageCollection(true);
+		}
 	}
 	if ( GAdaptiveAddToWorld.IsEnabled() )
 	{
diff --git a/Engine/Source/Runtime/Engine/Private/WorldPartition/WorldPartitionSubsystem.cpp b/Engine/Source/Runtime/Engine/Private/WorldPartition/WorldPartitionSubsystem.cpp
--- a/Engine/Source/Runtime/Engine/Private/WorldPartition/WorldPartitionSubsystem.cpp
+++ b/Engine/Source/Runtime/Engine/Private/WorldPartition/WorldPartitionSubsystem.cpp
@@ -152,6 +152,14 @@ static FAutoConsoleVariableRef CVarGLevelStreamingContinuouslyIncrementalGCWhile
 	ECVF_Default
 );
 
+int32 GLevelStreamingContinuouslyIncrementalGCWhileActorsPendingPurgeForWP = 0;	// Disabled by default
+static FAutoConsoleVariableRef CVarGLevelStreamingContinuouslyIncrementalGCWhileActorsPendingPurgeForWP(
+	TEXT("wp.Runtime.LevelStreamingContinuouslyIncrementalGCWhileActorsPendingPurgeForWP"),
+	GLevelStreamingContinuouslyIncrementalGCWhileActorsPendingPurgeForWP,
+	TEXT("Force a GC update when there's more than the number of specified Actors (from levels) pending purge."),
+	ECVF_Default
+);
+
 static FAutoConsoleCommandWithOutputDevice GDumpStreamingSourcesCmd(
 	TEXT("wp.Runtime.DumpStreamingSources"),
 	TEXT("Dumps active streaming sources to the log"),
@@ -584,9 +592,11 @@ void UWorldPartitionSubsystem::OnWorldPartitionInitialized(UWorldPartition* InWo
 		if (GetWorld()->IsGameWorld())
 		{
 			LevelStreamingContinuouslyIncrementalGCWhileLevelsPendingPurge = GLevelStreamingContinuouslyIncrementalGCWhileLevelsPendingPurge;
+			LevelStreamingContinuouslyIncrementalGCWhileActorsPendingPurge = GLevelStreamingContinuouslyIncrementalGCWhileActorsPendingPurge;
 			LevelStreamingForceGCAfterLevelStreamedOut = GLevelStreamingForceGCAfterLevelStreamedOut;
 
 			GLevelStreamingContinuouslyIncrementalGCWhileLevelsPendingPurge = GLevelStreamingContinuouslyIncrementalGCWhileLevelsPendingPurgeForWP;
+			GLevelStreamingContinuouslyIncrementalGCWhileActorsPendingPurge = GLevelStreamingContinuouslyIncrementalGCWhileActorsPendingPurgeForWP;
 			GLevelStreamingForceGCAfterLevelStreamedOut = 0;
 		}
 	}
@@ -645,6 +655,7 @@ void UWorldPartitionSubsystem::OnWorldPartitionUninitialized(UWorldPartition* In
 		if (GetWorld()->IsGameWorld())
 		{
 			GLevelStreamingContinuouslyIncrementalGCWhileLevelsPendingPurge = LevelStreamingContinuouslyIncrementalGCWhileLevelsPendingPurge;
+			GLevelStreamingContinuouslyIncrementalGCWhileActorsPendingPurge = LevelStreamingContinuouslyIncrementalGCWhileActorsPendingPurge;
 			GLevelStreamingForceGCAfterLevelStreamedOut = LevelStreamingForceGCAfterLevelStreamedOut;
 		}
 
diff --git a/Engine/Source/Runtime/Engine/Public/WorldPartition/WorldPartitionSubsystem.h b/Engine/Source/Runtime/Engine/Public/WorldPartition/WorldPartitionSubsystem.h
--- a/Engine/Source/Runtime/Engine/Public/WorldPartition/WorldPartitionSubsystem.h
+++ b/Engine/Source/Runtime/Engine/Public/WorldPartition/WorldPartitionSubsystem.h
@@ -222,6 +222,7 @@ private:
 
 	// GC backup values
 	int32 LevelStreamingContinuouslyIncrementalGCWhileLevelsPendingPurge;
+	int32 LevelStreamingContinuouslyIncrementalGCWhileActorsPendingPurge;
 	int32 LevelStreamingForceGCAfterLevelStreamedOut;
 
 	// Tracks streaming levels of uninitialized world partition used to delay next initialization until they're done being removed
