[PostProcessing] Fixed an issue where PPVs that do not have a UObject base caused false negatives when detecting dangling pointers.
CL 52090109 used TWeakInterfacePtr<T> to detect dangling PPV pointers. This works for most PPVs as they have a UObject base. Some PPV implementations, such as water, do not, which causes the pointer to incorrectly be classified as invalid. To fix this, this CL adds a path for non-UObject raw pointers and shared pointers.

#rnx
#tests Verified this fixes the water postprocess effect using QA_Waterworks in QAGame

diff --git a/Engine/Plugins/MovieScene/MovieRenderPipeline/Source/MovieRenderPipelineCore/Private/Graph/MovieGraphCoreTimeStep.cpp b/Engine/Plugins/MovieScene/MovieRenderPipeline/Source/MovieRenderPipelineCore/Private/Graph/MovieGraphCoreTimeStep.cpp
--- a/Engine/Plugins/MovieScene/MovieRenderPipeline/Source/MovieRenderPipelineCore/Private/Graph/MovieGraphCoreTimeStep.cpp
+++ b/Engine/Plugins/MovieScene/MovieRenderPipeline/Source/MovieRenderPipelineCore/Private/Graph/MovieGraphCoreTimeStep.cpp
@@ -763,15 +763,9 @@ float UMovieGraphCoreTimeStep::GetBlendedMotionBlurAmount()
 	{
 		// Apply any motion blur settings from post process volumes in the world
 		FVector ViewLocation = PlayerCameraManager->GetCameraLocation();
-		for (const TWeakInterfacePtr<IInterface_PostProcessVolume>& PPVolumeWeak : GetWorld()->GetPostProcessVolumes())
+		for (IInterface_PostProcessVolume& PPVolume : GetWorld()->GetPostProcessVolumeIterator())
 		{
-			IInterface_PostProcessVolume* PPVolume = PPVolumeWeak.Get();
-			if (!PPVolume)
-			{
-				continue;
-			}
-
-			const FPostProcessVolumeProperties VolumeProperties = PPVolume->GetProperties();
+			const FPostProcessVolumeProperties VolumeProperties = PPVolume.GetProperties();
 
 			// Skip any volumes which are either disabled or don't modify blur amount
 			if (!VolumeProperties.bIsEnabled || !VolumeProperties.Settings->bOverride_MotionBlurAmount)
@@ -784,7 +778,7 @@ float UMovieGraphCoreTimeStep::GetBlendedMotionBlurAmount()
 			if (!VolumeProperties.bIsUnbound)
 			{
 				float DistanceToPoint = 0.0f;
-				PPVolume->EncompassesPoint(ViewLocation, 0.0f, &DistanceToPoint);
+				PPVolume.EncompassesPoint(ViewLocation, 0.0f, &DistanceToPoint);
 
 				if (DistanceToPoint >= 0 && DistanceToPoint < VolumeProperties.BlendRadius)
 				{
diff --git a/Engine/Plugins/MovieScene/MovieRenderPipeline/Source/MovieRenderPipelineCore/Private/MoviePipeline.cpp b/Engine/Plugins/MovieScene/MovieRenderPipeline/Source/MovieRenderPipelineCore/Private/MoviePipeline.cpp
--- a/Engine/Plugins/MovieScene/MovieRenderPipeline/Source/MovieRenderPipelineCore/Private/MoviePipeline.cpp
+++ b/Engine/Plugins/MovieScene/MovieRenderPipeline/Source/MovieRenderPipelineCore/Private/MoviePipeline.cpp
@@ -1409,15 +1409,9 @@ MoviePipeline::FFrameConstantMetrics UMoviePipeline::CalculateShotFrameMetrics(c
 	{
 		// Apply any motion blur settings from post process volumes in the world
 		FVector ViewLocation = PlayerCameraManager->GetCameraLocation();
-		for (const TWeakInterfacePtr<IInterface_PostProcessVolume>& PPVolumeWeak : GetWorld()->GetPostProcessVolumes())
+		for (IInterface_PostProcessVolume& PPVolume : GetWorld()->GetPostProcessVolumeIterator())
 		{
-			IInterface_PostProcessVolume* PPVolume = PPVolumeWeak.Get();
-			if (!PPVolume)
-			{
-				continue;
-			}
-
-			const FPostProcessVolumeProperties VolumeProperties = PPVolume->GetProperties();
+			const FPostProcessVolumeProperties VolumeProperties = PPVolume.GetProperties();
 
 			// Skip any volumes which are either disabled or don't modify blur amount
 			if (!VolumeProperties.bIsEnabled || !VolumeProperties.Settings->bOverride_MotionBlurAmount)
@@ -1430,7 +1424,7 @@ MoviePipeline::FFrameConstantMetrics UMoviePipeline::CalculateShotFrameMetrics(c
 			if (!VolumeProperties.bIsUnbound)
 			{
 				float DistanceToPoint = 0.0f;
-				PPVolume->EncompassesPoint(ViewLocation, 0.0f, &DistanceToPoint);
+				PPVolume.EncompassesPoint(ViewLocation, 0.0f, &DistanceToPoint);
 
 				if (DistanceToPoint >= 0 && DistanceToPoint < VolumeProperties.BlendRadius)
 				{
diff --git a/Engine/Plugins/MovieScene/MovieRenderPipeline/Source/MovieRenderPipelineCore/Private/MoviePipelineUtils.cpp b/Engine/Plugins/MovieScene/MovieRenderPipeline/Source/MovieRenderPipelineCore/Private/MoviePipelineUtils.cpp
--- a/Engine/Plugins/MovieScene/MovieRenderPipeline/Source/MovieRenderPipelineCore/Private/MoviePipelineUtils.cpp
+++ b/Engine/Plugins/MovieScene/MovieRenderPipeline/Source/MovieRenderPipelineCore/Private/MoviePipelineUtils.cpp
@@ -1271,15 +1271,9 @@ namespace UE
 
 		void DoPostProcessBlend(const FVector& InViewLocation, const UWorld* InWorld, const FMinimalViewInfo& InViewInfo, FSceneView* InOutView)
 		{
-			for (const TWeakInterfacePtr<IInterface_PostProcessVolume>& PPVolumeWeak : InWorld->GetPostProcessVolumes())
+			for (IInterface_PostProcessVolume& PPVolume : InWorld->GetPostProcessVolumeIterator())
 			{
-				IInterface_PostProcessVolume* PPVolume = PPVolumeWeak.Get();
-				if (!PPVolume)
-				{
-					continue;
-				}
-
-				const FPostProcessVolumeProperties VolumeProperties = PPVolume->GetProperties();
+				const FPostProcessVolumeProperties VolumeProperties = PPVolume.GetProperties();
 
 				// Skip any volumes which are disabled
 				if (!VolumeProperties.bIsEnabled)
@@ -1292,7 +1286,7 @@ namespace UE
 				if (!VolumeProperties.bIsUnbound)
 				{
 					float DistanceToPoint = 0.0f;
-					PPVolume->EncompassesPoint(InViewLocation, 0.0f, &DistanceToPoint);
+					PPVolume.EncompassesPoint(InViewLocation, 0.0f, &DistanceToPoint);
 
 					if (DistanceToPoint >= 0 && DistanceToPoint < VolumeProperties.BlendRadius)
 					{
diff --git a/Engine/Plugins/Runtime/nDisplay/Source/DisplayClusterMoviePipeline/Private/DisplayClusterMoviePipelineViewportPass.cpp b/Engine/Plugins/Runtime/nDisplay/Source/DisplayClusterMoviePipeline/Private/DisplayClusterMoviePipelineViewportPass.cpp
--- a/Engine/Plugins/Runtime/nDisplay/Source/DisplayClusterMoviePipeline/Private/DisplayClusterMoviePipelineViewportPass.cpp
+++ b/Engine/Plugins/Runtime/nDisplay/Source/DisplayClusterMoviePipeline/Private/DisplayClusterMoviePipelineViewportPass.cpp
@@ -175,15 +175,9 @@ void UDisplayClusterMoviePipelineViewportPassBase::BlendPostProcessSettings(FSce
 			// Use DCRA as a camera
 			FVector ViewLocation = DCViews[InViewportIndex].ViewLocation;
 
-			for (const TWeakInterfacePtr<IInterface_PostProcessVolume>& PPVolumeWeak : GetWorld()->GetPostProcessVolumes())
+			for (IInterface_PostProcessVolume& PPVolume : GetWorld()->GetPostProcessVolumeIterator())
 			{
-				IInterface_PostProcessVolume* PPVolume = PPVolumeWeak.Get();
-				if (!PPVolume)
-				{
-					continue;
-				}
-
-				const FPostProcessVolumeProperties VolumeProperties = PPVolume->GetProperties();
+				const FPostProcessVolumeProperties VolumeProperties = PPVolume.GetProperties();
 
 				// Skip any volumes which are disabled
 				if (!VolumeProperties.bIsEnabled)
@@ -196,7 +190,7 @@ void UDisplayClusterMoviePipelineViewportPassBase::BlendPostProcessSettings(FSce
 				if (!VolumeProperties.bIsUnbound)
 				{
 					float DistanceToPoint = 0.0f;
-					PPVolume->EncompassesPoint(ViewLocation, 0.0f, &DistanceToPoint);
+					PPVolume.EncompassesPoint(ViewLocation, 0.0f, &DistanceToPoint);
 
 					if (DistanceToPoint >= 0 && DistanceToPoint < VolumeProperties.BlendRadius)
 					{
diff --git a/Engine/Source/Runtime/Engine/Classes/Engine/World.h b/Engine/Source/Runtime/Engine/Classes/Engine/World.h
--- a/Engine/Source/Runtime/Engine/Classes/Engine/World.h
+++ b/Engine/Source/Runtime/Engine/Classes/Engine/World.h
@@ -1909,30 +1909,97 @@ public:
 	DECLARE_EVENT_TwoParams(UWorld, FOnBeginPostProcessSettings, FVector, FSceneView*);
 	FOnBeginPostProcessSettings OnBeginPostProcessSettings;
 
-	/** Inserts a post process volume into the world in priority order */
+	// Storing PPVs safely is somewhat awkward due to historical API decisions.
+	// We want to store IInterface_PostProcessVolume*, but callers do not always properly call the Remove function,
+	// which leads to dangling pointers.
+	// Most PPVs are UObject-based, and so can use TWeakInterfacePtr to detect object expiration.
+	// However, some usages implement IInterface_PostProcessVolume directly without the UObject inheritance, 
+	// in which case TWeakInterfacePtr doesn't work. For these cases we have TWeakPtr as a safe option.
+	// The API signature historically allows for passing a raw pointer though, so we also support those cases for now.
+	// TLDR: Variant type for PPV storage: UObject-based PPVs use TWeakInterfacePtr (detects stale entries),
+	// non-UObject shared-ownership PPVs use TWeakPtr (detects stale entries),
+	// non-UObject raw pointer PPVs are assumed always valid.
+	using FPostProcessVolumeEntry = TVariant<
+		TWeakInterfacePtr<IInterface_PostProcessVolume>,
+		TWeakPtr<IInterface_PostProcessVolume>,
+		IInterface_PostProcessVolume*>;
+
+	/** Iterator range over post process volumes in ascending priority order, filtering out invalid/expired entries. */
+	struct FPostProcessVolumeIterator
+	{
+		struct FSentinel {};
+
+		struct FIterator
+		{
+			FIterator(const TConstArrayView<FPostProcessVolumeEntry> InEntries, int32 InIndex)
+				: Entries(InEntries), Index(InIndex)
+			{
+				AdvanceToValid();
+			}
+
+			IInterface_PostProcessVolume& operator*() const { return *Current; }
+
+			FIterator& operator++()
+			{
+				++Index;
+				AdvanceToValid();
+				return *this;
+			}
+
+			bool operator!=(FSentinel) const { return Current != nullptr; }
+
+		private:
+			UE_API void AdvanceToValid();
+
+			const TConstArrayView<FPostProcessVolumeEntry> Entries;
+			int32 Index;
+			TSharedPtr<IInterface_PostProcessVolume> PinnedPtr;
+			IInterface_PostProcessVolume* Current = nullptr;
+		};
+
+		explicit FPostProcessVolumeIterator(const TArray<FPostProcessVolumeEntry>& InEntries)
+			: Entries(InEntries)
+		{}
+
+		FIterator begin() const { return FIterator(Entries, 0); }
+		FSentinel end() const { return {}; }
+
+	private:
+		const TConstArrayView<FPostProcessVolumeEntry> Entries;
+	};
+
+	/** Inserts a non-UObject post process volume with shared ownership into the world in priority order */
+	UE_API void InsertPostProcessVolume(TSharedPtr<IInterface_PostProcessVolume>& InVolume);
+
+	/** Inserts a UObject-based post process volume into the world in priority order */
 	UE_API void InsertPostProcessVolume(IInterface_PostProcessVolume* InVolume);
 
 	/** Removes a post process volume from the world.  Returns true if it was successfully removed. */
 	UE_API bool RemovePostProcessVolume(IInterface_PostProcessVolume* InVolume);
-	
+
+	bool RemovePostProcessVolume(TSharedPtr<IInterface_PostProcessVolume>& InVolume)
+	{
+		return RemovePostProcessVolume(InVolume.Get());
+	}
+
 	/** Called when a scene view for this world needs the worlds post process settings computed */
 	UE_API void AddPostProcessingSettings(FVector ViewLocation, FSceneView* SceneView);
 
-	/** Return an array of post processing volumes, sorted in ascending order of priority. */
-	const TArray<TWeakInterfacePtr<IInterface_PostProcessVolume>>& GetPostProcessVolumes() const
+	/** Returns an iterator over post process volumes in ascending priority order, filtering out invalid entries. */
+	FPostProcessVolumeIterator GetPostProcessVolumeIterator() const
 	{
-		return PostProcessVolumesWeak;
+		return FPostProcessVolumeIterator(PostProcessVolumeEntries);
 	}
 
-	UE_DEPRECATED(5.8, "Use GetPostProcessVolumes() instead")
+	UE_DEPRECATED(5.8, "Use GetPostProcessVolumeIterator() instead")
 	TArray<IInterface_PostProcessVolume*> PostProcessVolumes;
 
 
 private:
-	/** Must be identical to `PostProcessVolumes`. TWeakInterfacePtr is used here to detect stale (destroyed) volumes that were not removed. */
-	TArray<TWeakInterfacePtr<IInterface_PostProcessVolume>> PostProcessVolumesWeak;
+	TArray<FPostProcessVolumeEntry> PostProcessVolumeEntries;
+	void InsertPostProcessVolumeEntry(FPostProcessVolumeEntry InVolume);
 	void RemovePostProcessVolume(int32 Index); // do not make public, indices are unstable
-	/** Caches the size of each volume at the time of insertion, for sort determinism in case it later changes.  One for one with PostProcessVolumes above. */
+	/** Caches the size of each volume at the time of insertion, for sort determinism in case it later changes.  One for one with PostProcessVolumeEntries above. */
 	TArray<double> PostProcessVolumeCachedSizes;
 public:
 
diff --git a/Engine/Source/Runtime/Engine/Private/World.cpp b/Engine/Source/Runtime/Engine/Private/World.cpp
--- a/Engine/Source/Runtime/Engine/Private/World.cpp
+++ b/Engine/Source/Runtime/Engine/Private/World.cpp
@@ -10385,6 +10385,23 @@ FPrimaryAssetId UWorld::GetPrimaryAssetId() const
 	return FPrimaryAssetId();
 }
 
+static IInterface_PostProcessVolume* ResolvePPVEntry(
+	const UWorld::FPostProcessVolumeEntry& Entry,
+	TSharedPtr<IInterface_PostProcessVolume>& OutPinned)
+{
+	OutPinned = nullptr;
+	if (Entry.IsType<TWeakInterfacePtr<IInterface_PostProcessVolume>>())
+	{
+		return Entry.Get<TWeakInterfacePtr<IInterface_PostProcessVolume>>().Get();
+	}
+	if (Entry.IsType<TWeakPtr<IInterface_PostProcessVolume>>())
+	{
+		OutPinned = Entry.Get<TWeakPtr<IInterface_PostProcessVolume>>().Pin();
+		return OutPinned.Get();
+	}
+	return Entry.Get<IInterface_PostProcessVolume*>();
+}
+
 static bool PostProcessVolumeLess(const FPostProcessVolumeProperties& PropertiesA, const FPostProcessVolumeProperties& PropertiesB, double SizeB)
 {
 	if (PropertiesA.Priority != PropertiesB.Priority)
@@ -10406,16 +10423,40 @@ static bool PostProcessVolumeLess(const FPostProcessVolumeProperties& Properties
 	return PropertiesA.VolumeGuid < PropertiesB.VolumeGuid;
 }
 
-void UWorld::InsertPostProcessVolume(IInterface_PostProcessVolume* InVolume)
+void UWorld::FPostProcessVolumeIterator::FIterator::AdvanceToValid()
+{
+	PinnedPtr = nullptr;
+	Current = nullptr;
+	while (Index < Entries.Num())
+	{
+		const FPostProcessVolumeEntry& Entry = Entries[Index];
+		Current = ResolvePPVEntry(Entry, PinnedPtr);
+
+		if (Current)
+		{
+			break;
+		}
+		else
+		{
+			++Index;
+		}
+	}
+}
+
+void UWorld::InsertPostProcessVolumeEntry(FPostProcessVolumeEntry InVolume)
 {
-	int32 NumVolumes = PostProcessVolumesWeak.Num();
-	FPostProcessVolumeProperties TargetProperties = InVolume->GetProperties();
+	TSharedPtr<IInterface_PostProcessVolume> InVolumeSharedPtr = nullptr;
+	IInterface_PostProcessVolume* InVolumeRaw = ResolvePPVEntry(InVolume, InVolumeSharedPtr);
+
+	int32 NumVolumes = PostProcessVolumeEntries.Num();
+	FPostProcessVolumeProperties TargetProperties = InVolumeRaw->GetProperties();
 	int32 InsertIndex = 0;
+	TSharedPtr<IInterface_PostProcessVolume> CurrentPinned;
 	// TODO: replace with binary search.
 	while (InsertIndex < NumVolumes)
 	{
-		IInterface_PostProcessVolume* CurrentVolume = PostProcessVolumesWeak[InsertIndex].Get();
-		if (CurrentVolume == InVolume)
+		IInterface_PostProcessVolume* CurrentVolume = ResolvePPVEntry(PostProcessVolumeEntries[InsertIndex], CurrentPinned);
+		if (CurrentVolume == InVolumeRaw)
 		{
 			return;
 		}
@@ -10442,8 +10483,8 @@ void UWorld::InsertPostProcessVolume(IInterface_PostProcessVolume* InVolume)
 	// have been reached in the loop.
 	for (int32 Index = InsertIndex + 1; Index < NumVolumes; )
 	{
-		IInterface_PostProcessVolume* CurrentVolume = PostProcessVolumesWeak[Index].Get();
-		if (CurrentVolume == InVolume)
+		IInterface_PostProcessVolume* CurrentVolume = ResolvePPVEntry(PostProcessVolumeEntries[Index], CurrentPinned);
+		if (CurrentVolume == InVolumeRaw)
 		{
 			return;
 		}
@@ -10459,18 +10500,49 @@ void UWorld::InsertPostProcessVolume(IInterface_PostProcessVolume* InVolume)
 		}
 	}
 
-	PostProcessVolumesWeak.Insert(TWeakInterfacePtr<IInterface_PostProcessVolume>(InVolume), InsertIndex);
+	PostProcessVolumeEntries.Insert(MoveTemp(InVolume), InsertIndex);
 	PostProcessVolumeCachedSizes.Insert(TargetProperties.Size, InsertIndex);
 PRAGMA_DISABLE_DEPRECATION_WARNINGS
-	PostProcessVolumes.Insert(InVolume, InsertIndex);
+	PostProcessVolumes.Insert(InVolumeRaw, InsertIndex);
 PRAGMA_ENABLE_DEPRECATION_WARNINGS
 }
 
+void UWorld::InsertPostProcessVolume(IInterface_PostProcessVolume* InVolume)
+{
+	check(InVolume);
+
+	// Is InVolume a UObject?
+	TWeakInterfacePtr<IInterface_PostProcessVolume> WeakInterfacePtr { InVolume };
+	if (WeakInterfacePtr.Get())
+	{
+		FPostProcessVolumeEntry Entry;
+		Entry.Set<TWeakInterfacePtr<IInterface_PostProcessVolume>>(WeakInterfacePtr);
+		InsertPostProcessVolumeEntry(MoveTemp(Entry));
+	}
+	else
+	{
+		// InVolume is not a UObject; store as a raw pointer.
+		FPostProcessVolumeEntry Entry;
+		Entry.Set<IInterface_PostProcessVolume*>(InVolume);
+		InsertPostProcessVolumeEntry(MoveTemp(Entry));
+	}
+}
+
+void UWorld::InsertPostProcessVolume(TSharedPtr<IInterface_PostProcessVolume>& InVolume)
+{
+	check(InVolume);
+
+	FPostProcessVolumeEntry Entry;
+	Entry.Set<TWeakPtr<IInterface_PostProcessVolume>>(InVolume.ToWeakPtr());
+	InsertPostProcessVolumeEntry(MoveTemp(Entry));
+}
+
 bool UWorld::RemovePostProcessVolume(IInterface_PostProcessVolume* InVolume)
 {
-	for (int32 Index = 0, NumVolumes = PostProcessVolumesWeak.Num(); Index < NumVolumes; )
+	TSharedPtr<IInterface_PostProcessVolume> CurrentPinned;
+	for (int32 Index = 0, NumVolumes = PostProcessVolumeEntries.Num(); Index < NumVolumes; )
 	{
-		IInterface_PostProcessVolume* CurrentVolume = PostProcessVolumesWeak[Index].Get();
+		IInterface_PostProcessVolume* CurrentVolume = ResolvePPVEntry(PostProcessVolumeEntries[Index], CurrentPinned);
 		if (!CurrentVolume) // stale entry, volume was destroyed
 		{
 			RemovePostProcessVolume(Index);
@@ -10492,7 +10564,7 @@ bool UWorld::RemovePostProcessVolume(IInterface_PostProcessVolume* InVolume)
 
 void UWorld::RemovePostProcessVolume(int32 Index)
 {
-	PostProcessVolumesWeak.RemoveAt(Index);
+	PostProcessVolumeEntries.RemoveAt(Index);
 	PostProcessVolumeCachedSizes.RemoveAt(Index);
 PRAGMA_DISABLE_DEPRECATION_WARNINGS
 	PostProcessVolumes.RemoveAt(Index);
@@ -10686,12 +10758,9 @@ void UWorld::AddPostProcessingSettings(FVector ViewLocation, FSceneView* SceneVi
 	SceneView->FinalPostProcessDebugInfo.Reset();
 #endif
 
-	for (const TWeakInterfacePtr<IInterface_PostProcessVolume>& PPVolumeWeak : PostProcessVolumesWeak)
+	for (IInterface_PostProcessVolume& PPVolume : GetPostProcessVolumeIterator())
 	{
-		if (IInterface_PostProcessVolume* PPVolume = PPVolumeWeak.Get())
-		{
-			DoPostProcessVolume(PPVolume, ViewLocation, SceneView);
-		}
+		DoPostProcessVolume(&PPVolume, ViewLocation, SceneView);
 	}
 }
 
