Fixed edge case related to recovery from sequencehistory overflow.

#Synced-CL 52054134

diff --git a/Engine/Source/Runtime/Engine/Private/Net/NetPacketNotify.cpp b/Engine/Source/Runtime/Engine/Private/Net/NetPacketNotify.cpp
--- a/Engine/Source/Runtime/Engine/Private/Net/NetPacketNotify.cpp
+++ b/Engine/Source/Runtime/Engine/Private/Net/NetPacketNotify.cpp
@@ -85,8 +85,15 @@ void FNetPacketNotify::Init(SequenceNumberT InitialInSeq, SequenceNumberT Initia
 
 void FNetPacketNotify::SetWaitForSequenceHistoryFlush()
 {
-	UE_LOG_PACKET_NOTIFY_WARNING(TEXT("FNetPacketNotify::SetWaitForSequenceHistoryFlush - Wait for ack of next OutSeq: %u"), OutSeq.Get());
-	WaitingForFlushSeqAck = OutSeq;
+	SequenceNumberT FlushSeq = OutSeq;
+	// If we're in the middle of writing a header (WriteHeader called but not yet Committed), we must wait for the next OutSeq since ackdata might have changed before it is commited.
+	// This handles the edge case where flush triggers between WriteHeader and CommitAndIncrementOutSeq.
+	if (WrittenHistoryWordCount != 0)
+	{
+		++FlushSeq;
+	}
+	UE_LOG_PACKET_NOTIFY_WARNING(TEXT("FNetPacketNotify::SetWaitForSequenceHistoryFlush - Wait for ack of next OutSeq: %u"), FlushSeq.Get());
+	WaitingForFlushSeqAck = FlushSeq;
 }
 
 void FNetPacketNotify::AckSeq(SequenceNumberT AckedSeq, bool IsAck)
