[SkinnedMeshComponent & Deformer Graph] set GeometryReadbackRequests property to protected so that derived classes can also access it

#rnx

#Synced-CL 52346140

diff --git a/Engine/Source/Runtime/Engine/Classes/Components/SkinnedMeshComponent.h b/Engine/Source/Runtime/Engine/Classes/Components/SkinnedMeshComponent.h
--- a/Engine/Source/Runtime/Engine/Classes/Components/SkinnedMeshComponent.h
+++ b/Engine/Source/Runtime/Engine/Classes/Components/SkinnedMeshComponent.h
@@ -383,7 +383,7 @@ public:
 	ENGINE_API virtual int32 GetMeshDeformerMaxLOD() const;
 
 #if WITH_EDITORONLY_DATA
-private:
+protected:
 	TArray<TUniquePtr<FMeshDeformerGeometryReadbackRequest>> GeometryReadbackRequests;
 public:
 	/**
