Do not reset external package when renaming actor changing outer in PIE packages, fixing the asymetry.

#Synced-CL 52252558

diff --git a/Engine/Source/Runtime/Engine/Private/Actor.cpp b/Engine/Source/Runtime/Engine/Private/Actor.cpp
--- a/Engine/Source/Runtime/Engine/Private/Actor.cpp
+++ b/Engine/Source/Runtime/Engine/Private/Actor.cpp
@@ -1867,7 +1867,7 @@ bool AActor::Rename( const TCHAR* InName, UObject* NewOuter, ERenameFlags Flags
 				bShouldSetPackageExternal = bChangingOuters || IsMainPackageActor();
 				// If we are changing outers always change the external package because the outer chain has an impact on the filename
 				// but if we are not changing outers only change the external flag if the actor is the main actor in a package. (this is to avoid Child Actors creating packages)
-				if (bShouldSetPackageExternal)
+				if (bShouldSetPackageExternal && !bIsPIEPackage)
 				{
 					bool bLevelPackageWasDirty = MyLevel->GetPackage()->IsDirty();
 					SetPackageExternal(false, MyLevel->IsUsingExternalActors());
