Match mangled ICH and SCS component templates against their unmangled base. The DefaultSceneRoot is specifically mangled by the ICH.

[RN] Fixed issue looking up default values for the root component when it is named DefaultSceneRoot
#Synced-CL 51986075

diff --git a/Engine/Source/Runtime/CoreUObject/Private/UObject/UObjectArchetype.cpp b/Engine/Source/Runtime/CoreUObject/Private/UObject/UObjectArchetype.cpp
--- a/Engine/Source/Runtime/CoreUObject/Private/UObject/UObjectArchetype.cpp
+++ b/Engine/Source/Runtime/CoreUObject/Private/UObject/UObjectArchetype.cpp
@@ -149,6 +149,11 @@ FORCENOINLINE UObject* FindArchetypeFromRequiredInfoImpl(const UClass* Class, co
 			{
 				Result = ArchetypeToSearch->GetClass()->FindArchetype(Class, Name);
 			}
+
+			if(!Result && (ObjectFlags & RF_ArchetypeObject) && Outer->IsA<UClass>())
+			{
+				Result = CastChecked<UClass>(Outer)->GetSuperClass()->FindArchetype(Class, Name);
+			}
 		}
 
 		if (!Result)
diff --git a/Engine/Source/Runtime/Engine/Private/BlueprintGeneratedClass.cpp b/Engine/Source/Runtime/Engine/Private/BlueprintGeneratedClass.cpp
--- a/Engine/Source/Runtime/Engine/Private/BlueprintGeneratedClass.cpp
+++ b/Engine/Source/Runtime/Engine/Private/BlueprintGeneratedClass.cpp
@@ -1572,6 +1572,30 @@ UObject* UBlueprintGeneratedClass::FindArchetype(const UClass* ArchetypeClass, c
 				Class = nullptr;
 			}
 		}
+		
+		if(Archetype == nullptr)
+		{
+			// We failed to find an archetype by name - but the ICH mangles
+			// the DefaultSceneRoot name *and* tries to cull unused templates
+			// so do a relaxed compare to find the real default scene root:
+			static const FName DefaultSceneRoot("DefaultSceneRoot");
+			static const FName DefaultSceneRootSCS("DefaultSceneRoot_GEN_VARIABLE");
+			static const FName ICHDefaultSceneRoot("ICH-DefaultSceneRoot_GEN_VARIABLE");
+			if(ArchetypeName == DefaultSceneRootSCS || ArchetypeName == ICHDefaultSceneRoot)
+			{
+				Class = const_cast<UBlueprintGeneratedClass*>(this);
+				while (Class && !Archetype)
+				{
+					// first look for an SCS template, then fall back to an unmangled DefaultSceneRoot.
+					Archetype = static_cast<UObject*>(FindObjectWithOuter(Class, ArchetypeClass, DefaultSceneRootSCS));
+					if(Archetype == nullptr)
+					{
+						Archetype = static_cast<UObject*>(FindObjectWithOuter(Class, ArchetypeClass, DefaultSceneRoot));
+					}
+					Class = Cast<UBlueprintGeneratedClass>(Class->GetSuperClass());
+				}
+			}
+		}
 	}
 
 	return Archetype;
