remove ensure for editor PIE when in stereo rendering project

#Synced-CL 52353134

#Synced-CL 52353134

diff --git a/Engine/Source/Runtime/Engine/Private/SceneView.cpp b/Engine/Source/Runtime/Engine/Private/SceneView.cpp
--- a/Engine/Source/Runtime/Engine/Private/SceneView.cpp
+++ b/Engine/Source/Runtime/Engine/Private/SceneView.cpp
@@ -1110,7 +1110,7 @@ bool FSceneView::UsesInstancedStereoShaders() const
 
 bool FSceneView::IsInstancedStereo() const
 {
-	return bIsInstancedStereoEnabled && ensure(IStereoRendering::IsStereoEyeView(*this));
+	return bIsInstancedStereoEnabled && IStereoRendering::IsStereoEyeView(*this);
 }
 
 FRDGTextureDesc FSceneView::CreateStereoTextureDesc(FIntPoint Size, EPixelFormat Format, FClearValueBinding Clear, ETextureCreateFlags Flags) const
