Enable ShowDebug Animation in Test builds. Replace two DrawDebugCanvas2DLine calls with Canvas->K2_DrawLine() so the anim node tree connector lines also render when ENABLE_DRAW_DEBUG is disabled.
#Synced-CL 52354471

diff --git a/Engine/Source/Runtime/Engine/Private/Animation/AnimInstance.cpp b/Engine/Source/Runtime/Engine/Private/Animation/AnimInstance.cpp
--- a/Engine/Source/Runtime/Engine/Private/Animation/AnimInstance.cpp
+++ b/Engine/Source/Runtime/Engine/Private/Animation/AnimInstance.cpp
@@ -1348,7 +1348,6 @@ void UAnimInstance::DisplayDebugInstance(FDisplayDebugManager& DisplayDebugManag
 
 void UAnimInstance::DisplayDebug(class UCanvas* Canvas, const FDebugDisplayInfo& DebugDisplay, float& YL, float& YPos)
 {
-#if ENABLE_DRAW_DEBUG
 	FAnimInstanceProxy& Proxy = GetProxyOnGameThread<FAnimInstanceProxy>();
 
 	float Indent = 0.f;
@@ -1438,11 +1437,11 @@ void UAnimInstance::DisplayDebug(class UCanvas* Canvas, const FDebugDisplayInfo&
 						float EndX = DisplayDebugManager.GetXPos() + CurrIndent;
 						float StartX = EndX - AttachLineLength;
 
-						//horizontal line to node
-						DrawDebugCanvas2DLine(Canvas, FVector(StartX, CurrLineYBase, 0.f), FVector(EndX, CurrLineYBase, 0.f), ActiveColor);
+						// horizontal line to node
+						Canvas->K2_DrawLine(FVector2D(StartX, CurrLineYBase), FVector2D(EndX, CurrLineYBase), 1.0f, ActiveColor);
 
-						//vertical line
-						DrawDebugCanvas2DLine(Canvas, FVector(StartX, LineStartCoord.Y, 0.f), FVector(StartX, CurrLineYBase, 0.f), ActiveColor);
+						// vertical line
+						Canvas->K2_DrawLine(FVector2D(StartX, LineStartCoord.Y), FVector2D(StartX, CurrLineYBase), 1.0f, ActiveColor);
 					}
 				}
 
@@ -1589,7 +1588,6 @@ void UAnimInstance::DisplayDebug(class UCanvas* Canvas, const FDebugDisplayInfo&
 			}
 		}
 	}
-#endif // ENABLE_DRAW_DEBUG
 }
 
 void UAnimInstance::ResetDynamics(ETeleportType InTeleportType)
