[ISM/ISKM] Moved InstancedStaticMeshComponent instance body management into a new class so that many operations can be reused to manage collision bodies for InstancedSkinnedMeshComponent as well.

Known limitation: ISKM collision is currently static and in bind pose.

#Synced-CL 52352702

#Synced-CL 52352702

diff --git a/Engine/Source/Runtime/Engine/Classes/Components/InstancedSkinnedMeshComponent.h b/Engine/Source/Runtime/Engine/Classes/Components/InstancedSkinnedMeshComponent.h
--- a/Engine/Source/Runtime/Engine/Classes/Components/InstancedSkinnedMeshComponent.h
+++ b/Engine/Source/Runtime/Engine/Classes/Components/InstancedSkinnedMeshComponent.h
@@ -21,6 +21,7 @@ class UAnimBank;
 struct FAnimBankRecordHandle;
 struct FPrimitiveSceneDesc;
 struct FBoneAttachmentSocket;
+class FInstancedMeshComponentBodies;
 
 USTRUCT()
 struct FSkinnedMeshInstanceData
@@ -53,7 +54,7 @@ struct FSkinnedMeshInstanceData
 	}
 };
 
-UCLASS(ClassGroup=Rendering, hidecategories=(Object,Activation,Collision,"Components|Activation",Physics), editinlinenew, meta=(BlueprintSpawnableComponent), Blueprintable, MinimalAPI)
+UCLASS(ClassGroup=Rendering, hidecategories=(Object,Activation,"Components|Activation"), editinlinenew, meta=(BlueprintSpawnableComponent), Blueprintable, MinimalAPI)
 class UInstancedSkinnedMeshComponent : public USkinnedMeshComponent
 {
 	GENERATED_UCLASS_BODY()
@@ -163,6 +164,28 @@ public:
 	virtual void OnUpdateTransform(EUpdateTransformFlags UpdateTransformFlags, ETeleportType Teleport) override;
 	//~ End USceneComponent Interface
 
+	//~ Begin UPrimitiveComponent Interface
+	ENGINE_API virtual UBodySetup* GetBodySetup() override;
+	ENGINE_API virtual bool ShouldCreatePhysicsState() const override;
+	ENGINE_API virtual FBodyInstance* GetBodyInstance(FName BoneName = NAME_None, bool bGetWelded = true, int32 Index = INDEX_NONE) const override;
+	ENGINE_API virtual void RecreateInstanceBody(int32 InstanceBodyIndex) override;
+	ENGINE_API virtual bool CanEditSimulatePhysics() override;
+	ENGINE_API virtual bool LineTraceComponent(FHitResult& OutHit, const FVector Start, const FVector End, const FCollisionQueryParams& Params) override;
+	ENGINE_API virtual bool SweepComponent(FHitResult& OutHit, const FVector Start, const FVector End, const FQuat& ShapeWorldRotation, const FCollisionShape& CollisionShape, bool bTraceComplex) override;
+	ENGINE_API virtual bool OverlapComponent(const FVector& Pos, const FQuat& Rot, const FCollisionShape& CollisionShape) const override;
+	ENGINE_API virtual bool ComponentOverlapComponentImpl(UPrimitiveComponent* PrimComp, const FVector Pos, const FQuat& Quat, const FCollisionQueryParams& Params) override;
+	ENGINE_API virtual bool ComponentOverlapMultiImpl(TArray<FOverlapResult>& OutOverlaps, const UWorld* InWorld, const FVector& Pos, const FQuat& Rot, ECollisionChannel TestChannel, const FComponentQueryParams& Params, const FCollisionObjectQueryParams& ObjectQueryParams) const override;
+	ENGINE_API virtual void OnActorEnableCollisionChanged() override;
+	ENGINE_API virtual Chaos::FPhysicsObject* GetPhysicsObjectById(Chaos::FPhysicsObjectId Id) const override;
+	ENGINE_API virtual TArray<Chaos::FPhysicsObject*> GetAllPhysicsObjects() const override;
+protected:
+	ENGINE_API virtual void OnCreatePhysicsState() override;
+	ENGINE_API virtual void OnDestroyPhysicsState() override;
+	ENGINE_API virtual void OnAsyncDestroyPhysicsStateBegin_GameThread(TSet<UObject*>& OutRootedObjects) override;
+	ENGINE_API virtual bool OnAsyncDestroyPhysicsState(const UE::FTimeout& Timeout) override;
+	//~ End UPrimitiveComponent Interface
+
+public:
 	//~ Begin USkinnedMeshComponent Interface
 	virtual const Nanite::FResources* GetNaniteResources() const;
 	virtual void RefreshBoneTransforms(FActorComponentTickFunction* TickFunction = nullptr) override;
@@ -332,6 +355,13 @@ protected:
 	TBitArray<> SelectedInstances;
 #endif
 
+	/** Don't create any collision when this bool is set */
+	UPROPERTY()
+	uint8 bDisableCollision : 1;
+
+	/** Per-instance physics bodies manager. */
+	TUniquePtr<FInstancedMeshComponentBodies> InstancePhysicsBodies;
+
 	bool bIsInstanceDataApplyCompleted = true;
 
 	void CreateHitProxyData(TArray<TRefCountPtr<HHitProxy>>& HitProxies);
@@ -354,6 +384,20 @@ private:
 	/** Sets up new instance data to sensible defaults, creates physics counterparts if possible. */
 	void SetupNewInstanceData(FSkinnedMeshInstanceData& InOutNewInstanceData, int32 InInstanceIndex, const FTransform3f& InInstanceTransform, int32 AnimationIndex);
 
+	/** Create all physics body instances for all instances. */
+	void CreateAllInstanceBodies();
+	/** Destroy all physics body instances. */
+	void ClearAllInstanceBodies();
+	/** Update a single instance's physics body transform. */
+	void UpdateInstanceBodyTransform(int32 InstanceIndex, const FTransform& WorldSpaceInstanceTransform, bool bTeleport);
+
+	/** Build a merged UBodySetup that aggregates collision from all bones in the physics asset. */
+	void BuildMergedBodySetup();
+
+	/** Cached merged body setup combining collision from all physics asset bones. */
+	UPROPERTY(Transient)
+	TObjectPtr<UBodySetup> MergedBodySetup;
+
 	static ENGINE_API bool ShouldForceRefPose();
 	static ENGINE_API bool ShouldUseAnimationBounds();
 
diff --git a/Engine/Source/Runtime/Engine/Classes/Components/InstancedStaticMeshComponent.h b/Engine/Source/Runtime/Engine/Classes/Components/InstancedStaticMeshComponent.h
--- a/Engine/Source/Runtime/Engine/Classes/Components/InstancedStaticMeshComponent.h
+++ b/Engine/Source/Runtime/Engine/Classes/Components/InstancedStaticMeshComponent.h
@@ -34,6 +34,7 @@ class FInstancedShadowMap2D;
 class FStaticMeshInstanceData;
 class FISMInstanceUpdateChangeSet;
 struct FInstanceUpdateComponentDesc;
+class FInstancedMeshComponentBodies;
 
 USTRUCT()
 struct FInstancedStaticMeshInstanceData
@@ -485,8 +486,6 @@ private:
 	/** Current cached conservativ bounds. */
 	FBox CachedConservativeInstanceBounds;
 
-	/** Payload used by asynchronous destruction of physics state (see OnAsyncDestroyPhysicsState). */
-	TArray<FBodyInstance*> AsyncDestroyPhysicsStatePayload;
 public:
 
 #if WITH_EDITOR
@@ -509,7 +508,11 @@ public:
 	 *  for deselecting instances. */
 	EInstanceDeletionReason DeletionState = EInstanceDeletionReason::NotDeleting;
 #endif
-	/** Physics representation of the instance bodies. */
+
+	/** Returns the per-instance physics bodies array. */
+	ENGINE_API const TArray<FBodyInstance*>& GetInstanceBodies() const;
+
+	UE_DEPRECATED(5.8, "This member is no longer used. Use GetInstanceBodies() instead.")
 	TArray<FBodyInstance*> InstanceBodies;
 
 	//~ Begin UActorComponent Interface
@@ -693,6 +696,9 @@ protected:
 	int32 NumPendingLightmaps;
 
 	FPrimitiveInstanceDataManager PrimitiveInstanceDataManager;
+	
+	/** Per-instance physics bodies manager. */
+	TUniquePtr<FInstancedMeshComponentBodies> InstancePhysicsBodies;
 
 	void InvalidateCachedBounds();
 
@@ -750,8 +756,9 @@ protected:
 	ENGINE_API virtual void GetNavigationPerInstanceTransforms(const FBox& AreaBox, TArray<FTransform>& InstanceData) const;
 
 	/** Initializes the body instance for the specified instance of the static mesh. */
+	UE_DEPRECATED(5.8, "Use InstancePhysicsBodies instead.")
 	ENGINE_API void InitInstanceBody(int32 InstanceIdx, FBodyInstance* InBodyInstance);
-
+	UE_DEPRECATED(5.8, "Use InstancePhysicsBodies instead.")
 	ENGINE_API void InitInstanceBody(int32 InstanceIdx, FBodyInstance* InstanceBodyInstance, FBodyInstance* ReferenceBodyInstance, bool bRuntimeBodyInstance, UBodySetup* BodySetup);
 
 	/** The mappings for all the instances of this component. */
diff --git a/Engine/Source/Runtime/Engine/Private/Components/InstancedSkinnedMeshComponent.cpp b/Engine/Source/Runtime/Engine/Private/Components/InstancedSkinnedMeshComponent.cpp
--- a/Engine/Source/Runtime/Engine/Private/Components/InstancedSkinnedMeshComponent.cpp
+++ b/Engine/Source/Runtime/Engine/Private/Components/InstancedSkinnedMeshComponent.cpp
@@ -8,6 +8,7 @@
 #include "Elements/Interfaces/TypedElementWorldInterface.h"
 #include "Elements/SMInstance/SMInstanceElementData.h" // For SMInstanceElementDataUtil::SMInstanceElementsEnabled
 #include "Elements/SMInstance/SMInstanceElementId.h"
+#include "Engine/OverlapResult.h"
 #include "Engine/StaticMesh.h"
 #include "HitProxies.h"
 #include "NaniteSceneProxy.h"
@@ -28,6 +29,15 @@
 #include "InstancedSkinnedMeshComponentHelper.h"
 #include "PrimitiveSceneDesc.h"
 #include "InstancedSkinnedMeshSceneProxy.h"
+#include "InstancedMeshComponentBodies.h"
+#include "AnimationRuntime.h"
+#include "PhysicsEngine/PhysicsAsset.h"
+#include "PhysicsEngine/SkeletalBodySetup.h"
+#include "WorldCollision.h"
+#include "PhysicsEngine/BodySetup.h"
+#include "Chaos/PhysicsObjectInterface.h"
+#include "PhysicsProxy/SingleParticlePhysicsProxy.h"
+#include "Physics/Experimental/PhysScene_Chaos.h"
 
 #if WITH_EDITOR
 #include "WorldPartition/HLOD/HLODHashBuilder.h"
@@ -68,14 +78,18 @@ UInstancedSkinnedMeshComponent::UInstancedSkinnedMeshComponent(FVTableHelper& He
 : Super(Helper)
 , bInheritPerInstanceData(false)
 , bUseGpuLodSelection(true)
+, bDisableCollision(false)
+, InstancePhysicsBodies(MakeUnique<FInstancedMeshComponentBodies>(this))
 , InstanceDataManager(this)
 {
 }
 
-UInstancedSkinnedMeshComponent::UInstancedSkinnedMeshComponent(const FObjectInitializer& ObjectInitializer) 
-: Super(ObjectInitializer) 
+UInstancedSkinnedMeshComponent::UInstancedSkinnedMeshComponent(const FObjectInitializer& ObjectInitializer)
+: Super(ObjectInitializer)
 , bInheritPerInstanceData(false)
 , bUseGpuLodSelection(true)
+, bDisableCollision(false)
+, InstancePhysicsBodies(MakeUnique<FInstancedMeshComponentBodies>(this))
 , InstanceDataManager(this)
 {
 }
@@ -401,9 +415,10 @@ void UInstancedSkinnedMeshComponent::SetupNewInstanceData(FSkinnedMeshInstanceDa
 	InOutNewInstanceData.Transform = InInstanceTransform;
 	InOutNewInstanceData.AnimationIndex = InAnimationIndex;
 
-	if (bPhysicsStateCreated)
+	if (bPhysicsStateCreated && !IsAsyncDestroyPhysicsStateRunning())
 	{
-		// ..
+		const FTransform WorldTransform = FTransform(InInstanceTransform) * GetComponentTransform();
+		InstancePhysicsBodies->Insert(InInstanceIndex, WorldTransform);
 	}
 }
 
@@ -995,8 +1010,20 @@ FPrimitiveSceneProxy* UInstancedSkinnedMeshComponent::CreateSceneProxy()
 
 void UInstancedSkinnedMeshComponent::OnUpdateTransform(EUpdateTransformFlags UpdateTransformFlags, ETeleportType Teleport)
 {
-	Super::OnUpdateTransform(UpdateTransformFlags, Teleport);
+	Super::OnUpdateTransform(UpdateTransformFlags | EUpdateTransformFlags::SkipPhysicsUpdate, Teleport);
 	InstanceDataManager.PrimitiveTransformChanged();
+
+	const bool bTeleport = TeleportEnumToFlag(Teleport);
+	const bool bUpdateBodies = bPhysicsStateCreated && !(EUpdateTransformFlags::SkipPhysicsUpdate & UpdateTransformFlags);
+
+	if (bUpdateBodies)
+	{
+		for (int32 InstanceIndex = 0; InstanceIndex < InstanceData.Num(); InstanceIndex++)
+		{
+			const FTransform NewInstanceTransform(FTransform(InstanceData[InstanceIndex].Transform) * GetComponentTransform());
+			UpdateInstanceBodyTransform(InstanceIndex, NewInstanceTransform, bTeleport);
+		}
+	}
 }
 
 #if WITH_EDITOR
@@ -1204,6 +1231,7 @@ FBoxSphereBounds UInstancedSkinnedMeshComponent::CalcBounds(const FTransform& Lo
 
 void UInstancedSkinnedMeshComponent::SetSkinnedAssetCallback()
 {
+	MergedBodySetup = nullptr;
 	MarkRenderStateDirty();
 }
 
@@ -1533,6 +1561,11 @@ void UInstancedSkinnedMeshComponent::ClearInstances()
 	SelectedInstances.Empty();
 #endif
 	InstanceDataManager.ClearInstances();
+
+	if (bPhysicsStateCreated)
+	{
+		ClearAllInstanceBodies();
+	}
 }
 
 struct HSkinnedMeshInstance : public HHitProxy
@@ -1642,13 +1675,13 @@ bool UInstancedSkinnedMeshComponent::RemoveInstanceInternal(int32 InstanceIndex,
 		return false;
 	}
 	InstanceDataManager.RemoveAtSwap(InstanceIndex);
-	
+
 	// remove instance
 	if (!bInstanceAlreadyRemoved)
 	{
 		InstanceData.RemoveAtSwap(InstanceIndex, EAllowShrinking::No);
 	}
-	
+
 	if (InstanceCustomData.IsValidIndex(InstanceIndex * NumCustomDataFloats))
 	{
 		InstanceCustomData.RemoveAtSwap(InstanceIndex * NumCustomDataFloats, NumCustomDataFloats, EAllowShrinking::No);
@@ -1692,6 +1725,12 @@ bool UInstancedSkinnedMeshComponent::RemoveInstanceInternal(int32 InstanceIndex,
 		}
 	}
 
+	// update the physics state
+	if (bPhysicsStateCreated && InstancePhysicsBodies->IsValidIndex(InstanceIndex))
+	{
+		InstancePhysicsBodies->RemoveAtSwap(InstanceIndex);
+	}
+
 	return true;
 }
 
@@ -1741,3 +1780,245 @@ TSharedPtr<FInstanceDataSceneProxy, ESPMode::ThreadSafe> UInstancedSkinnedMeshCo
 
 	return MakeShared<FInstanceDataSceneProxy, ESPMode::ThreadSafe>(MoveTemp(InstanceSceneDataBuffers));
 }
+
+// ---- Physics / Collision ----
+
+UBodySetup* UInstancedSkinnedMeshComponent::GetBodySetup()
+{
+	if (!MergedBodySetup)
+	{
+		BuildMergedBodySetup();
+	}
+	return MergedBodySetup;
+}
+
+void UInstancedSkinnedMeshComponent::BuildMergedBodySetup()
+{
+	MergedBodySetup = nullptr;
+
+	UPhysicsAsset* const PhysicsAsset = GetPhysicsAsset();
+	USkinnedAsset* const Asset = GetSkinnedAsset();
+	if (!Asset || !PhysicsAsset || PhysicsAsset->SkeletalBodySetups.IsEmpty())
+	{
+		return;
+	}
+
+	const FReferenceSkeleton& RefSkeleton = Asset->GetRefSkeleton();
+	const TArray<FTransform>& BoneSpaceTransforms = RefSkeleton.GetRawRefBonePose();
+	const int32 NumBones = BoneSpaceTransforms.Num();
+
+	// Pre-compute component-space ref-pose transforms for all bones
+	TArray<FTransform> CachedComponentSpaceTransforms;
+	TArray<bool> CachedTransformReady;
+	CachedComponentSpaceTransforms.SetNumUninitialized(NumBones);
+	CachedTransformReady.SetNumZeroed(NumBones);
+
+	MergedBodySetup = NewObject<UBodySetup>(this);
+
+	// Copy properties from the first valid body setup (collision trace flag, phys material, etc.)
+	for (USkeletalBodySetup* SkelBodySetup : PhysicsAsset->SkeletalBodySetups)
+	{
+		if (SkelBodySetup)
+		{
+			MergedBodySetup->CopyBodyPropertiesFrom(SkelBodySetup);
+			break;
+		}
+	}
+
+	// Clear geometry — we'll rebuild it with transformed elements from all bones
+	MergedBodySetup->AggGeom.EmptyElements();
+
+	for (USkeletalBodySetup* SkelBodySetup : PhysicsAsset->SkeletalBodySetups)
+	{
+		if (!SkelBodySetup)
+		{
+			continue;
+		}
+
+		const int32 BoneIndex = RefSkeleton.FindBoneIndex(SkelBodySetup->BoneName);
+		if (BoneIndex == INDEX_NONE)
+		{
+			continue;
+		}
+
+		const FTransform& BoneTransform = FAnimationRuntime::GetComponentSpaceTransformWithCache(
+			RefSkeleton, BoneSpaceTransforms, BoneIndex,
+			CachedComponentSpaceTransforms, CachedTransformReady);
+
+		MergedBodySetup->AddCollisionFrom(SkelBodySetup->AggGeom);
+
+		// Transform the newly appended elements by the bone's component-space ref-pose transform
+		for (int32 ShapeType = 0; ShapeType < EAggCollisionShape::Unknown; ++ShapeType)
+		{
+			const EAggCollisionShape::Type Type = static_cast<EAggCollisionShape::Type>(ShapeType);
+			const int32 PrevCount = SkelBodySetup->AggGeom.GetElementCount(Type);
+			const int32 TotalCount = MergedBodySetup->AggGeom.GetElementCount(Type);
+			for (int32 ElemIndex = TotalCount - PrevCount; ElemIndex < TotalCount; ++ElemIndex)
+			{
+				FKShapeElem* Elem = MergedBodySetup->AggGeom.GetElement(Type, ElemIndex);
+				if (Elem)
+				{
+					MergedBodySetup->AggGeom.VisitShapeAndContainer(*Elem, [&BoneTransform](auto& TypedElem, auto&)
+					{
+						if constexpr (requires { TypedElem.SetTransform(FTransform::Identity); })
+						{
+							TypedElem.SetTransform(TypedElem.GetTransform() * BoneTransform);
+						}
+					});
+				}
+			}
+		}
+	}
+
+	MergedBodySetup->bNeverNeedsCookedCollisionData = true;
+	MergedBodySetup->CreatePhysicsMeshes();
+}
+
+bool UInstancedSkinnedMeshComponent::ShouldCreatePhysicsState() const
+{
+	return !bDisableCollision && (IsRegistered() || (IsPreRegistering() && AllowsAsyncPhysicsStateCreation())) && !IsBeingDestroyed() && GetSkinnedAsset() && IsCollisionEnabled();
+}
+
+void UInstancedSkinnedMeshComponent::CreateAllInstanceBodies()
+{
+	const int32 NumInstances = InstanceData.Num();
+	if (GetBodySetup() == nullptr)
+	{
+		UE_LOG(LogInstancedSkinnedMesh, Warning, TEXT("Instance Skinned Mesh Component unable to create InstanceBodies!"));
+		return;
+	}
+	
+	TArray<FTransform> WorldTransforms;
+	WorldTransforms.SetNumUninitialized(NumInstances);
+	for (int32 InstanceIndex = 0; InstanceIndex < NumInstances; ++InstanceIndex)
+	{
+		WorldTransforms[InstanceIndex] = FTransform(InstanceData[InstanceIndex].Transform) * GetComponentTransform();
+	}
+
+	InstancePhysicsBodies->CreateAll(WorldTransforms);	
+}
+
+void UInstancedSkinnedMeshComponent::ClearAllInstanceBodies()
+{
+	InstancePhysicsBodies->ClearAll();
+}
+
+void UInstancedSkinnedMeshComponent::UpdateInstanceBodyTransform(int32 InstanceIndex, const FTransform& WorldSpaceInstanceTransform, bool bTeleport)
+{
+	check(bPhysicsStateCreated);
+
+	InstancePhysicsBodies->UpdateTransform(InstanceIndex, WorldSpaceInstanceTransform, bTeleport);
+}
+
+void UInstancedSkinnedMeshComponent::OnCreatePhysicsState()
+{
+	check(InstancePhysicsBodies->Num() == 0);
+
+	FPhysScene* PhysScene = GetWorld()->GetPhysicsScene();
+	if (!PhysScene)
+	{
+		return;
+	}
+
+	CreateAllInstanceBodies();
+	USceneComponent::OnCreatePhysicsState();
+}
+
+void UInstancedSkinnedMeshComponent::OnDestroyPhysicsState()
+{
+	USceneComponent::OnDestroyPhysicsState();
+	ClearAllInstanceBodies();
+}
+
+void UInstancedSkinnedMeshComponent::OnAsyncDestroyPhysicsStateBegin_GameThread(TSet<UObject*>& OutRootedObjects)
+{
+	USceneComponent::OnAsyncDestroyPhysicsStateBegin_GameThread(OutRootedObjects);
+
+	if (BodyInstance.bNotifyRigidBodyCollision)
+	{
+		if (FPhysScene* PhysScene = GetWorld()->GetPhysicsScene())
+		{
+			PhysScene->UnRegisterForCollisionEvents(this);
+		}
+	}
+
+	InstancePhysicsBodies->BeginAsyncDestroy();
+}
+
+bool UInstancedSkinnedMeshComponent::OnAsyncDestroyPhysicsState(const UE::FTimeout& Timeout)
+{
+	USceneComponent::OnAsyncDestroyPhysicsState(Timeout);
+
+	InstancePhysicsBodies->AsyncDestroy();
+	return true;
+}
+
+FBodyInstance* UInstancedSkinnedMeshComponent::GetBodyInstance(FName BoneName, bool bGetWelded, int32 Index) const
+{
+	return InstancePhysicsBodies->GetBodyInstance(Index);
+}
+
+void UInstancedSkinnedMeshComponent::RecreateInstanceBody(int32 InstanceBodyIndex)
+{
+	if (InstanceData.IsValidIndex(InstanceBodyIndex))
+	{
+		const FTransform InstanceTM = FTransform(InstanceData[InstanceBodyIndex].Transform) * GetComponentTransform();
+		InstancePhysicsBodies->Recreate(InstanceBodyIndex, InstanceTM);
+	}
+}
+
+Chaos::FPhysicsObject* UInstancedSkinnedMeshComponent::GetPhysicsObjectById(Chaos::FPhysicsObjectId Id) const
+{
+	return InstancePhysicsBodies->GetPhysicsObjectById(Id);
+}
+
+TArray<Chaos::FPhysicsObject*> UInstancedSkinnedMeshComponent::GetAllPhysicsObjects() const
+{
+	return InstancePhysicsBodies->GetAllPhysicsObjects();
+}
+
+bool UInstancedSkinnedMeshComponent::CanEditSimulatePhysics()
+{
+	return false;
+}
+
+bool UInstancedSkinnedMeshComponent::LineTraceComponent(FHitResult& OutHit, const FVector Start, const FVector End, const FCollisionQueryParams& Params)
+{
+	return InstancePhysicsBodies->LineTrace(OutHit, Start, End, Params.bTraceComplex, Params.bReturnPhysicalMaterial);
+}
+
+bool UInstancedSkinnedMeshComponent::SweepComponent(FHitResult& OutHit, const FVector Start, const FVector End, const FQuat& ShapeWorldRotation, const FCollisionShape& CollisionShape, bool bTraceComplex)
+{
+	return InstancePhysicsBodies->Sweep(OutHit, Start, End, ShapeWorldRotation, CollisionShape, bTraceComplex);
+}
+
+bool UInstancedSkinnedMeshComponent::OverlapComponent(const FVector& Pos, const FQuat& Rot, const FCollisionShape& CollisionShape) const
+{
+	return InstancePhysicsBodies->Overlap(Pos, Rot, CollisionShape);
+}
+
+bool UInstancedSkinnedMeshComponent::ComponentOverlapComponentImpl(UPrimitiveComponent* PrimComp, const FVector Pos, const FQuat& Quat, const FCollisionQueryParams& Params)
+{
+	if (PrimComp->IsA<UInstancedSkinnedMeshComponent>())
+	{
+		UE_LOG(LogCollision, Warning, TEXT("ComponentOverlapComponent : (%s) Does not support InstancedSkinnedMesh vs InstancedSkinnedMesh overlap"), *PrimComp->GetPathName());
+		return false;
+	}
+
+	return InstancePhysicsBodies->OverlapComponent(PrimComp->GetBodyInstance(), Pos, Quat);
+}
+
+bool UInstancedSkinnedMeshComponent::ComponentOverlapMultiImpl(TArray<FOverlapResult>& OutOverlaps, const class UWorld* InWorld, const FVector& Pos, const FQuat& Rot, ECollisionChannel TestChannel, const struct FComponentQueryParams& Params, const struct FCollisionObjectQueryParams& ObjectQueryParams) const
+{
+	OutOverlaps.Reset();
+
+	return InstancePhysicsBodies->OverlapMulti(OutOverlaps, Pos, Rot, TestChannel, Params, ObjectQueryParams);
+}
+
+void UInstancedSkinnedMeshComponent::OnActorEnableCollisionChanged()
+{
+	if (IsRegistered() && GetSkinnedAsset())
+	{
+		RecreatePhysicsState();
+	}
+}
diff --git a/Engine/Source/Runtime/Engine/Private/Components/PrimitiveComponent.cpp b/Engine/Source/Runtime/Engine/Private/Components/PrimitiveComponent.cpp
--- a/Engine/Source/Runtime/Engine/Private/Components/PrimitiveComponent.cpp
+++ b/Engine/Source/Runtime/Engine/Private/Components/PrimitiveComponent.cpp
@@ -3760,7 +3760,7 @@ bool UPrimitiveComponent::ComponentOverlapComponentImpl(class UPrimitiveComponen
 		{
 			if (ensure(ThisBodyInstance->IsValidBodyInstance()))
 			{
-				return ThisBodyInstance->OverlapTestForBodies(Pos, Quat, InstancedStaticMesh->InstanceBodies);
+				return ThisBodyInstance->OverlapTestForBodies(Pos, Quat, InstancedStaticMesh->GetInstanceBodies());
 			}
 		}
 		return false;
diff --git a/Engine/Source/Runtime/Engine/Private/InstancedMeshComponentBodies.cpp b/Engine/Source/Runtime/Engine/Private/InstancedMeshComponentBodies.cpp
new file mode 100644
--- /dev/null
+++ b/Engine/Source/Runtime/Engine/Private/InstancedMeshComponentBodies.cpp
@@ -0,0 +1,438 @@
+// Copyright Epic Games, Inc. All Rights Reserved.
+
+#include "InstancedMeshComponentBodies.h"
+#include "PhysicsEngine/BodySetup.h"
+#include "Components/PrimitiveComponent.h"
+#include "Engine/World.h"
+#include "Chaos/PhysicsObjectInterface.h"
+#include "PhysicsProxy/SingleParticlePhysicsProxy.h"
+#include "CollisionQueryParams.h"
+
+FInstancedMeshComponentBodies::FInstancedMeshComponentBodies(UPrimitiveComponent* InOwner)
+	: Owner(InOwner)
+{
+	check(Owner);
+}
+
+FInstancedMeshComponentBodies::~FInstancedMeshComponentBodies()
+{
+	if (!ensure(!IsAsyncDestroying()))
+	{
+		AsyncDestroy();
+	}
+	ClearAll();
+}
+
+void FInstancedMeshComponentBodies::InitInstanceBody(FBodyInstance* Body, int32 Idx, const FTransform& Transform, bool bRuntime)
+{
+	PRAGMA_DISABLE_DEPRECATION_WARNINGS
+	InitInstanceBody(Body, &Owner->BodyInstance, Idx, Transform, bRuntime, Owner->GetBodySetup());
+	PRAGMA_ENABLE_DEPRECATION_WARNINGS
+}
+
+void FInstancedMeshComponentBodies::InitInstanceBody(
+	FBodyInstance* Body, FBodyInstance* ReferenceBody, int32 Idx, const FTransform& Transform,
+	bool bRuntime, UBodySetup* BodySetup)
+{
+	check(Body);
+	check(ReferenceBody);
+	check(BodySetup);
+
+	if (bRuntime)
+	{
+		Body->CopyRuntimeBodyInstancePropertiesFrom(ReferenceBody);
+	}
+	else
+	{
+		Body->CopyBodyInstancePropertiesFrom(ReferenceBody);
+	}
+
+	Body->InstanceBodyIndex = Idx;
+	Body->bSimulatePhysics = false;
+	Body->bAutoWeld = false;
+	Body->InitBody(BodySetup, Transform, Owner, Owner->GetWorld()->GetPhysicsScene(), nullptr);
+}
+
+void FInstancedMeshComponentBodies::CreateAll(TConstArrayView<FTransform> WorldTransforms)
+{
+	check(Bodies.Num() == 0);
+
+	const int32 NumInstances = WorldTransforms.Num();
+
+	UBodySetup* BodySetup = Owner->GetBodySetup();
+	if (!BodySetup)
+	{
+		Bodies.AddZeroed(NumInstances);
+		return;
+	}
+
+	FBodyInstance& ReferenceBody = Owner->BodyInstance;
+	FPhysScene* PhysScene = Owner->GetWorld()->GetPhysicsScene();
+	const EComponentMobility::Type Mobility = Owner->Mobility;
+
+	if (!ReferenceBody.GetOverrideWalkableSlopeOnInstance())
+	{
+		ReferenceBody.SetWalkableSlopeOverride(BodySetup->WalkableSlopeOverride, false);
+	}
+
+	Bodies.SetNumUninitialized(NumInstances);
+
+	TArray<FBodyInstance*> SanitizedBodies;
+	SanitizedBodies.Reserve(NumInstances);
+
+	TArray<FTransform> StaticTransforms;
+	StaticTransforms.Reserve(NumInstances);
+
+	for (int32 InstanceIndex = 0; InstanceIndex < NumInstances; ++InstanceIndex)
+	{
+		const FTransform& InstanceTM = WorldTransforms[InstanceIndex];
+		if (InstanceTM.GetScale3D().IsNearlyZero())
+		{
+			Bodies[InstanceIndex] = nullptr;
+		}
+		else
+		{
+			FBodyInstance* Instance = new FBodyInstance;
+
+			SanitizedBodies.Add(Instance);
+			Bodies[InstanceIndex] = Instance;
+			Instance->CopyBodyInstancePropertiesFrom(&ReferenceBody);
+			Instance->InstanceBodyIndex = InstanceIndex;
+			Instance->bAutoWeld = false;
+			Instance->bSimulatePhysics = false;
+
+			if (Mobility == EComponentMobility::Movable)
+			{
+				Instance->InitBody(BodySetup, InstanceTM, Owner, PhysScene, nullptr);
+			}
+			else
+			{
+				StaticTransforms.Add(InstanceTM);
+			}
+		}
+	}
+
+	if (SanitizedBodies.Num() > 0 && Mobility != EComponentMobility::Movable)
+	{
+		FBodyInstance::InitStaticBodies(MoveTemp(SanitizedBodies), MoveTemp(StaticTransforms), BodySetup, Owner, PhysScene);
+	}
+}
+
+void FInstancedMeshComponentBodies::ClearAll()
+{
+	for (FBodyInstance*& Instance : Bodies)
+	{
+		if (Instance)
+		{
+			Instance->TermBody();
+			delete Instance;
+		}
+	}
+
+	Bodies.Empty();
+}
+
+void FInstancedMeshComponentBodies::UpdateTransform(
+	int32 Idx, const FTransform& WorldTransform, bool bTeleport)
+{
+	if (!Bodies.IsValidIndex(Idx))
+	{
+		return;
+	}
+
+	FBodyInstance*& Body = Bodies[Idx];
+	UBodySetup* BodySetup = Owner->GetBodySetup();
+
+	if (BodySetup == nullptr || WorldTransform.GetScale3D().IsNearlyZero())
+	{
+		if (Body)
+		{
+			Body->TermBody();
+			delete Body;
+			Body = nullptr;
+		}
+	}
+	else
+	{
+		if (Body)
+		{
+			Body->SetBodyTransform(WorldTransform, TeleportFlagToEnum(bTeleport));
+			Body->UpdateBodyScale(WorldTransform.GetScale3D());
+		}
+		else
+		{			
+			Body = new FBodyInstance();
+			InitInstanceBody(Body, Idx, WorldTransform, false);			
+		}
+	}
+}
+
+void FInstancedMeshComponentBodies::RemoveAt(int32 Idx)
+{
+	if (!Bodies.IsValidIndex(Idx))
+	{
+		return;
+	}
+
+	if (FBodyInstance* Body = Bodies[Idx])
+	{
+		Body->TermBody();
+		delete Body;
+	}
+	
+	Bodies.RemoveAt(Idx);
+
+	for (int32 Index = Idx; Index < Bodies.Num(); ++Index)
+	{
+		if (Bodies[Index])
+		{
+			Bodies[Index]->InstanceBodyIndex = Index;
+		}
+	}
+}
+
+void FInstancedMeshComponentBodies::RemoveAtSwap(int32 Idx)
+{
+	if (!Bodies.IsValidIndex(Idx))
+	{
+		return;
+	}
+
+	if (FBodyInstance* Body = Bodies[Idx])
+	{
+		Body->TermBody();
+		delete Body;
+	}
+
+	Bodies.RemoveAtSwap(Idx);
+
+	if (Bodies.IsValidIndex(Idx) && Bodies[Idx])
+	{
+		Bodies[Idx]->InstanceBodyIndex = Idx;
+	}	
+}
+
+void FInstancedMeshComponentBodies::Insert(
+	int32 Idx, const FTransform& WorldTransform)
+{
+	UBodySetup* BodySetup = Owner->GetBodySetup();
+	if (BodySetup == nullptr || WorldTransform.GetScale3D().IsNearlyZero())
+	{
+		Bodies.Insert(nullptr, Idx);
+	}
+	else
+	{
+		FBodyInstance* NewBody = new FBodyInstance();
+		int32 BodyIndex = Bodies.Insert(NewBody, Idx);
+		check(Idx == BodyIndex);
+		InitInstanceBody(NewBody, BodyIndex, WorldTransform, false);
+	}
+
+	// Fix up InstanceBodyIndex for all bodies shifted by the insertion
+	for (int32 Index = Idx + 1; Index < Bodies.Num(); ++Index)
+	{
+		if (Bodies[Index])
+		{
+			Bodies[Index]->InstanceBodyIndex = Index;
+		}
+	}
+}
+
+void FInstancedMeshComponentBodies::Recreate(int32 Idx, const FTransform& InstanceTransform)
+{
+	if (!Bodies.IsValidIndex(Idx))
+	{
+		return;
+	}
+
+	FBodyInstance* OldBody = Bodies[Idx];
+	if (!OldBody)
+	{
+		return;
+	}
+
+	FBodyInstance* NewBody = nullptr;
+	if (UBodySetup* BodySetup = OldBody->GetBodySetup())
+	{
+		NewBody = new FBodyInstance;
+		NewBody->CopyRuntimeBodyInstancePropertiesFrom(OldBody);
+		NewBody->InstanceBodyIndex = Idx;
+		NewBody->bSimulatePhysics = false;
+		NewBody->bAutoWeld = false;
+		NewBody->InitBody(BodySetup, InstanceTransform, Owner, Owner->GetWorld()->GetPhysicsScene(), nullptr);
+	}
+	Bodies[Idx] = NewBody;
+
+	OldBody->TermBody();
+	delete OldBody;
+}
+
+FBodyInstance* FInstancedMeshComponentBodies::GetBodyInstance(int32 Index) const
+{
+	if (Index != INDEX_NONE && Bodies.IsValidIndex(Index))
+	{
+		return Bodies[Index];
+	}
+	return &Owner->BodyInstance;
+}
+
+Chaos::FPhysicsObject* FInstancedMeshComponentBodies::GetPhysicsObjectById(Chaos::FPhysicsObjectId Id) const
+{
+	if (!Bodies.IsValidIndex(Id) || !Bodies[Id] || !Bodies[Id]->GetPhysicsActor())
+	{
+		return nullptr;
+	}
+	return Bodies[Id]->GetPhysicsActor()->GetPhysicsObject();
+}
+
+TArray<Chaos::FPhysicsObject*> FInstancedMeshComponentBodies::GetAllPhysicsObjects() const
+{
+	TArray<Chaos::FPhysicsObject*> Objects;
+	Objects.Reserve(Bodies.Num());
+	for (int32 Index = 0; Index < Bodies.Num(); ++Index)
+	{
+		Objects.Add(GetPhysicsObjectById(Index));
+	}
+	return Objects;
+}
+
+void FInstancedMeshComponentBodies::BeginAsyncDestroy()
+{
+	check(AsyncDestroyPayload.IsEmpty());
+	AsyncDestroyPayload = MoveTemp(Bodies);
+}
+
+void FInstancedMeshComponentBodies::AsyncDestroy()
+{
+	check(Bodies.IsEmpty());
+	for (FBodyInstance*& Instance : AsyncDestroyPayload)
+	{
+		if (Instance)
+		{
+			Instance->TermBody();
+			delete Instance;
+		}
+	}
+	AsyncDestroyPayload.Empty();
+}
+
+bool FInstancedMeshComponentBodies::IsAsyncDestroying() const
+{
+	return !AsyncDestroyPayload.IsEmpty();
+}
+
+// ---- Collision queries ----
+
+bool FInstancedMeshComponentBodies::LineTrace(
+	FHitResult& OutHit, const FVector& Start, const FVector& End,
+	bool bTraceComplex, bool bReturnPhysicalMaterial) const
+{
+	bool bHaveHit = false;
+	float MinTime = UE_MAX_FLT;
+	FHitResult Hit;
+
+	for (FBodyInstance* Body : Bodies)
+	{
+		if (!Body)
+		{
+			continue;
+		}
+
+		if (Body->LineTrace(Hit, Start, End, bTraceComplex, bReturnPhysicalMaterial))
+		{
+			bHaveHit = true;
+			if (MinTime > Hit.Time)
+			{
+				MinTime = Hit.Time;
+				OutHit = Hit;
+			}
+		}
+	}
+
+	return bHaveHit;
+}
+
+bool FInstancedMeshComponentBodies::Sweep(
+	FHitResult& OutHit, const FVector& Start, const FVector& End,
+	const FQuat& Rot, const FCollisionShape& Shape, bool bTraceComplex) const
+{
+	bool bHaveHit = false;
+	FHitResult Hit;
+
+	for (FBodyInstance* Body : Bodies)
+	{
+		if (!Body)
+		{
+			continue;
+		}
+
+		if (Body->Sweep(Hit, Start, End, Rot, Shape, bTraceComplex))
+		{
+			if (!bHaveHit || Hit.Time < OutHit.Time)
+			{
+				OutHit = Hit;
+			}
+			bHaveHit = true;
+		}
+	}
+
+	return bHaveHit;
+}
+
+bool FInstancedMeshComponentBodies::Overlap(
+	const FVector& Pos, const FQuat& Rot, const FCollisionShape& Shape) const
+{
+	for (FBodyInstance* Body : Bodies)
+	{
+		if (!Body)
+		{
+			continue;
+		}
+
+		if (Body->OverlapTest(Pos, Rot, Shape))
+		{
+			return true;
+		}
+	}
+
+	return false;
+}
+
+bool FInstancedMeshComponentBodies::OverlapComponent(
+	FBodyInstance* OtherBody, const FVector& Pos, const FQuat& Quat) const
+{
+	if (!OtherBody)
+	{
+		return false;
+	}
+
+	return OtherBody->OverlapTestForBodies(Pos, Quat, Bodies);
+}
+
+bool FInstancedMeshComponentBodies::OverlapMulti(
+	TArray<FOverlapResult>& OutOverlaps, const FVector& Pos, const FQuat& Rot,
+	ECollisionChannel Channel, const FComponentQueryParams& Params,
+	const FCollisionObjectQueryParams& ObjectParams) const
+{
+	const FTransform WorldToComponent(Owner->GetComponentTransform().Inverse());
+	const FCollisionResponseParams ResponseParams(Owner->GetCollisionResponseToChannels());
+
+	FComponentQueryParams ParamsWithSelf = Params;
+	ParamsWithSelf.AddIgnoredComponent(Owner);
+
+	bool bHaveBlockingHit = false;
+	for (FBodyInstance* Body : Bodies)
+	{
+		if (!Body)
+		{
+			continue;
+		}
+
+		if (Body->OverlapMulti(OutOverlaps, Owner->GetWorld(), &WorldToComponent, Pos, Rot, Channel, ParamsWithSelf, ResponseParams, ObjectParams))
+		{
+			bHaveBlockingHit = true;
+		}
+	}
+
+	return bHaveBlockingHit;
+}
diff --git a/Engine/Source/Runtime/Engine/Private/InstancedMeshComponentBodies.h b/Engine/Source/Runtime/Engine/Private/InstancedMeshComponentBodies.h
new file mode 100644
--- /dev/null
+++ b/Engine/Source/Runtime/Engine/Private/InstancedMeshComponentBodies.h
@@ -0,0 +1,111 @@
+// Copyright Epic Games, Inc. All Rights Reserved.
+
+#pragma once
+
+#include "CoreMinimal.h"
+#include "PhysicsEngine/BodyInstance.h"
+#include "Chaos/PhysicsObject.h"
+#include "PhysicsInterfaceDeclaresCore.h"
+
+class UPrimitiveComponent;
+class UBodySetup;
+struct FCollisionShape;
+struct FCollisionQueryParams;
+struct FComponentQueryParams;
+struct FCollisionObjectQueryParams;
+struct FCollisionResponseParams;
+struct FHitResult;
+struct FOverlapResult;
+
+/**
+ * Manages per-instance physics body instances for instanced mesh components.
+ * Owns the bodies array and async destroy payload. Provides collision query helpers.
+ * Derives BodyInstance, PhysScene, Mobility, World, and collision response from the owning component.
+ */
+class FInstancedMeshComponentBodies
+{
+public:
+	explicit FInstancedMeshComponentBodies(UPrimitiveComponent* InOwner);
+	~FInstancedMeshComponentBodies();
+
+	/**
+	 * Initialize a single body: copies properties from reference, sets index, calls InitBody.
+	 * NOTE: Do not call directly. TODO: Move implementation to private overload once external
+	 * refs have been deprecated.
+	 */
+	UE_DEPRECATED(5.8, "This method will be removed in a future release.")
+	void InitInstanceBody(FBodyInstance* Body, FBodyInstance* ReferenceBody, int32 Idx,
+		const FTransform& Transform, bool bRuntime, UBodySetup* BodySetup);
+
+	/** Create all instance bodies from world-space transforms. */
+	void CreateAll(TConstArrayView<FTransform> WorldTransforms);
+
+	/** Terminate and delete all instance bodies. */
+	void ClearAll();
+
+	/** Update a single body's transform; handles zero-scale deletion and lazy creation. */
+	void UpdateTransform(int32 Idx, const FTransform& WorldTransform, bool bTeleport);
+
+	/** Remove a body at the given index. */
+	void RemoveAt(int32 Idx);
+	
+	/** Remove a body at the given index, move the last body into the empty slot. */
+	void RemoveAtSwap(int32 Idx);
+
+	/** Insert a new body at the given index. */
+	void Insert(int32 Idx, const FTransform& WorldTransform);
+
+	/** Recreate a body at the given index, copying runtime properties from the existing body. */
+	void Recreate(int32 Idx, const FTransform& InstanceTransform);
+
+	/** Get body by index, or owner's default BodyInstance if Index is INDEX_NONE. */
+	FBodyInstance* GetBodyInstance(int32 Index) const;
+
+	Chaos::FPhysicsObject* GetPhysicsObjectById(Chaos::FPhysicsObjectId Id) const;
+	TArray<Chaos::FPhysicsObject*> GetAllPhysicsObjects() const;
+
+	bool LineTrace(FHitResult& OutHit, const FVector& Start, const FVector& End,
+		bool bTraceComplex, bool bReturnPhysicalMaterial) const;
+
+	bool Sweep(FHitResult& OutHit, const FVector& Start, const FVector& End,
+		const FQuat& Rot, const FCollisionShape& Shape, bool bTraceComplex) const;
+
+	bool Overlap(const FVector& Pos, const FQuat& Rot, const FCollisionShape& Shape) const;
+
+	bool OverlapComponent(FBodyInstance* OtherBody, const FVector& Pos, const FQuat& Quat) const;
+
+	/** Multi-overlap query. Derives World, WorldToComponent, and ResponseParams from owner. */
+	bool OverlapMulti(TArray<FOverlapResult>& OutOverlaps, const FVector& Pos, const FQuat& Rot,
+		ECollisionChannel Channel, const FComponentQueryParams& Params,
+		const FCollisionObjectQueryParams& ObjectParams) const;
+
+	/** Move bodies to async payload for off-thread destruction. */
+	void BeginAsyncDestroy();
+
+	/** Destroy bodies in the async payload. */
+	void AsyncDestroy();
+
+	/** Returns true if an async destroy is in progress. */
+	bool IsAsyncDestroying() const;
+
+	const TArray<FBodyInstance*>& GetBodies() const { return Bodies; }
+	int32 Num() const { return Bodies.Num(); }
+	bool IsValidIndex(int32 Idx) const { return Bodies.IsValidIndex(Idx); }
+	bool IsEmpty() const { return Bodies.IsEmpty(); }
+	FBodyInstance*& operator[](int32 Idx) { return Bodies[Idx]; }
+	FBodyInstance* operator[](int32 Idx) const { return Bodies[Idx]; }
+
+	// Range-for iteration
+	auto begin()       { return Bodies.begin(); }
+	auto end()         { return Bodies.end(); }
+	auto begin() const { return Bodies.begin(); }
+	auto end()   const { return Bodies.end(); }
+
+private:
+	/** Initialize a single body: copies properties from owner's BodyInstance, sets index, calls InitBody. */
+	void InitInstanceBody(FBodyInstance* Body, int32 Idx, const FTransform& Transform, bool bRuntime);
+
+	UPrimitiveComponent* Owner;
+	TArray<FBodyInstance*> Bodies;
+	TArray<FBodyInstance*> AsyncDestroyPayload;
+};
diff --git a/Engine/Source/Runtime/Engine/Private/InstancedStaticMesh.cpp b/Engine/Source/Runtime/Engine/Private/InstancedStaticMesh.cpp
--- a/Engine/Source/Runtime/Engine/Private/InstancedStaticMesh.cpp
+++ b/Engine/Source/Runtime/Engine/Private/InstancedStaticMesh.cpp
@@ -44,6 +44,7 @@
 #include "Physics/Experimental/PhysScene_Chaos.h"
 #include "PhysicsProxy/SingleParticlePhysicsProxy.h"
 
+#include "InstancedMeshComponentBodies.h"
 #include "Elements/Framework/EngineElementsLibrary.h"
 #include "Elements/Interfaces/TypedElementWorldInterface.h"
 #include "UObject/UObjectIterator.h"
@@ -246,7 +247,7 @@ class FISMExecHelper : public FSelfRegisteringExec
 
 				if(ISMComponent->ShouldCreatePhysicsState())
 				{
-					const TArray<FBodyInstance*>& Bodies = ISMComponent->InstanceBodies;
+					const TArray<FBodyInstance*>& Bodies = ISMComponent->GetInstanceBodies();
 					if(Bodies.Num() == 0)
 					{
 						continue;
@@ -2350,6 +2351,7 @@ UInstancedStaticMeshComponent::UInstancedStaticMeshComponent(const FObjectInitia
 	: Super(ObjectInitializer)
 	, bIsInstanceDataApplyCompleted(true)
 	, PrimitiveInstanceDataManager(this)
+	, InstancePhysicsBodies(MakeUnique<FInstancedMeshComponentBodies>(this))
 {
 	Mobility = EComponentMobility::Movable;
 	BodyInstance.bSimulatePhysics = false;
@@ -2381,6 +2383,7 @@ UInstancedStaticMeshComponent::UInstancedStaticMeshComponent(FVTableHelper& Help
 	: Super(Helper)
 	, bIsInstanceDataApplyCompleted(true)
 	, PrimitiveInstanceDataManager(this)
+	, InstancePhysicsBodies(MakeUnique<FInstancedMeshComponentBodies>(this))
 {
 }
 
@@ -2565,13 +2568,14 @@ void UInstancedStaticMeshComponent::SendRenderInstanceData_Concurrent()
 	}
 }
 
+const TArray<FBodyInstance*>& UInstancedStaticMeshComponent::GetInstanceBodies() const
+{
+	return InstancePhysicsBodies->GetBodies();
+}
+
 FBodyInstance* UInstancedStaticMeshComponent::GetBodyInstance(FName BoneName, bool bGetWelded, int32 Index) const
 {
-	if (Index != INDEX_NONE && InstanceBodies.IsValidIndex(Index))
-	{
-		return const_cast<FBodyInstance*>(InstanceBodies[Index]);
-	}
-	return  const_cast<FBodyInstance*>(&BodyInstance); // If no index is specified we return the primitive component body instance instead
+	return InstancePhysicsBodies->GetBodyInstance(Index);
 }
 
 FPrimitiveSceneProxy* UInstancedStaticMeshComponent::CreateStaticMeshSceneProxy(Nanite::FMaterialAudit& NaniteMaterials, bool bCreateNanite)
@@ -2842,37 +2846,23 @@ void UInstancedStaticMeshComponent::InitInstanceBody(int32 InstanceIdx, FBodyIns
 	}
 
 	check(InstanceIdx < PerInstanceSMData.Num());
-	check(InstanceIdx < InstanceBodies.Num());
+	check(InstanceIdx < InstancePhysicsBodies->Num());
 	check(InstanceBodyInstance);
 	check(ReferenceBodyInstance);
 	check(BodySetup);
 
-	// Get transform of the instance
-	FTransform InstanceTransform = FTransform(PerInstanceSMData[InstanceIdx].Transform) * GetComponentTransform();
-
-	if(bRuntimeBodyInstance)
-	{
-		InstanceBodyInstance->CopyRuntimeBodyInstancePropertiesFrom(ReferenceBodyInstance);
-	}
-	else
-	{
-		InstanceBodyInstance->CopyBodyInstancePropertiesFrom(ReferenceBodyInstance);
-	}
-
-	
-	InstanceBodyInstance->InstanceBodyIndex = InstanceIdx; // Set body index 
-
-	// make sure we never enable bSimulatePhysics for ISMComps
-	InstanceBodyInstance->bSimulatePhysics = false;
-
-	// Create physics body instance.
-	InstanceBodyInstance->bAutoWeld = false;	//We don't support this for instanced meshes.
-	InstanceBodyInstance->InitBody(BodySetup, InstanceTransform, this, GetWorld()->GetPhysicsScene(), nullptr);
+	const FTransform InstanceTransform = FTransform(PerInstanceSMData[InstanceIdx].Transform) * GetComponentTransform();
+	PRAGMA_DISABLE_DEPRECATION_WARNINGS
+	InstancePhysicsBodies->InitInstanceBody(InstanceBodyInstance, ReferenceBodyInstance, InstanceIdx, InstanceTransform,
+		bRuntimeBodyInstance, BodySetup);	
+	PRAGMA_ENABLE_DEPRECATION_WARNINGS
 }
 
 void UInstancedStaticMeshComponent::InitInstanceBody(int32 InstanceIdx, FBodyInstance* InstanceBodyInstance)
 {
+	PRAGMA_DISABLE_DEPRECATION_WARNINGS
 	InitInstanceBody(InstanceIdx, InstanceBodyInstance, &BodyInstance, false, GetBodySetup());
+	PRAGMA_ENABLE_DEPRECATION_WARNINGS
 }
 
 void UInstancedStaticMeshComponent::CreateAllInstanceBodies()
@@ -2883,68 +2873,22 @@ void UInstancedStaticMeshComponent::CreateAllInstanceBodies()
 	STAT(FScopeCycleCounter Context(StatId);)
 
 	const int32 NumBodies = PerInstanceSMData.Num();
-	check(InstanceBodies.Num() == 0);
 
-	if (UBodySetup* BodySetup = GetBodySetup())
+	UBodySetup* BodySetup = GetBodySetup();
+	if (!BodySetup)
 	{
-		FPhysScene* PhysScene = GetWorld()->GetPhysicsScene();
-
-		if (!BodyInstance.GetOverrideWalkableSlopeOnInstance())
-		{
-			BodyInstance.SetWalkableSlopeOverride(BodySetup->WalkableSlopeOverride, false);
-		}
-
-		InstanceBodies.SetNumUninitialized(NumBodies);
-
-		// Sanitized array does not contain any nulls
-		TArray<FBodyInstance*> InstanceBodiesSanitized;
-		InstanceBodiesSanitized.Reserve(NumBodies);
-
-		TArray<FTransform> Transforms;
-	    Transforms.Reserve(NumBodies);
-	    for (int32 i = 0; i < NumBodies; ++i)
-	    {
-			const FTransform InstanceTM = FTransform(PerInstanceSMData[i].Transform) * GetComponentTransform();
-			if (InstanceTM.GetScale3D().IsNearlyZero())
-			{
-				InstanceBodies[i] = nullptr;
-			}
-			else
-			{
-				FBodyInstance* Instance = new FBodyInstance;
-
-				InstanceBodiesSanitized.Add(Instance);
-				InstanceBodies[i] = Instance;
-				Instance->CopyBodyInstancePropertiesFrom(&BodyInstance);
-				Instance->InstanceBodyIndex = i; // Set body index 
-				Instance->bAutoWeld = false;
-
-				// make sure we never enable bSimulatePhysics for ISMComps
-				Instance->bSimulatePhysics = false;
-
-				if (Mobility == EComponentMobility::Movable)
-				{
-					Instance->InitBody(BodySetup, InstanceTM, this, PhysScene, nullptr );
-				}
-				else
-				{
-					Transforms.Add(InstanceTM);
-				}
-			}
-	    }
-
-		if (InstanceBodiesSanitized.Num() > 0 && Mobility != EComponentMobility::Movable)
-		{
-			FBodyInstance::InitStaticBodies(MoveTemp(InstanceBodiesSanitized), MoveTemp(Transforms), BodySetup, this, GetWorld()->GetPhysicsScene());
-		}
+		UE_LOG(LogStaticMesh, Warning, TEXT("Instance Static Mesh Component unable to create InstanceBodies!"));
+		return;
 	}
-	else
+
+	TArray<FTransform> WorldTransforms;
+	WorldTransforms.SetNumUninitialized(NumBodies);
+	for (int32 i = 0; i < NumBodies; ++i)
 	{
-		// In case we get into some bad state where the BodySetup is invalid but bPhysicsStateCreated is true,
-		// issue a warning and add nullptrs to InstanceBodies.
-		UE_LOG(LogStaticMesh, Warning, TEXT("Instance Static Mesh Component unable to create InstanceBodies!"));
-		InstanceBodies.AddZeroed(NumBodies);
+		WorldTransforms[i] = FTransform(PerInstanceSMData[i].Transform) * GetComponentTransform();
 	}
+
+	InstancePhysicsBodies->CreateAll(WorldTransforms);
 }
 
 void UInstancedStaticMeshComponent::ClearAllInstanceBodies()
@@ -2953,23 +2897,14 @@ void UInstancedStaticMeshComponent::ClearAllInstanceBodies()
 	QUICK_SCOPE_CYCLE_COUNTER(STAT_UInstancedStaticMeshComponent_ClearAllInstanceBodies);
 	STAT(FScopeCycleCounter Context(StatId);)
 
-	for (FBodyInstance*& Instance : InstanceBodies)
-	{
-		if (Instance)
-		{
-			Instance->TermBody();
-			delete Instance;
-		}
-	}
-
-	InstanceBodies.Empty();
+	InstancePhysicsBodies->ClearAll();
 }
 
 
 void UInstancedStaticMeshComponent::OnCreatePhysicsState()
 {
 	TRACE_CPUPROFILER_EVENT_SCOPE(UInstancedStaticMeshComponent::OnCreatePhysicsState)
-	check(InstanceBodies.Num() == 0);
+	check(InstancePhysicsBodies->Num() == 0);
 
 	FPhysScene* PhysScene = GetWorld()->GetPhysicsScene();
 
@@ -3015,9 +2950,7 @@ void UInstancedStaticMeshComponent::OnAsyncDestroyPhysicsStateBegin_GameThread(T
 		}
 	}
 
-	// Move InstanceBodies in AsyncDestroyPhysicsStatePayload
-	check(AsyncDestroyPhysicsStatePayload.IsEmpty());
-	AsyncDestroyPhysicsStatePayload = MoveTemp(InstanceBodies);
+	InstancePhysicsBodies->BeginAsyncDestroy();
 
 	bNavigationRelevant = IsNavigationRelevant();
 }
@@ -3026,7 +2959,7 @@ void UInstancedStaticMeshComponent::OnAsyncCreatePhysicsStateEnd_GameThread()
 {
 	Super::OnAsyncCreatePhysicsStateEnd_GameThread();
 
-	for (FBodyInstance* Body : InstanceBodies)
+	for (FBodyInstance* Body : *InstancePhysicsBodies)
 	{
 		if (Body)
 		{
@@ -3040,37 +2973,18 @@ bool UInstancedStaticMeshComponent::OnAsyncDestroyPhysicsState(const UE::FTimeou
 	TRACE_CPUPROFILER_EVENT_SCOPE(UInstancedStaticMeshComponent::OnAsyncDestroyPhysicsState);
 	USceneComponent::OnAsyncDestroyPhysicsState(Timeout);
 
-	check(InstanceBodies.IsEmpty());
-	for (FBodyInstance*& Instance : AsyncDestroyPhysicsStatePayload)
-	{
-		if (Instance)
-		{
-			Instance->TermBody();
-			delete Instance;
-		}
-	}
-	AsyncDestroyPhysicsStatePayload.Empty();
+	InstancePhysicsBodies->AsyncDestroy();
 	return true;
 }
 
 Chaos::FPhysicsObject* UInstancedStaticMeshComponent::GetPhysicsObjectById(Chaos::FPhysicsObjectId Id) const
 {
-	if (!InstanceBodies.IsValidIndex(Id) || !InstanceBodies[Id] || !InstanceBodies[Id]->GetPhysicsActor())
-	{
-		return nullptr;
-	}
-	return InstanceBodies[Id]->GetPhysicsActor()->GetPhysicsObject();
+	return InstancePhysicsBodies->GetPhysicsObjectById(Id);
 }
 
 TArray<Chaos::FPhysicsObject*> UInstancedStaticMeshComponent::GetAllPhysicsObjects() const
 {
-	TArray<Chaos::FPhysicsObject*> Objects;
-	Objects.Reserve(InstanceBodies.Num());
-	for (int32 Index = 0; Index < InstanceBodies.Num(); ++Index)
-	{
-		Objects.Add(GetPhysicsObjectById(Index));
-	}
-	return Objects;
+	return InstancePhysicsBodies->GetAllPhysicsObjects();
 }
 
 bool UInstancedStaticMeshComponent::CanEditSimulatePhysics()
@@ -4147,42 +4061,15 @@ bool UInstancedStaticMeshComponent::RemoveInstanceInternal(int32 InstanceIndex,
 	const int32 LastInstanceIndex = PerInstanceSMData.Num();
 
 	// update the physics state
-	if (bPhysicsStateCreated && InstanceBodies.IsValidIndex(InstanceIndex))
+	if (bPhysicsStateCreated && InstancePhysicsBodies->IsValidIndex(InstanceIndex))
 	{
-		FBodyInstance*& InstanceBody = InstanceBodies[InstanceIndex];
+		if (bUseRemoveAtSwap)
 		{
-			// Not having a body is a valid case when our physics state cannot be created (see CreateAllInstanceBodies)
-			if(InstanceBody)
-			{
-				InstanceBody->TermBody();
-				delete InstanceBody;
-				InstanceBody = nullptr;
-			}
-
-			if (bUseRemoveAtSwap && InstanceIndex != LastInstanceIndex)
-			{
-				// swap in the last instance body if we have one
-				InstanceBodies.RemoveAtSwap(InstanceIndex);
-
-				// Update the Instance body index to the new swapped location
-				if (InstanceBodies[InstanceIndex])
-				{
-					InstanceBodies[InstanceIndex]->InstanceBodyIndex = InstanceIndex;
-				}
-			}
-			else
-			{
-				InstanceBodies.RemoveAt(InstanceIndex);
-
-				// Re-target instance indices for shifting of array.
-				for (int32 Index = InstanceIndex; Index < InstanceBodies.Num(); ++Index)
-				{
-					if (InstanceBodies[Index])
-					{
-						InstanceBodies[Index]->InstanceBodyIndex = Index;
-					}
-				}
-			}
+			InstancePhysicsBodies->RemoveAtSwap(InstanceIndex);
+		}
+		else
+		{
+			InstancePhysicsBodies->RemoveAt(InstanceIndex);
 		}
 	}
 
@@ -4366,60 +4253,18 @@ void UInstancedStaticMeshComponent::UpdateComponentTransform(EUpdateTransformFla
 
 void UInstancedStaticMeshComponent::RecreateInstanceBody(int32 InstanceBodyIndex)
 {
-	if (InstanceBodies.IsValidIndex(InstanceBodyIndex))
+	if (PerInstanceSMData.IsValidIndex(InstanceBodyIndex))
 	{
-		FBodyInstance* Body = InstanceBodies[InstanceBodyIndex];
-		if (Body)
-		{
-			const FTransform InstanceTM = FTransform(PerInstanceSMData[InstanceBodyIndex].Transform) * GetComponentTransform();
-
-			FBodyInstance* Instance = new FBodyInstance;
-
-			InitInstanceBody(InstanceBodyIndex, Instance, Body, true, Body->GetBodySetup());
-
-			InstanceBodies[InstanceBodyIndex] = Instance;
-
-			Body->TermBody();
-			delete Body;
-		}
+		const FTransform InstanceTM = FTransform(PerInstanceSMData[InstanceBodyIndex].Transform) * GetComponentTransform();
+		InstancePhysicsBodies->Recreate(InstanceBodyIndex, InstanceTM);
 	}
 }
 
 void UInstancedStaticMeshComponent::UpdateInstanceBodyTransform(int32 InstanceIndex, const FTransform& WorldSpaceInstanceTransform, bool bTeleport)
 {
 	check(bPhysicsStateCreated);
-	if (!InstanceBodies.IsValidIndex(InstanceIndex))
-	{
-		return;
-	}
-
-	FBodyInstance*& InstanceBodyInstance = InstanceBodies[InstanceIndex];
 
-	if (WorldSpaceInstanceTransform.GetScale3D().IsNearlyZero())
-	{
-		if (InstanceBodyInstance)
-		{
-			// delete BodyInstance
-			InstanceBodyInstance->TermBody();
-			delete InstanceBodyInstance;
-			InstanceBodyInstance = nullptr;
-		}
-	}
-	else
-	{
-		if (InstanceBodyInstance)
-		{
-			// Update existing BodyInstance
-			InstanceBodyInstance->SetBodyTransform(WorldSpaceInstanceTransform, TeleportFlagToEnum(bTeleport));
-			InstanceBodyInstance->UpdateBodyScale(WorldSpaceInstanceTransform.GetScale3D());
-		}
-		else
-		{
-			// create new BodyInstance
-			InstanceBodyInstance = new FBodyInstance();
-			InitInstanceBody(InstanceIndex, InstanceBodyInstance);
-		}
-	}
+	InstancePhysicsBodies->UpdateTransform(InstanceIndex, WorldSpaceInstanceTransform, bTeleport);
 }
 
 bool UInstancedStaticMeshComponent::UpdateInstanceTransform(const int32 InstanceIndex, const FTransform& NewInstanceTransform, const bool bWorldSpace, const bool bMarkRenderStateDirty, const bool bTeleport)
@@ -5024,17 +4869,8 @@ void UInstancedStaticMeshComponent::SetupNewInstanceData(FInstancedStaticMeshIns
 
 	if (bPhysicsStateCreated && !IsAsyncDestroyPhysicsStateRunning())
 	{
-		if (InInstanceTransform.GetScale3D().IsNearlyZero())
-		{
-			InstanceBodies.Insert(nullptr, InInstanceIndex);
-		}
-		else
-		{
-			FBodyInstance* NewBodyInstance = new FBodyInstance();
-			int32 BodyIndex = InstanceBodies.Insert(NewBodyInstance, InInstanceIndex);
-			check(InInstanceIndex == BodyIndex);
-			InitInstanceBody(BodyIndex, NewBodyInstance);
-		}
+		const FTransform WorldTransform = InInstanceTransform * GetComponentTransform();
+		InstancePhysicsBodies->Insert(InInstanceIndex, WorldTransform);
 	}
 }
 
@@ -5148,7 +4984,7 @@ void UInstancedStaticMeshComponent::OnActorEnableCollisionChanged()
 {
 	Super::OnActorEnableCollisionChanged();
 	
-	for (FBodyInstance* InstanceBody : InstanceBodies)
+	for (FBodyInstance* InstanceBody : *InstancePhysicsBodies)
 	{
 		if (InstanceBody)
 		{
@@ -5565,29 +5401,7 @@ extern float DebugLineLifetime;
 
 bool UInstancedStaticMeshComponent::LineTraceComponent(FHitResult& OutHit, const FVector Start, const FVector End, const FCollisionQueryParams& Params)
 {
-	bool bHaveHit = false;	
-	float MinTime = UE_MAX_FLT;
-	FHitResult Hit;
-
-	// TODO: use spatial acceleration instead
-	for (FBodyInstance* Body : InstanceBodies)
-	{
-		// Not having a body is a valid case when our physics state cannot be created (see CreateAllInstanceBodies)
-		if(!Body)
-		{
-			continue;
-		}
-
-		if (Body->LineTrace(Hit, Start, End, Params.bTraceComplex, Params.bReturnPhysicalMaterial))
-		{
-			bHaveHit = true;
-			if (MinTime > Hit.Time)
-			{
-				MinTime = Hit.Time;
-				OutHit = Hit;
-			}
-		}
-	}
+	bool bHaveHit = InstancePhysicsBodies->LineTrace(OutHit, Start, End, Params.bTraceComplex, Params.bReturnPhysicalMaterial);
 
 #if !(UE_BUILD_SHIPPING || UE_BUILD_TEST)
 	UWorld* const World = GetWorld();
@@ -5607,59 +5421,16 @@ bool UInstancedStaticMeshComponent::LineTraceComponent(FHitResult& OutHit, const
 
 bool UInstancedStaticMeshComponent::SweepComponent(FHitResult& OutHit, const FVector Start, const FVector End, const FQuat& ShapeWorldRotation, const FCollisionShape& CollisionShape, bool bTraceComplex)
 {
-	bool bHaveHit = false;
-	
-	FHitResult Hit;
-	// TODO: use spatial acceleration instead
-	for (FBodyInstance* Body : InstanceBodies)
-	{
-		// Not having a body is a valid case when our physics state cannot be created (see CreateAllInstanceBodies)
-		if(!Body)
-		{
-			continue;
-		}
-
-		if (Body->Sweep(Hit, Start, End, ShapeWorldRotation, CollisionShape, bTraceComplex))
-		{
-			if (!bHaveHit || Hit.Time < OutHit.Time)
-			{
-				OutHit = Hit;
-			}
-			bHaveHit = true;
-		}
-	}
-	return bHaveHit;
+	return InstancePhysicsBodies->Sweep(OutHit, Start, End, ShapeWorldRotation, CollisionShape, bTraceComplex);
 }
 
 bool UInstancedStaticMeshComponent::OverlapComponent(const FVector& Pos, const FQuat& Rot, const FCollisionShape& CollisionShape) const
-{	
-	for (FBodyInstance* Body : InstanceBodies)
-	{
-		// Not having a body is a valid case when our physics state cannot be created (see CreateAllInstanceBodies)
-		if(!Body)
-		{
-			continue;
-		}
-
-		if (Body->OverlapTest(Pos, Rot, CollisionShape))
-		{
-			return true;
-		}
-	}
-
-	return false;	
+{
+	return InstancePhysicsBodies->Overlap(Pos, Rot, CollisionShape);
 }
 
 bool UInstancedStaticMeshComponent::ComponentOverlapComponentImpl(class UPrimitiveComponent* PrimComp, const FVector Pos, const FQuat& Quat, const FCollisionQueryParams& Params)
 {
-	//we do not support skeletal mesh vs InstancedStaticMesh overlap test
-	// Todo Add this warning again
-	/*if (PrimComp->IsA<USkeletalMeshComponent>())
-	{
-		UE_LOG(LogCollision, Warning, TEXT("ComponentOverlapComponent : (%s) Does not support InstancedStaticMesh with Physics Asset"), *PrimComp->GetPathName());
-		return false;
-	}*/
-
 	//We do not support Instanced static meshes vs Instanced static meshes
 	if (PrimComp->IsA<UInstancedStaticMeshComponent>())
 	{
@@ -5667,40 +5438,14 @@ bool UInstancedStaticMeshComponent::ComponentOverlapComponentImpl(class UPrimiti
 		return false;
 	}
 
-	if (FBodyInstance* BI = PrimComp->GetBodyInstance())
-	{
-		return BI->OverlapTestForBodies(Pos, Quat, InstanceBodies);
-	}
-
-	return false;
+	return InstancePhysicsBodies->OverlapComponent(PrimComp->GetBodyInstance(), Pos, Quat);
 }
 
 bool UInstancedStaticMeshComponent::ComponentOverlapMultiImpl(TArray<struct FOverlapResult>& OutOverlaps, const class UWorld* InWorld, const FVector& Pos, const FQuat& Rot, ECollisionChannel TestChannel, const struct FComponentQueryParams& Params, const struct FCollisionObjectQueryParams& ObjectQueryParams) const
 {
 	OutOverlaps.Reset();
 
-	const FTransform WorldToComponent(GetComponentTransform().Inverse());
-	const FCollisionResponseParams ResponseParams(GetCollisionResponseToChannels());
-
-	FComponentQueryParams ParamsWithSelf = Params;
-	ParamsWithSelf.AddIgnoredComponent(this);
-
-	bool bHaveBlockingHit = false;
-	for (FBodyInstance* Body : InstanceBodies)
-	{
-		// Not having a body is a valid case when our physics state cannot be created (see CreateAllInstanceBodies)
-		if(!Body)
-		{
-			continue;
-		}
-
-		if (Body->OverlapMulti(OutOverlaps, InWorld, &WorldToComponent, Pos, Rot, TestChannel, ParamsWithSelf, ResponseParams, ObjectQueryParams))
-		{
-			bHaveBlockingHit = true;
-		}
-	}
-
-	return bHaveBlockingHit;
+	return InstancePhysicsBodies->OverlapMulti(OutOverlaps, Pos, Rot, TestChannel, Params, ObjectQueryParams);
 }
 
 void UInstancedStaticMeshComponent::GetNavigationData(FNavigationRelevantData& Data) const
@@ -5741,12 +5486,12 @@ void UInstancedStaticMeshComponent::GetResourceSizeEx(FResourceSizeEx& Cumulativ
 	Super::GetResourceSizeEx(CumulativeResourceSize);
 
 	// component stuff
-	CumulativeResourceSize.AddDedicatedSystemMemoryBytes(InstanceBodies.GetAllocatedSize());
-	for (int32 i=0; i < InstanceBodies.Num(); ++i)
+	CumulativeResourceSize.AddDedicatedSystemMemoryBytes(InstancePhysicsBodies->GetBodies().GetAllocatedSize());
+	for (FBodyInstance* Body : *InstancePhysicsBodies)
 	{
-		if (InstanceBodies[i] != NULL && InstanceBodies[i]->IsValidBodyInstance())
+		if (Body != nullptr && Body->IsValidBodyInstance())
 		{
-			InstanceBodies[i]->GetBodyInstanceResourceSizeEx(CumulativeResourceSize);
+			Body->GetBodyInstanceResourceSizeEx(CumulativeResourceSize);
 		}
 	}
 	CumulativeResourceSize.AddDedicatedSystemMemoryBytes(PerInstanceSMData.GetAllocatedSize());
