Fix crash in SetShaderMapsOnMaterialResources due to GC potentially destroying some proxy Parent while the function is iterating all proxies and finding those that need their uniform expression set to be recached.

To avoid this, we submit increment then fence target value on the render thread of each material in the Proxy set and their parents. This will make sure those material interfaces won't be ready for destruction and the GC will pause collecting them until SetShaderMapsOnMaterialResources_RenderThread completes.

[RN] Fix crash in SetShaderMapsOnMaterialResources due to GC potentially destroying some proxy Parent while the function is iterating all proxies and finding those that need their uniform expression set to be recached.
#Synced-CL 52295224

diff --git a/Engine/Source/Runtime/Engine/Private/Materials/MaterialShared.cpp b/Engine/Source/Runtime/Engine/Private/Materials/MaterialShared.cpp
--- a/Engine/Source/Runtime/Engine/Private/Materials/MaterialShared.cpp
+++ b/Engine/Source/Runtime/Engine/Private/Materials/MaterialShared.cpp
@@ -5734,8 +5734,6 @@ void SetShaderMapsOnMaterialResources_RenderThread(FRHICommandListImmediate& RHI
 		bUpdateFeatureLevel[Material->GetFeatureLevel()] = true;
 	}
 
-	bool bFoundAnyInitializedMaterials = false;
-
 	// Iterate through all loaded material render proxies and recache their uniform expressions if needed
 	// This search does not scale well, but is only used when uploading async shader compile results
 	for (int32 FeatureLevelIndex = 0; FeatureLevelIndex < UE_ARRAY_COUNT(bUpdateFeatureLevel); ++FeatureLevelIndex)
@@ -5745,16 +5743,14 @@ void SetShaderMapsOnMaterialResources_RenderThread(FRHICommandListImmediate& RHI
 			const ERHIFeatureLevel::Type MaterialFeatureLevel = (ERHIFeatureLevel::Type) FeatureLevelIndex;
 
 			FScopeLock Locker(&FMaterialRenderProxy::GetMaterialRenderProxyMapLock());
-			for (TSet<FMaterialRenderProxy*>::TConstIterator It(FMaterialRenderProxy::GetMaterialRenderProxyMap()); It; ++It)
+			for (FMaterialRenderProxy* MaterialProxy : FMaterialRenderProxy::GetMaterialRenderProxyMap())
 			{
-				FMaterialRenderProxy* MaterialProxy = *It;
 				const FMaterial* Material = MaterialProxy->GetMaterialNoFallback(MaterialFeatureLevel);
 
 				// Using ContainsByHash so we can pass a raw-ptr to TMap method that wants a TRefCountPtr
 				if (Material && Material->GetRenderingThreadShaderMap() && MaterialsToUpdate.ContainsByHash(GetTypeHash(Material), Material))
 				{
 					MaterialProxy->CacheUniformExpressions(RHICmdList, true);
-					bFoundAnyInitializedMaterials = true;
 				}
 			}
 		}
@@ -5762,6 +5758,8 @@ void SetShaderMapsOnMaterialResources_RenderThread(FRHICommandListImmediate& RHI
 #endif // WITH_EDITOR
 }
 
+#if WITH_EDITOR
+
 void FMaterial::SetShaderMapsOnMaterialResources(const TMap<TRefCountPtr<FMaterial>, TRefCountPtr<FMaterialShaderMap>>& MaterialsToUpdate)
 {
 	for (const auto& It : MaterialsToUpdate)
@@ -5785,8 +5783,28 @@ void FMaterial::SetShaderMapsOnMaterialResources(const TMap<TRefCountPtr<FMateri
 	{
 		SetShaderMapsOnMaterialResources_RenderThread(RHICmdList, InMaterialsToUpdate);
 	});
+
+	// The parent material of some material in GetMaterialRenderProxyMap may be destroyed before SetShaderMapsOnMaterialResources_RenderThread
+	// executes. Iterate over all materials in the proxy map and bump their parent fences so that their IsReadyForFinishDestroy() will return false
+	// until the fence executes, which is after the RT task above completes.
+	// This guarantees that the parents of all materials in GetMaterialRenderProxyMap are safe from deletion while the RT task executes.
+	FScopeLock Locker(&FMaterialRenderProxy::GetMaterialRenderProxyMapLock());
+	for (FMaterialRenderProxy* Proxy : FMaterialRenderProxy::GetMaterialRenderProxyMap())
+	{
+		UMaterialInstance* MI = Cast<UMaterialInstance>(Proxy->GetMaterialInterface());
+		while (MI)
+		{
+			if (MI->Parent)
+			{
+				MI->Parent->ParentRefFence.BeginFence();
+			}
+			MI = Cast<UMaterialInstance>(MI->Parent);
+		}
+	}
 }
 
+#endif // WITH_EDITOR
+
 FMaterialParameterValue::FMaterialParameterValue(EMaterialParameterType InType, const UE::Shader::FValue& InValue)
 {
 	switch (InType)
