[MIR] Fixes ShadingModel an subsurface color handling in the new translator.

- Shading model selection now uses the appropriate shading model depending on whether the Material propertty is set to FromMaterialExpresison.
- SubsurfaceColor properly set now by the new translator.

#rnx
#Synced-CL 52377045

diff --git a/Engine/Source/Runtime/Engine/Private/Materials/MaterialAggregate.cpp b/Engine/Source/Runtime/Engine/Private/Materials/MaterialAggregate.cpp
--- a/Engine/Source/Runtime/Engine/Private/Materials/MaterialAggregate.cpp
+++ b/Engine/Source/Runtime/Engine/Private/Materials/MaterialAggregate.cpp
@@ -198,10 +198,10 @@ static EMaterialAggregateAttributeType GetPropertyType(EMaterialProperty Propert
 static FMaterialAggregateAttribute MakeAttributeFromProperty(EMaterialProperty Property)
 { 
 	FMaterialAggregateAttribute Attribute;
-	Attribute.Name = (Property == MP_SubsurfaceColor) ? TEXT("Subsurface") : *FMaterialAttributeDefinitionMap::GetAttributeName(Property);
+	Attribute.Name = FName(FMaterialAttributeDefinitionMap::GetAttributeName(Property));
 	Attribute.Type = GetPropertyType(Property);
 	
-	//ShadingModel defaults to 0, i.e. unlit. This is overriden for the old material system via editor only data, but we can't access that here so we force the value to Lit(1.0)
+	// ShadingModel defaults to 0, i.e. unlit. This is overriden for the old material system via editor only data, but we can't access that here so we force the value to Lit(1.0)
 	Attribute.DefaultValue = (Property == MP_ShadingModel) ? FVector4f(1.0f, 0.0f, 0.0f, 0.0f) : FMaterialAttributeDefinitionMap::GetDefaultValue(Property); // Should that be changed for substrate?
 	
 	return Attribute;
diff --git a/Engine/Source/Runtime/Engine/Private/Materials/MaterialIRModuleBuilder.cpp b/Engine/Source/Runtime/Engine/Private/Materials/MaterialIRModuleBuilder.cpp
--- a/Engine/Source/Runtime/Engine/Private/Materials/MaterialIRModuleBuilder.cpp
+++ b/Engine/Source/Runtime/Engine/Private/Materials/MaterialIRModuleBuilder.cpp
@@ -37,6 +37,7 @@
 #include "PostProcess/PostProcessMaterialInputs.h"
 #include "RenderUtils.h"
 #include "SubstrateTranslatorCommon.h"
+#include "Engine/SubsurfaceProfile.h"
 
 static TAutoConsoleVariable<int32> CVarMaterialIRDebugDumpLevel(
 	TEXT("r.Material.Translator.DebugDump"),
@@ -184,7 +185,7 @@ struct FMaterialIRModuleBuilderImpl
 
 	void Step_PushRootExpressionDependencies()
 	{
-		//If we are processing from the Material Attributes output node, push that dependency here.
+		// If we are processing from the Material Attributes output node, push that dependency here.
 		if (Builder->MaterialInterface->GetMaterial()->bUseMaterialAttributes)
 		{
 			FMaterialInputDescription Input;
@@ -202,6 +203,12 @@ struct FMaterialIRModuleBuilderImpl
 					continue;
 				}
 
+				// Skip evaluating shading model property if not set to "From Expression"
+				if (Property == MP_ShadingModel && !Emitter.MaterialInterface->IsShadingModelFromMaterialExpression())
+				{
+					continue;
+				}
+
 				// Read the material input associated to this property
 				FMaterialInputDescription Input;
 				GetExpressionInputDescription(Property, Input);
@@ -734,9 +741,17 @@ struct FMaterialIRModuleBuilderImpl
 				FlowValueThroughConnection(InputDesc.Input);
 				PropertyValue = MIR::Internal::FetchValueFromExpressionInput(this, InputDesc.Input);
 			}
-
-			if (PropertyValue)
+			
+			// Determine the value to assign to this material attribute "SetMaterialOutput" instruction
+			if (Property == MP_ShadingModel && !Emitter.MaterialInterface->IsShadingModelFromMaterialExpression())
+			{
+				// Don't use shading model coming from expressions if the material states a shading model explicitly, intead of setting it "FromMaterialExpression".
+				PropertyValue = Emitter.ConstantInt(Emitter.MaterialInterface->GetShadingModels().GetFirstShadingModel());
+			}
+			else if (PropertyValue)
 			{
+				// Case when the value to assigne to the material attribute is given.
+
 				// If this property is the emissive color and we're previewing the material, apply gamma correction to the previewed value.
 				if (Property == MP_EmissiveColor && PreviewInput.IsConnected())
 				{
@@ -753,7 +768,7 @@ struct FMaterialIRModuleBuilderImpl
 				}
 				else if (Property == MP_FrontMaterial)
 				{
-					// Just check that we get a struct for now. 
+					// Just check that we get a struct for now.
 					// SUBSTRATE_TODO We should probably add a proper SubstrateData struct enum?
 					if (InputDesc.Type != UE::Shader::EValueType::Struct)
 					{
@@ -776,11 +791,13 @@ struct FMaterialIRModuleBuilderImpl
 			}
 			else if (Property == MP_FrontMaterial && PropertyValue == nullptr && SubstrateTranslatorData.bMaterialUsesRootNodeToSubstrateHiddenConversion)
 			{
+				// SUBSTRATE_TODO Explain this case
 				PropertyValue = BuildSubstrateFrontMaterialIRFromRootNodeHiddenConversion();
 			}
 			else if (Property != MP_Normal && InputDesc.bUseConstant)
 			{
-				// If input is marked to use constant, assign this output to the specified constant value.
+				// If non-Normal input is marked to use constant, assign this output to the specified constant value.
+				// The Normal attribute should never use the constant as they have special default values.
 				PropertyValue = Emitter.ConstantFromShaderValue(InputDesc.ConstantValue);
 			}
 			else
@@ -806,6 +823,20 @@ struct FMaterialIRModuleBuilderImpl
 				MIR::FValue* SelectionColor = Emitter.NamedPrimitiveUniform(NAME_SelectionColor, Emitter.ConstantFloat4(UE::Math::TVector4<float>{ 0.f, 0.f, 0.f, 0.f }));
 				PropertyValue = Emitter.Lerp(PropertyValue, Emitter.Swizzle(SelectionColor, MIR::FSwizzleMask::XYZ()), Emitter.Subscript(SelectionColor, 3));
 			}
+			// Pack SubsurfaceColor (float3) with SubsurfaceProfile ID (float1) into float4,
+			// matching the old translator's behavior (HLSLMaterialTranslator.cpp:1623-1636).
+			else if (Property == MP_SubsurfaceColor)
+			{
+				MIR::FValueRef ColorRGB = Emitter.Swizzle(PropertyValue, MIR::FSwizzleMask::XYZ());
+				MIR::FValueRef ProfileParam = Emitter.CastToFloat(Emitter.NamedPrimitiveUniform(SubsurfaceProfile::GetSubsurfaceProfileParameterName(), Emitter.ConstantFloat(1.0f)), 1);
+				PropertyValue = Emitter.Vector4(ColorRGB, ProfileParam);
+			}
+			
+			// Quit if some error occurred in the operations above.
+			if (!Module->IsValid())
+			{
+				return false;
+			}
 
 			// Emit the SetMaterialProperty instruction
 			Emitter.Private_SetMaterialOutput(Property, PropertyValue);
diff --git a/Engine/Source/Runtime/Engine/Private/Materials/MaterialIRToHLSLTranslator.cpp b/Engine/Source/Runtime/Engine/Private/Materials/MaterialIRToHLSLTranslator.cpp
--- a/Engine/Source/Runtime/Engine/Private/Materials/MaterialIRToHLSLTranslator.cpp
+++ b/Engine/Source/Runtime/Engine/Private/Materials/MaterialIRToHLSLTranslator.cpp
@@ -271,6 +271,20 @@ static FPrinter& operator<<(FPrinter& Printer, const MIR::FType& Type)
 	return Printer;
 }
 
+// Prints the name of the property as it is called in the HLSL shader template.
+static void PrintMaterialPropertyName(FPrinter& Printer, EMaterialProperty Property)
+{
+	if (Property == MP_SubsurfaceColor)
+	{
+		// Because, for some reason, this is called "Subsurface" in the shader.
+		Printer << TEXTVIEW("Subsurface");
+	}
+	else
+	{
+		Printer << FMaterialAttributeDefinitionMap::GetAttributeName(Property);
+	}
+}
+
 // This struct encapsulates the full state and logic required to translate a single Material IR entry point into HLSL. It outputs
 // a the generated HLSL that evaluates a given entry point as described in the module.
 struct FEntryPointCodeGen
@@ -1044,7 +1058,8 @@ struct FEntryPointCodeGen
 		}
 		else
 		{
-			Printer << (CurrentStage != MIR::Stage_Vertex && SetMaterialOutput->Property != MP_CustomOutput ? TEXTVIEW("PixelMaterialInputs.") : TEXTVIEW("Parameters.")) << SetMaterialOutput->Name.ToView();
+			Printer << (CurrentStage != MIR::Stage_Vertex && SetMaterialOutput->Property != MP_CustomOutput ? TEXTVIEW("PixelMaterialInputs.") : TEXTVIEW("Parameters."));
+			PrintMaterialPropertyName(Printer, SetMaterialOutput->Property);
 		}
 
 		Printer << TEXTVIEW(" = ") << LowerValue(SetMaterialOutput->Arg);
@@ -2183,14 +2198,12 @@ static void GenerateStageEntryPoints(const FMaterialIRModule* Module, TMap<FStri
 
 	// Generate pixel stage HLSL
 	Result = CodeGen.GenerateEntryPoint(MIR::Stage_Pixel);
-	Result.NonNormalHLSL += TEXTVIEW("PixelMaterialInputs.Subsurface = 0;\n");
 	Result.NonNormalHLSL += CodeGen.Module->GetSubstrateData().PixelInputInitializerValuesHLSL;
 	OutParameters.Add(TEXT("calc_pixel_material_inputs_normal"), MoveTemp(Result.NormalAttributeHLSL));
 	OutParameters.Add(TEXT("calc_pixel_material_inputs_other_inputs"), MoveTemp(Result.NonNormalHLSL));
 
 	// Generate compute stage HLSL
 	Result = CodeGen.GenerateEntryPoint(MIR::Stage_Compute);
-	Result.NonNormalHLSL += TEXTVIEW("PixelMaterialInputs.Subsurface = 0;\n");
 	Result.NonNormalHLSL += CodeGen.Module->GetSubstrateData().PixelInputInitializerValuesHLSL;
 	OutParameters.Add(TEXT("calc_pixel_material_inputs_analytic_derivatives_normal"), MoveTemp(Result.NormalAttributeHLSL));
 	OutParameters.Add(TEXT("calc_pixel_material_inputs_analytic_derivatives_other_inputs"), MoveTemp(Result.NonNormalHLSL));
@@ -2385,7 +2398,9 @@ static void GenerateOtherTemplateStringParameters(const FMaterial* Material, con
 			}
 			if (const FMaterialAggregateAttribute* PropertyAggregate = MaterialAttributesAggregate::GetAttribute(Property))
 			{
-				Printer << MIR::FType::FromMaterialValueType(PropertyAggregate->ToMaterialValueType()) << TEXT(' ') << *PropertyAggregate->Name.ToString() << EndOfStatement;
+				Printer << MIR::FType::FromMaterialValueType(PropertyAggregate->ToMaterialValueType()) << TEXT(' ');
+				PrintMaterialPropertyName(Printer, Property);
+				Printer << EndOfStatement;
 			}
 		}
 
