Montages: Fix for prematurely calling AnimNotifyStates end events when ForcedNextFromPosition == ForcedNextToPosition (e.g. frame-locked Sequencer with unstable frame times)

#Synced-CL 52381678

diff --git a/Engine/Source/Runtime/Engine/Private/Animation/AnimMontage.cpp b/Engine/Source/Runtime/Engine/Private/Animation/AnimMontage.cpp
--- a/Engine/Source/Runtime/Engine/Private/Animation/AnimMontage.cpp
+++ b/Engine/Source/Runtime/Engine/Private/Animation/AnimMontage.cpp
@@ -2714,6 +2714,15 @@ void FAnimMontageInstance::Advance(float DeltaTime, struct FRootMotionMovementPa
 
 				if (SubStepResult == EMontageSubStepResult::NotMoved)
 				{
+					// If the montage has not advanced because ForcedNextFromPosition and ForcedNextToPosition are set to the same value, we need to keep any active notify state alive
+					// (This can happen for example in frame-locked Sequencer setups with unstable frame times)
+					if (ForcedNextFromPosition.IsSet() && ForcedNextToPosition.IsSet() 
+						&& ForcedNextFromPosition.GetValue() == ForcedNextToPosition.GetValue())
+					{
+						const float ForcedPosition = ForcedNextFromPosition.GetValue();
+						HandleEvents(ForcedPosition, ForcedPosition, nullptr);
+					}
+					
 					// If it hasn't moved, there is nothing much to do but weight update
 					break;
 				}
