[FTimerManager + AutoRTFM] Fix GetTimerElapsed() and GetTimerRemaining() for transactional timers

These were returning incorrect values when the handle was to a timer created within a transaction.

#Synced-CL 52295615

diff --git a/Engine/Source/Runtime/Engine/Private/Tests/AutoRTFM/AutoRTFMTimerManagerTests.cpp b/Engine/Source/Runtime/Engine/Private/Tests/AutoRTFM/AutoRTFMTimerManagerTests.cpp
--- a/Engine/Source/Runtime/Engine/Private/Tests/AutoRTFM/AutoRTFMTimerManagerTests.cpp
+++ b/Engine/Source/Runtime/Engine/Private/Tests/AutoRTFM/AutoRTFMTimerManagerTests.cpp
@@ -86,6 +86,8 @@ bool FAutoRTFMTimerManagerTest::RunTest(const FString & Parameters)
 			{
 				TimerManager.SetTimer(Timer.Handle, Timer.Callback(), /* InRate */ 1.0f, FTimerManagerTimerParameters{});
 				check(TimerManager.IsTimerActive(Timer.Handle));
+				check(TimerManager.GetTimerRemaining(Timer.Handle) == 1.0f);
+				check(TimerManager.GetTimerElapsed(Timer.Handle) == 0.0f);
 			}
 		});
 
@@ -94,6 +96,8 @@ bool FAutoRTFMTimerManagerTest::RunTest(const FString & Parameters)
 		for (FTimer& Timer : Timers)
 		{
 			check(Timer.bCalled);
+			check(TimerManager.GetTimerRemaining(Timer.Handle) == -1.0f);
+			check(TimerManager.GetTimerElapsed(Timer.Handle) == -1.0f);
 		}
 
 		return true;
@@ -110,6 +114,38 @@ bool FAutoRTFMTimerManagerTest::RunTest(const FString & Parameters)
 			{
 				TimerManager.SetTimer(Timer.Handle, Timer.Callback(), /* InRate */ 1.0f, FTimerManagerTimerParameters{});
 				check(TimerManager.IsTimerActive(Timer.Handle));
+				check(TimerManager.GetTimerRemaining(Timer.Handle) == 1.0f);
+				check(TimerManager.GetTimerElapsed(Timer.Handle) == 0.0f);
+			}
+			AutoRTFM::AbortTransaction();
+		});
+
+		Tick(TimerManager);
+
+		for (FTimer& Timer : Timers)
+		{
+			check(!Timer.bCalled);
+			check(!TimerManager.IsTimerActive(Timer.Handle));
+		}
+
+		return true;
+	});
+
+	Section(TEXT("Tick, Transact(SetTimer, Abort), Tick"), [&] (FTimerManager& TimerManager)
+	{
+		constexpr size_t NumTimers = 16;
+		FTimer Timers[NumTimers]{};
+
+		Tick(TimerManager);
+
+		AutoRTFM::Testing::Abort([&]
+		{
+			for (FTimer& Timer : Timers)
+			{
+				TimerManager.SetTimer(Timer.Handle, Timer.Callback(), /* InRate */ 1.0f, FTimerManagerTimerParameters{});
+				check(TimerManager.IsTimerActive(Timer.Handle));
+				check(TimerManager.GetTimerRemaining(Timer.Handle) == 1.0f);
+				check(TimerManager.GetTimerElapsed(Timer.Handle) == 0.0f);
 			}
 			AutoRTFM::AbortTransaction();
 		});
@@ -136,8 +172,12 @@ bool FAutoRTFMTimerManagerTest::RunTest(const FString & Parameters)
 			{
 				TimerManager.SetTimer(Timer.Handle, Timer.Callback(), /* InRate */ 1.0f, FTimerManagerTimerParameters{});
 				check(TimerManager.IsTimerActive(Timer.Handle));
+				check(TimerManager.GetTimerRemaining(Timer.Handle) == 1.0f);
+				check(TimerManager.GetTimerElapsed(Timer.Handle) == 0.0f);
 				ClearTimer(TimerManager, Timer.Handle);
 				check(!TimerManager.IsTimerActive(Timer.Handle));
+				check(TimerManager.GetTimerRemaining(Timer.Handle) == -1.0f);
+				check(TimerManager.GetTimerElapsed(Timer.Handle) == -1.0f);
 			}
 		});
 
@@ -147,6 +187,8 @@ bool FAutoRTFMTimerManagerTest::RunTest(const FString & Parameters)
 		{
 			check(!Timer.bCalled);
 			check(!TimerManager.IsTimerActive(Timer.Handle));
+			check(TimerManager.GetTimerRemaining(Timer.Handle) == -1.0f);
+			check(TimerManager.GetTimerElapsed(Timer.Handle) == -1.0f);
 		}
 
 		return true;
@@ -163,8 +205,12 @@ bool FAutoRTFMTimerManagerTest::RunTest(const FString & Parameters)
 			{
 				TimerManager.SetTimer(Timer.Handle, Timer.Callback(), /* InRate */ 1.0f, FTimerManagerTimerParameters{});
 				check(TimerManager.IsTimerActive(Timer.Handle));
+				check(TimerManager.GetTimerRemaining(Timer.Handle) == 1.0f);
+				check(TimerManager.GetTimerElapsed(Timer.Handle) == 0.0f);
 				ClearTimer(TimerManager, Timer.Handle);
 				check(!TimerManager.IsTimerActive(Timer.Handle));
+				check(TimerManager.GetTimerRemaining(Timer.Handle) == -1.0f);
+				check(TimerManager.GetTimerElapsed(Timer.Handle) == -1.0f);
 			}
 			AutoRTFM::AbortTransaction();
 		});
@@ -175,6 +221,8 @@ bool FAutoRTFMTimerManagerTest::RunTest(const FString & Parameters)
 		{
 			check(!Timer.bCalled);
 			check(!TimerManager.IsTimerActive(Timer.Handle));
+			check(TimerManager.GetTimerRemaining(Timer.Handle) == -1.0f);
+			check(TimerManager.GetTimerElapsed(Timer.Handle) == -1.0f);
 		}
 
 		return true;
@@ -193,6 +241,8 @@ bool FAutoRTFMTimerManagerTest::RunTest(const FString & Parameters)
 			{
 				TimerManager.SetTimer(ClosedTimer.Handle, ClosedTimer.Callback(), /* InRate */ 1.0f, FTimerManagerTimerParameters{});
 				check(TimerManager.IsTimerActive(ClosedTimer.Handle));
+				check(TimerManager.GetTimerRemaining(ClosedTimer.Handle) == 1.0f);
+				check(TimerManager.GetTimerElapsed(ClosedTimer.Handle) == 0.0f);
 			}
 			
 			AutoRTFM::Open([&]
@@ -201,6 +251,8 @@ bool FAutoRTFMTimerManagerTest::RunTest(const FString & Parameters)
 				{
 					TimerManager.SetTimer(OpenTimer.Handle, OpenTimer.Callback(), /* InRate */ 1.0f, FTimerManagerTimerParameters{});
 					check(TimerManager.IsTimerActive(OpenTimer.Handle));
+					check(TimerManager.GetTimerRemaining(OpenTimer.Handle) == 1.0f);
+					check(TimerManager.GetTimerElapsed(OpenTimer.Handle) == 0.0f);
 				}
 			});
 		});
@@ -210,11 +262,15 @@ bool FAutoRTFMTimerManagerTest::RunTest(const FString & Parameters)
 		for (FTimer& ClosedTimer : ClosedTimers)
 		{
 			check(ClosedTimer.bCalled);
+			check(TimerManager.GetTimerRemaining(ClosedTimer.Handle) == -1.0f);
+			check(TimerManager.GetTimerElapsed(ClosedTimer.Handle) == -1.0f);
 		}
 
 		for (FTimer& OpenTimer : OpenTimers)
 		{
 			check(OpenTimer.bCalled);
+			check(TimerManager.GetTimerRemaining(OpenTimer.Handle) == -1.0f);
+			check(TimerManager.GetTimerElapsed(OpenTimer.Handle) == -1.0f);
 		}
 
 		return true;
@@ -233,6 +289,8 @@ bool FAutoRTFMTimerManagerTest::RunTest(const FString & Parameters)
 			{
 				TimerManager.SetTimer(ClosedTimer.Handle, ClosedTimer.Callback(), /* InRate */ 1.0f, FTimerManagerTimerParameters{});
 				check(TimerManager.IsTimerActive(ClosedTimer.Handle));
+				check(TimerManager.GetTimerRemaining(ClosedTimer.Handle) == 1.0f);
+				check(TimerManager.GetTimerElapsed(ClosedTimer.Handle) == 0.0f);
 			}
 			
 			AutoRTFM::Open([&]
@@ -241,6 +299,8 @@ bool FAutoRTFMTimerManagerTest::RunTest(const FString & Parameters)
 				{
 					TimerManager.SetTimer(OpenTimer.Handle, OpenTimer.Callback(), /* InRate */ 1.0f, FTimerManagerTimerParameters{});
 					check(TimerManager.IsTimerActive(OpenTimer.Handle));
+					check(TimerManager.GetTimerRemaining(OpenTimer.Handle) == 1.0f);
+					check(TimerManager.GetTimerElapsed(OpenTimer.Handle) == 0.0f);
 				}
 			});
 			
@@ -252,11 +312,15 @@ bool FAutoRTFMTimerManagerTest::RunTest(const FString & Parameters)
 		for (FTimer& ClosedTimer : ClosedTimers)
 		{
 			check(!ClosedTimer.bCalled);
+			check(TimerManager.GetTimerRemaining(ClosedTimer.Handle) == -1.0f);
+			check(TimerManager.GetTimerElapsed(ClosedTimer.Handle) == -1.0f);
 		}
 
 		for (FTimer& OpenTimer : OpenTimers)
 		{
 			check(OpenTimer.bCalled);
+			check(TimerManager.GetTimerRemaining(OpenTimer.Handle) == -1.0f);
+			check(TimerManager.GetTimerElapsed(OpenTimer.Handle) == -1.0f);
 		}
 
 		return true;
@@ -277,6 +341,8 @@ bool FAutoRTFMTimerManagerTest::RunTest(const FString & Parameters)
 				{
 					TimerManager.SetTimer(OpenTimer.Handle, OpenTimer.Callback(), /* InRate */ 1.0f, FTimerManagerTimerParameters{});
 					check(TimerManager.IsTimerActive(OpenTimer.Handle));
+					check(TimerManager.GetTimerRemaining(OpenTimer.Handle) == 1.0f);
+					check(TimerManager.GetTimerElapsed(OpenTimer.Handle) == 0.0f);
 				}
 			});
 
@@ -284,6 +350,8 @@ bool FAutoRTFMTimerManagerTest::RunTest(const FString & Parameters)
 			{
 				TimerManager.SetTimer(ClosedTimer.Handle, ClosedTimer.Callback(), /* InRate */ 1.0f, FTimerManagerTimerParameters{});
 				check(TimerManager.IsTimerActive(ClosedTimer.Handle));
+				check(TimerManager.GetTimerRemaining(ClosedTimer.Handle) == 1.0f);
+				check(TimerManager.GetTimerElapsed(ClosedTimer.Handle) == 0.0f);
 			}
 			
 			AutoRTFM::Open([&]
@@ -293,6 +361,8 @@ bool FAutoRTFMTimerManagerTest::RunTest(const FString & Parameters)
 					check(TimerManager.IsTimerActive(OpenTimer.Handle));
 					ClearTimer(TimerManager, OpenTimer.Handle);
 					check(!TimerManager.IsTimerActive(OpenTimer.Handle));
+					check(TimerManager.GetTimerRemaining(OpenTimer.Handle) == -1.0f);
+					check(TimerManager.GetTimerElapsed(OpenTimer.Handle) == -1.0f);
 				}
 			});
 		});
@@ -302,11 +372,15 @@ bool FAutoRTFMTimerManagerTest::RunTest(const FString & Parameters)
 		for (FTimer& OpenTimer : OpenTimers)
 		{
 			check(!OpenTimer.bCalled);
+			check(TimerManager.GetTimerRemaining(OpenTimer.Handle) == -1.0f);
+			check(TimerManager.GetTimerElapsed(OpenTimer.Handle) == -1.0f);
 		}
 
 		for (FTimer& ClosedTimer : ClosedTimers)
 		{
 			check(ClosedTimer.bCalled);
+			check(TimerManager.GetTimerRemaining(ClosedTimer.Handle) == -1.0f);
+			check(TimerManager.GetTimerElapsed(ClosedTimer.Handle) == -1.0f);
 		}
 
 		return true;
@@ -327,6 +401,8 @@ bool FAutoRTFMTimerManagerTest::RunTest(const FString & Parameters)
 				{
 					TimerManager.SetTimer(OpenTimer.Handle, OpenTimer.Callback(), /* InRate */ 1.0f, FTimerManagerTimerParameters{});
 					check(TimerManager.IsTimerActive(OpenTimer.Handle));
+					check(TimerManager.GetTimerRemaining(OpenTimer.Handle) == 1.0f);
+					check(TimerManager.GetTimerElapsed(OpenTimer.Handle) == 0.0f);
 				}
 			});
 
@@ -334,6 +410,8 @@ bool FAutoRTFMTimerManagerTest::RunTest(const FString & Parameters)
 			{
 				TimerManager.SetTimer(ClosedTimer.Handle, ClosedTimer.Callback(), /* InRate */ 1.0f, FTimerManagerTimerParameters{});
 				check(TimerManager.IsTimerActive(ClosedTimer.Handle));
+				check(TimerManager.GetTimerRemaining(ClosedTimer.Handle) == 1.0f);
+				check(TimerManager.GetTimerElapsed(ClosedTimer.Handle) == 0.0f);
 			}
 			
 			AutoRTFM::Open([&]
@@ -343,6 +421,8 @@ bool FAutoRTFMTimerManagerTest::RunTest(const FString & Parameters)
 					check(TimerManager.IsTimerActive(OpenTimer.Handle));
 					ClearTimer(TimerManager, OpenTimer.Handle);
 					check(!TimerManager.IsTimerActive(OpenTimer.Handle));
+					check(TimerManager.GetTimerRemaining(OpenTimer.Handle) == -1.0f);
+					check(TimerManager.GetTimerElapsed(OpenTimer.Handle) == -1.0f);
 				}
 			});
 			
@@ -354,11 +434,15 @@ bool FAutoRTFMTimerManagerTest::RunTest(const FString & Parameters)
 		for (FTimer& OpenTimer : OpenTimers)
 		{
 			check(!OpenTimer.bCalled);
+			check(TimerManager.GetTimerRemaining(OpenTimer.Handle) == -1.0f);
+			check(TimerManager.GetTimerElapsed(OpenTimer.Handle) == -1.0f);
 		}
 
 		for (FTimer& ClosedTimer : ClosedTimers)
 		{
 			check(!ClosedTimer.bCalled);
+			check(TimerManager.GetTimerRemaining(ClosedTimer.Handle) == -1.0f);
+			check(TimerManager.GetTimerElapsed(ClosedTimer.Handle) == -1.0f);
 		}
 
 		return true;
@@ -373,6 +457,8 @@ bool FAutoRTFMTimerManagerTest::RunTest(const FString & Parameters)
 		{
 			TimerManager.SetTimer(Timer.Handle, Timer.Callback(), /* InRate */ 1.0f, FTimerManagerTimerParameters{});
 			check(TimerManager.IsTimerActive(Timer.Handle));
+			check(TimerManager.GetTimerRemaining(Timer.Handle) == 1.0f);
+			check(TimerManager.GetTimerElapsed(Timer.Handle) == 0.0f);
 		}
 
 		AutoRTFM::Testing::Commit([&]
@@ -385,6 +471,8 @@ bool FAutoRTFMTimerManagerTest::RunTest(const FString & Parameters)
 					check(TimerManager.IsTimerActive(Timer.Handle));
 					ClearTimer(TimerManager, Timer.Handle);
 					check(!TimerManager.IsTimerActive(Timer.Handle));
+					check(TimerManager.GetTimerRemaining(Timer.Handle) == -1.0f);
+					check(TimerManager.GetTimerElapsed(Timer.Handle) == -1.0f);
 				}
 			}
 		});
@@ -396,6 +484,8 @@ bool FAutoRTFMTimerManagerTest::RunTest(const FString & Parameters)
 			FTimer& Timer = Timers[I];
 			bool const bCleared = (I & 1) == 0;
 			check(Timer.bCalled == !bCleared);
+			check(TimerManager.GetTimerRemaining(Timer.Handle) == -1.0f);
+			check(TimerManager.GetTimerElapsed(Timer.Handle) == -1.0f);
 		}
 
 		return true;
@@ -410,6 +500,8 @@ bool FAutoRTFMTimerManagerTest::RunTest(const FString & Parameters)
 		{
 			TimerManager.SetTimer(Timer.Handle, Timer.Callback(), /* InRate */ 1.0f, FTimerManagerTimerParameters{});
 			check(TimerManager.IsTimerActive(Timer.Handle));
+			check(TimerManager.GetTimerRemaining(Timer.Handle) == 1.0f);
+			check(TimerManager.GetTimerElapsed(Timer.Handle) == 0.0f);
 		}
 
 		AutoRTFM::Testing::Abort([&]
@@ -418,16 +510,27 @@ bool FAutoRTFMTimerManagerTest::RunTest(const FString & Parameters)
 			{
 				FTimer& Timer = Timers[I];
 				ClearTimer(TimerManager, Timer.Handle);
+				check(TimerManager.GetTimerRemaining(Timer.Handle) == -1.0f);
+				check(TimerManager.GetTimerElapsed(Timer.Handle) == -1.0f);
 			}
 			AutoRTFM::AbortTransaction();
 		});
 
+		for (size_t I = 0; I < NumTimers; I++)
+		{
+			FTimer& Timer = Timers[I];
+			check(TimerManager.GetTimerRemaining(Timer.Handle) == 1.0f);
+			check(TimerManager.GetTimerElapsed(Timer.Handle) == 0.0f);
+		}
+
 		Tick(TimerManager);
 
 		for (size_t I = 0; I < NumTimers; I++)
 		{
 			FTimer& Timer = Timers[I];
 			check(Timer.bCalled);
+			check(TimerManager.GetTimerRemaining(Timer.Handle) == -1.0f);
+			check(TimerManager.GetTimerElapsed(Timer.Handle) == -1.0f);
 		}
 
 		return true;
@@ -442,6 +545,8 @@ bool FAutoRTFMTimerManagerTest::RunTest(const FString & Parameters)
 		{
 			TimerManager.SetTimer(Timer.Handle, Timer.Callback(), /* InRate */ 1.0f, FTimerManagerTimerParameters{});
 			check(TimerManager.IsTimerActive(Timer.Handle));
+			check(TimerManager.GetTimerRemaining(Timer.Handle) == 1.0f);
+			check(TimerManager.GetTimerElapsed(Timer.Handle) == 0.0f);
 		}
 
 		AutoRTFM::Testing::Commit([&]
@@ -452,8 +557,12 @@ bool FAutoRTFMTimerManagerTest::RunTest(const FString & Parameters)
 				{
 					FTimer& Timer = Timers[I];
 					check(!TimerManager.IsTimerPaused(Timer.Handle));
+					check(TimerManager.GetTimerRemaining(Timer.Handle) == 1.0f);
+					check(TimerManager.GetTimerElapsed(Timer.Handle) == 0.0f);
 					TimerManager.PauseTimer(Timer.Handle);
 					check(TimerManager.IsTimerPaused(Timer.Handle));
+					check(TimerManager.GetTimerRemaining(Timer.Handle) == 1.0f);
+					check(TimerManager.GetTimerElapsed(Timer.Handle) == 0.0f);
 				}
 			}
 		});
@@ -466,6 +575,8 @@ bool FAutoRTFMTimerManagerTest::RunTest(const FString & Parameters)
 			bool const bPaused = (I & 1) == 0;
 			check(Timer.bCalled == !bPaused);
 			check(TimerManager.IsTimerPaused(Timer.Handle) == bPaused);
+			check(TimerManager.GetTimerRemaining(Timer.Handle) == (bPaused ? 1.0f : -1.0f));
+			check(TimerManager.GetTimerElapsed(Timer.Handle) == (bPaused ? 0.0f : -1.0f));
 		}
 
 		return true;
@@ -480,6 +591,8 @@ bool FAutoRTFMTimerManagerTest::RunTest(const FString & Parameters)
 		{
 			TimerManager.SetTimer(Timer.Handle, Timer.Callback(), /* InRate */ 1.0f, FTimerManagerTimerParameters{});
 			check(TimerManager.IsTimerActive(Timer.Handle));
+			check(TimerManager.GetTimerRemaining(Timer.Handle) == 1.0f);
+			check(TimerManager.GetTimerElapsed(Timer.Handle) == 0.0f);
 		}
 
 		AutoRTFM::Testing::Abort([&]
@@ -491,6 +604,8 @@ bool FAutoRTFMTimerManagerTest::RunTest(const FString & Parameters)
 					FTimer& Timer = Timers[I];
 					TimerManager.PauseTimer(Timer.Handle);
 					check(TimerManager.IsTimerPaused(Timer.Handle));
+					check(TimerManager.GetTimerRemaining(Timer.Handle) == 1.0f);
+					check(TimerManager.GetTimerElapsed(Timer.Handle) == 0.0f);
 				}
 			}
 			AutoRTFM::AbortTransaction();
@@ -503,6 +618,8 @@ bool FAutoRTFMTimerManagerTest::RunTest(const FString & Parameters)
 			FTimer& Timer = Timers[I];
 			check(Timer.bCalled);
 			check(!TimerManager.IsTimerPaused(Timer.Handle));
+			check(TimerManager.GetTimerRemaining(Timer.Handle) == -1.0f);
+			check(TimerManager.GetTimerElapsed(Timer.Handle) == -1.0f);
 		}
 
 		return true;
@@ -517,8 +634,12 @@ bool FAutoRTFMTimerManagerTest::RunTest(const FString & Parameters)
 		{
 			TimerManager.SetTimer(Timer.Handle, Timer.Callback(), /* InRate */ 1.0f, FTimerManagerTimerParameters{});
 			check(TimerManager.IsTimerActive(Timer.Handle));
+			check(TimerManager.GetTimerRemaining(Timer.Handle) == 1.0f);
+			check(TimerManager.GetTimerElapsed(Timer.Handle) == 0.0f);
 			TimerManager.PauseTimer(Timer.Handle);
 			check(TimerManager.IsTimerPaused(Timer.Handle));
+			check(TimerManager.GetTimerRemaining(Timer.Handle) == 1.0f);
+			check(TimerManager.GetTimerElapsed(Timer.Handle) == 0.0f);
 		}
 
 		AutoRTFM::Testing::Commit([&]
@@ -529,8 +650,12 @@ bool FAutoRTFMTimerManagerTest::RunTest(const FString & Parameters)
 				{
 					FTimer& Timer = Timers[I];
 					check(TimerManager.IsTimerPaused(Timer.Handle));
+					check(TimerManager.GetTimerRemaining(Timer.Handle) == 1.0f);
+					check(TimerManager.GetTimerElapsed(Timer.Handle) == 0.0f);
 					TimerManager.UnPauseTimer(Timer.Handle);
 					check(!TimerManager.IsTimerPaused(Timer.Handle));
+					check(TimerManager.GetTimerRemaining(Timer.Handle) == 1.0f);
+					check(TimerManager.GetTimerElapsed(Timer.Handle) == 0.0f);
 				}
 			}
 		});
@@ -543,6 +668,8 @@ bool FAutoRTFMTimerManagerTest::RunTest(const FString & Parameters)
 			bool const bPaused = (I & 1) != 0;
 			check(Timer.bCalled == !bPaused);
 			check(TimerManager.IsTimerPaused(Timer.Handle) == bPaused);
+			check(TimerManager.GetTimerRemaining(Timer.Handle) == (bPaused ? 1.0f : -1.0f));
+			check(TimerManager.GetTimerElapsed(Timer.Handle) == (bPaused ? 0.0f : -1.0f));
 		}
 
 		return true;
@@ -557,8 +684,12 @@ bool FAutoRTFMTimerManagerTest::RunTest(const FString & Parameters)
 		{
 			TimerManager.SetTimer(Timer.Handle, Timer.Callback(), /* InRate */ 1.0f, FTimerManagerTimerParameters{});
 			check(TimerManager.IsTimerActive(Timer.Handle));
+			check(TimerManager.GetTimerRemaining(Timer.Handle) == 1.0f);
+			check(TimerManager.GetTimerElapsed(Timer.Handle) == 0.0f);
 			TimerManager.PauseTimer(Timer.Handle);
 			check(TimerManager.IsTimerPaused(Timer.Handle));
+			check(TimerManager.GetTimerRemaining(Timer.Handle) == 1.0f);
+			check(TimerManager.GetTimerElapsed(Timer.Handle) == 0.0f);
 		}
 
 		AutoRTFM::Testing::Abort([&]
@@ -569,8 +700,12 @@ bool FAutoRTFMTimerManagerTest::RunTest(const FString & Parameters)
 				{
 					FTimer& Timer = Timers[I];
 					check(TimerManager.IsTimerPaused(Timer.Handle));
+					check(TimerManager.GetTimerRemaining(Timer.Handle) == 1.0f);
+					check(TimerManager.GetTimerElapsed(Timer.Handle) == 0.0f);
 					TimerManager.UnPauseTimer(Timer.Handle);
 					check(!TimerManager.IsTimerPaused(Timer.Handle));
+					check(TimerManager.GetTimerRemaining(Timer.Handle) == 1.0f);
+					check(TimerManager.GetTimerElapsed(Timer.Handle) == 0.0f);
 				}
 			}
 			AutoRTFM::AbortTransaction();
@@ -583,6 +718,8 @@ bool FAutoRTFMTimerManagerTest::RunTest(const FString & Parameters)
 			FTimer& Timer = Timers[I];
 			check(!Timer.bCalled);
 			check(TimerManager.IsTimerPaused(Timer.Handle));
+			check(TimerManager.GetTimerRemaining(Timer.Handle) == 1.0f);
+			check(TimerManager.GetTimerElapsed(Timer.Handle) == 0.0f);
 		}
 
 		return true;
@@ -600,6 +737,8 @@ bool FAutoRTFMTimerManagerTest::RunTest(const FString & Parameters)
 			{
 				TimerManager.SetTimer(Timer.Handle, Timer.Callback(), /* InRate */ 1.0f, FTimerManagerTimerParameters{});
 				check(TimerManager.IsTimerActive(Timer.Handle));
+				check(TimerManager.GetTimerRemaining(Timer.Handle) == 1.0f);
+				check(TimerManager.GetTimerElapsed(Timer.Handle) == 0.0f);
 			}
 
 			for (FTimer& Timer : ClosedTimers)
@@ -607,6 +746,8 @@ bool FAutoRTFMTimerManagerTest::RunTest(const FString & Parameters)
 				check(TimerManager.IsTimerActive(Timer.Handle));
 				ClearTimer(TimerManager, Timer.Handle);
 				check(!TimerManager.IsTimerActive(Timer.Handle));
+				check(TimerManager.GetTimerRemaining(Timer.Handle) == -1.0f);
+				check(TimerManager.GetTimerElapsed(Timer.Handle) == -1.0f);
 			}
 
 			AutoRTFM::Open([&]
@@ -614,6 +755,8 @@ bool FAutoRTFMTimerManagerTest::RunTest(const FString & Parameters)
 				for (FTimer& Timer : OpenTimers)
 				{
 					TimerManager.SetTimer(Timer.Handle, Timer.Callback(), /* InRate */ 1.0f, FTimerManagerTimerParameters{});
+					check(TimerManager.GetTimerRemaining(Timer.Handle) == 1.0f);
+					check(TimerManager.GetTimerElapsed(Timer.Handle) == 0.0f);
 				}
 			});
 
@@ -625,10 +768,14 @@ bool FAutoRTFMTimerManagerTest::RunTest(const FString & Parameters)
 		for (FTimer& Timer : ClosedTimers)
 		{
 			check(!Timer.bCalled);
+			check(TimerManager.GetTimerRemaining(Timer.Handle) == -1.0f);
+			check(TimerManager.GetTimerElapsed(Timer.Handle) == -1.0f);
 		}
 		for (FTimer& Timer : OpenTimers)
 		{
 			check(Timer.bCalled);
+			check(TimerManager.GetTimerRemaining(Timer.Handle) == -1.0f);
+			check(TimerManager.GetTimerElapsed(Timer.Handle) == -1.0f);
 		}
 
 		return true;
diff --git a/Engine/Source/Runtime/Engine/Private/TimerManager.cpp b/Engine/Source/Runtime/Engine/Private/TimerManager.cpp
--- a/Engine/Source/Runtime/Engine/Private/TimerManager.cpp
+++ b/Engine/Source/Runtime/Engine/Private/TimerManager.cpp
@@ -912,6 +912,9 @@ float FTimerManager::InternalGetTimerRemaining(FTimerData const* const TimerData
 			case ETimerStatus::Executing:
 				return 0.0f;
 
+			case ETimerStatus::ReservedForTransaction:
+				break;
+
 			default:
 				// ExpireTime is time remaining for paused timers
 				return TimerData->ExpireTime;
diff --git a/Engine/Source/Runtime/Engine/Public/TimerManager.h b/Engine/Source/Runtime/Engine/Public/TimerManager.h
--- a/Engine/Source/Runtime/Engine/Public/TimerManager.h
+++ b/Engine/Source/Runtime/Engine/Public/TimerManager.h
@@ -417,6 +417,19 @@ public:
 	 */
 	inline float GetTimerElapsed(FTimerHandle InHandle) const
 	{
+#if UE_AUTORTFM
+		if (AutoRTFM::IsClosed())
+		{
+			if (const FTransactionalTimer* Timer = TransactionalTimers.Find(InHandle))
+			{
+				// A transactionally created timer cannot be executed until the
+				// transaction commits, so will have an elapsed time of 0,
+				// unless it was cancelled, in which case we return -1 (same as
+				// InternalGetTimerElapsed()).
+				return Timer->Status == FTransactionalTimer::Canceled ? -1.0f : 0.0f;
+			}
+		}
+#endif  // UE_AUTORTFM
 		FTimerData const* const TimerData = FindTimer(InHandle);
 		return InternalGetTimerElapsed(TimerData);
 	}
@@ -429,6 +442,22 @@ public:
 	 */
 	inline float GetTimerRemaining(FTimerHandle InHandle) const
 	{
+#if UE_AUTORTFM
+		if (AutoRTFM::IsClosed())
+		{
+			if (const FTransactionalTimer* Timer = TransactionalTimers.Find(InHandle))
+			{
+				if (Timer->NewData.IsSet())
+				{
+					return Timer->NewData->ExpireTime;
+				}
+				else if (Timer->Status == FTransactionalTimer::Canceled)
+				{
+					return -1.0f;
+				}
+			}
+		}
+#endif  // UE_AUTORTFM
 		FTimerData const* const TimerData = FindTimer(InHandle);
 		return InternalGetTimerRemaining(TimerData);
 	}
