Add warnings for failed codec lookup during animation deserialization
#fix Log warnings when bone/curve compression codec lookup fails during deserialization with non-empty DDC handle or codec path
#swarm https://p4-swarm.epicgames.net/reviews/52178159
#Synced-CL 52128385

#Synced-CL 52128385

diff --git a/Engine/Source/Runtime/Engine/Private/Animation/AnimCompressionTypes.cpp b/Engine/Source/Runtime/Engine/Private/Animation/AnimCompressionTypes.cpp
--- a/Engine/Source/Runtime/Engine/Private/Animation/AnimCompressionTypes.cpp
+++ b/Engine/Source/Runtime/Engine/Private/Animation/AnimCompressionTypes.cpp
@@ -1405,6 +1405,16 @@ void FCompressedAnimSequence::SerializeCompressedData(FArchive& Ar, bool bDDCDat
 		BoneCompressionCodec = BoneCompressionSettings->GetCodec(BoneCodecDDCHandle);
 		CurveCompressionCodec = CurveCompressionSettings->GetCodec(CurveCodecPath);
 
+		if (!BoneCompressionCodec && BoneCodecDDCHandle.Len() > 0)
+		{
+			UE_LOG(LogAnimation, Warning, TEXT("%s: Failed to find bone compression codec for DDC handle '%s'. Compressed bone data will not be available."), *DataOwner->GetPathName(), *BoneCodecDDCHandle);
+		}
+
+		if (!CurveCompressionCodec && CurveCodecPath.Len() > 0)
+		{
+			UE_LOG(LogAnimation, Warning, TEXT("%s: Failed to find curve compression codec for codec path '%s'. Compressed curve data will not be available."), *DataOwner->GetPathName(), *CurveCodecPath);
+		}
+
 		if (BoneCompressionCodec != nullptr)
 		{
 			CompressedDataStructure = BoneCompressionCodec->AllocateAnimData();
