Substarte - toon - fix for crash in material editor platform stat window. Only request anisotropy when Toon and Anisotropty are encoutnered.

#tests Pc editor, create toon materials, open stat window and add a mobile platform.
#rnx
#Synced-CL 52337136

diff --git a/Engine/Source/Runtime/Engine/Private/Materials/HLSLMaterialTranslator.cpp b/Engine/Source/Runtime/Engine/Private/Materials/HLSLMaterialTranslator.cpp
--- a/Engine/Source/Runtime/Engine/Private/Materials/HLSLMaterialTranslator.cpp
+++ b/Engine/Source/Runtime/Engine/Private/Materials/HLSLMaterialTranslator.cpp
@@ -1748,7 +1748,7 @@ void FHLSLMaterialTranslator::TranslateMaterial()
 		MaterialCompilationOutput.SubstrateMaterialCompilationOutput.SubstrateMaterialType = SUBSTRATE_MATERIAL_TYPE_COMPLEX;
 		MaterialCompilationOutput.SubstrateMaterialCompilationOutput.SubstrateMaterialBsdfFeatures |= ESubstrateBsdfFeature::Anisotropy;
 	}
-	if (uint16(MaterialCompilationOutput.SubstrateMaterialCompilationOutput.SubstrateMaterialBsdfFeatures & (ESubstrateBsdfFeature::Toon | ESubstrateBsdfFeature::Anisotropy)) > 0)
+	if (uint16(MaterialCompilationOutput.SubstrateMaterialCompilationOutput.SubstrateMaterialBsdfFeatures & (ESubstrateBsdfFeature::Toon | ESubstrateBsdfFeature::Anisotropy)) == uint16(ESubstrateBsdfFeature::Toon | ESubstrateBsdfFeature::Anisotropy))
 	{
 		MaterialCompilationOutput.bUsesAnisotropy = true;
 	}
