Move all OutData mutations in BuildData under UE_COMPRESSED_DATA_WRITE_SCOPE
#fix Moved CompressAnimBones, CompressAnimCurves, Update, and result assignment inside write lock to prevent readers observing partial compression state
#swarm https://p4-swarm.epicgames.net/reviews/52178169
#Synced-CL 52128385

diff --git a/Engine/Source/Runtime/Engine/Private/Animation/AnimationCompressionDerivedData.cpp b/Engine/Source/Runtime/Engine/Private/Animation/AnimationCompressionDerivedData.cpp
--- a/Engine/Source/Runtime/Engine/Private/Animation/AnimationCompressionDerivedData.cpp
+++ b/Engine/Source/Runtime/Engine/Private/Animation/AnimationCompressionDerivedData.cpp
@@ -215,35 +215,39 @@ bool FAnimationSequenceAsyncCacheTask::BuildData() const
 	{
 		return false;
 	}
-	DataToCompress.Update(OutData);
-		
-	const bool bBoneCompressionOk = FAnimationUtils::CompressAnimBones(DataToCompress, CompressionResult);
-	if (Owner.IsCanceled())
-	{
-		return false;
-	}
-	const bool bCurveCompressionOk = FAnimationUtils::CompressAnimCurves(DataToCompress, OutData);
-	const bool bIsCanceled = Owner.IsCanceled();
-
-	const bool bCompressionSuccessful = bBoneCompressionOk && bCurveCompressionOk;
+	
 	const FString BoneCompressionName = DataToCompress.BoneCompressionSettings ? DataToCompress.BoneCompressionSettings->GetFullName() : TEXT("NONE");
 	const FString CurveCompressionName = DataToCompress.CurveCompressionSettings ? DataToCompress.CurveCompressionSettings->GetFullName() : TEXT("NONE");
 
 	const UAnimSequence* AnimSequence = WeakAnimSequence.Get();
-	if (bCompressionSuccessful && !bIsCanceled && AnimSequence)
-	{
-		UE_COMPRESSED_DATA_WRITE_SCOPE(AnimSequence);		
-		OutData.CompressedByteStream = MoveTemp(CompressionResult.CompressedByteStream);
-		OutData.CompressedDataStructure = MoveTemp(CompressionResult.AnimData);
-		OutData.BoneCompressionCodec = CompressionResult.Codec;
-		OutData.CompressedRawData = DataToCompress.RawAnimationData;
-		OutData.OwnerName = DataToCompress.AnimFName;
-		
-		return true;
-	}
-	else if(!bCompressionSuccessful && !bIsCanceled)
+	if (!Owner.IsCanceled() && AnimSequence)
 	{
-		UE_LOG(LogAnimationCompression, Error, TEXT("Failed to generate compressed animation data for %s with bone compression scheme %s [%i] and curve compression scheme %s [%i] for target platform %s"), *CompressibleAnimPtr->FullName, *BoneCompressionName, bBoneCompressionOk ? 1 : 0, *CurveCompressionName, bCurveCompressionOk ? 1 : 0, *TargetPlatform->DisplayName().ToString());
+		UE_COMPRESSED_DATA_WRITE_SCOPE(AnimSequence);
+
+		// All compression and OutData mutations happen under the write lock
+		// to prevent readers from seeing partial/intermediate state
+		const bool bBoneCompressionOk = FAnimationUtils::CompressAnimBones(DataToCompress, CompressionResult);
+		if (Owner.IsCanceled())
+		{
+			return false;
+		}
+
+		const bool bCurveCompressionOk = FAnimationUtils::CompressAnimCurves(DataToCompress, OutData);
+		if (bBoneCompressionOk && bCurveCompressionOk && !Owner.IsCanceled())
+		{
+			DataToCompress.Update(OutData);
+			OutData.CompressedByteStream = MoveTemp(CompressionResult.CompressedByteStream);
+			OutData.CompressedDataStructure = MoveTemp(CompressionResult.AnimData);
+			OutData.BoneCompressionCodec = CompressionResult.Codec;
+			OutData.CompressedRawData = DataToCompress.RawAnimationData;
+			OutData.OwnerName = DataToCompress.AnimFName;
+
+			return true;
+		}
+		else if (!bBoneCompressionOk || !bCurveCompressionOk)
+		{
+			UE_LOG(LogAnimationCompression, Error, TEXT("Failed to generate compressed animation data for %s with bone compression scheme %s [%i] and curve compression scheme %s [%i] for target platform %s"), *CompressibleAnimPtr->FullName, *BoneCompressionName, bBoneCompressionOk ? 1 : 0, *CurveCompressionName, bCurveCompressionOk ? 1 : 0, *TargetPlatform->DisplayName().ToString());
+		}
 	}
 	else if (!AnimSequence)
 	{
