[Backout] - CL52389616
#rnx
Original CL Desc
-----------------------------------------------------------------
[FTimerManager + AutoRTFM] Fix GetTimerElapsed() and GetTimerRemaining() for transactional timers

These were returning incorrect values when the handle was to a timer created within a transaction.

diff --git a/Engine/Source/Runtime/Engine/Private/Tests/AutoRTFM/AutoRTFMTimerManagerTests.cpp b/Engine/Source/Runtime/Engine/Private/Tests/AutoRTFM/AutoRTFMTimerManagerTests.cpp
--- a/Engine/Source/Runtime/Engine/Private/Tests/AutoRTFM/AutoRTFMTimerManagerTests.cpp
+++ b/Engine/Source/Runtime/Engine/Private/Tests/AutoRTFM/AutoRTFMTimerManagerTests.cpp
@@ -86,8 +86,6 @@ bool FAutoRTFMTimerManagerTest::RunTest(const FString & Parameters)
 			{
 				TimerManager.SetTimer(Timer.Handle, Timer.Callback(), /* InRate */ 1.0f, FTimerManagerTimerParameters{});
 				check(TimerManager.IsTimerActive(Timer.Handle));
-				check(TimerManager.GetTimerRemaining(Timer.Handle) == 1.0f);
-				check(TimerManager.GetTimerElapsed(Timer.Handle) == 0.0f);
 			}
 		});
 
@@ -96,8 +94,6 @@ bool FAutoRTFMTimerManagerTest::RunTest(const FString & Parameters)
 		for (FTimer& Timer : Timers)
 		{
 			check(Timer.bCalled);
-			check(TimerManager.GetTimerRemaining(Timer.Handle) == -1.0f);
-			check(TimerManager.GetTimerElapsed(Timer.Handle) == -1.0f);
 		}
 
 		return true;
@@ -114,38 +110,6 @@ bool FAutoRTFMTimerManagerTest::RunTest(const FString & Parameters)
 			{
 				TimerManager.SetTimer(Timer.Handle, Timer.Callback(), /* InRate */ 1.0f, FTimerManagerTimerParameters{});
 				check(TimerManager.IsTimerActive(Timer.Handle));
-				check(TimerManager.GetTimerRemaining(Timer.Handle) == 1.0f);
-				check(TimerManager.GetTimerElapsed(Timer.Handle) == 0.0f);
-			}
-			AutoRTFM::AbortTransaction();
-		});
-
-		Tick(TimerManager);
-
-		for (FTimer& Timer : Timers)
-		{
-			check(!Timer.bCalled);
-			check(!TimerManager.IsTimerActive(Timer.Handle));
-		}
-
-		return true;
-	});
-
-	Section(TEXT("Tick, Transact(SetTimer, Abort), Tick"), [&] (FTimerManager& TimerManager)
-	{
-		constexpr size_t NumTimers = 16;
-		FTimer Timers[NumTimers]{};
-
-		Tick(TimerManager);
-
-		AutoRTFM::Testing::Abort([&]
-		{
-			for (FTimer& Timer : Timers)
-			{
-				TimerManager.SetTimer(Timer.Handle, Timer.Callback(), /* InRate */ 1.0f, FTimerManagerTimerParameters{});
-				check(TimerManager.IsTimerActive(Timer.Handle));
-				check(TimerManager.GetTimerRemaining(Timer.Handle) == 1.0f);
-				check(TimerManager.GetTimerElapsed(Timer.Handle) == 0.0f);
 			}
 			AutoRTFM::AbortTransaction();
 		});
@@ -172,12 +136,8 @@ bool FAutoRTFMTimerManagerTest::RunTest(const FString & Parameters)
 			{
 				TimerManager.SetTimer(Timer.Handle, Timer.Callback(), /* InRate */ 1.0f, FTimerManagerTimerParameters{});
 				check(TimerManager.IsTimerActive(Timer.Handle));
-				check(TimerManager.GetTimerRemaining(Timer.Handle) == 1.0f);
-				check(TimerManager.GetTimerElapsed(Timer.Handle) == 0.0f);
 				ClearTimer(TimerManager, Timer.Handle);
 				check(!TimerManager.IsTimerActive(Timer.Handle));
-				check(TimerManager.GetTimerRemaining(Timer.Handle) == -1.0f);
-				check(TimerManager.GetTimerElapsed(Timer.Handle) == -1.0f);
 			}
 		});
 
@@ -187,8 +147,6 @@ bool FAutoRTFMTimerManagerTest::RunTest(const FString & Parameters)
 		{
 			check(!Timer.bCalled);
 			check(!TimerManager.IsTimerActive(Timer.Handle));
-			check(TimerManager.GetTimerRemaining(Timer.Handle) == -1.0f);
-			check(TimerManager.GetTimerElapsed(Timer.Handle) == -1.0f);
 		}
 
 		return true;
@@ -205,12 +163,8 @@ bool FAutoRTFMTimerManagerTest::RunTest(const FString & Parameters)
 			{
 				TimerManager.SetTimer(Timer.Handle, Timer.Callback(), /* InRate */ 1.0f, FTimerManagerTimerParameters{});
 				check(TimerManager.IsTimerActive(Timer.Handle));
-				check(TimerManager.GetTimerRemaining(Timer.Handle) == 1.0f);
-				check(TimerManager.GetTimerElapsed(Timer.Handle) == 0.0f);
 				ClearTimer(TimerManager, Timer.Handle);
 				check(!TimerManager.IsTimerActive(Timer.Handle));
-				check(TimerManager.GetTimerRemaining(Timer.Handle) == -1.0f);
-				check(TimerManager.GetTimerElapsed(Timer.Handle) == -1.0f);
 			}
 			AutoRTFM::AbortTransaction();
 		});
@@ -221,8 +175,6 @@ bool FAutoRTFMTimerManagerTest::RunTest(const FString & Parameters)
 		{
 			check(!Timer.bCalled);
 			check(!TimerManager.IsTimerActive(Timer.Handle));
-			check(TimerManager.GetTimerRemaining(Timer.Handle) == -1.0f);
-			check(TimerManager.GetTimerElapsed(Timer.Handle) == -1.0f);
 		}
 
 		return true;
@@ -241,8 +193,6 @@ bool FAutoRTFMTimerManagerTest::RunTest(const FString & Parameters)
 			{
 				TimerManager.SetTimer(ClosedTimer.Handle, ClosedTimer.Callback(), /* InRate */ 1.0f, FTimerManagerTimerParameters{});
 				check(TimerManager.IsTimerActive(ClosedTimer.Handle));
-				check(TimerManager.GetTimerRemaining(ClosedTimer.Handle) == 1.0f);
-				check(TimerManager.GetTimerElapsed(ClosedTimer.Handle) == 0.0f);
 			}
 			
 			AutoRTFM::Open([&]
@@ -251,8 +201,6 @@ bool FAutoRTFMTimerManagerTest::RunTest(const FString & Parameters)
 				{
 					TimerManager.SetTimer(OpenTimer.Handle, OpenTimer.Callback(), /* InRate */ 1.0f, FTimerManagerTimerParameters{});
 					check(TimerManager.IsTimerActive(OpenTimer.Handle));
-					check(TimerManager.GetTimerRemaining(OpenTimer.Handle) == 1.0f);
-					check(TimerManager.GetTimerElapsed(OpenTimer.Handle) == 0.0f);
 				}
 			});
 		});
@@ -262,15 +210,11 @@ bool FAutoRTFMTimerManagerTest::RunTest(const FString & Parameters)
 		for (FTimer& ClosedTimer : ClosedTimers)
 		{
 			check(ClosedTimer.bCalled);
-			check(TimerManager.GetTimerRemaining(ClosedTimer.Handle) == -1.0f);
-			check(TimerManager.GetTimerElapsed(ClosedTimer.Handle) == -1.0f);
 		}
 
 		for (FTimer& OpenTimer : OpenTimers)
 		{
 			check(OpenTimer.bCalled);
-			check(TimerManager.GetTimerRemaining(OpenTimer.Handle) == -1.0f);
-			check(TimerManager.GetTimerElapsed(OpenTimer.Handle) == -1.0f);
 		}
 
 		return true;
@@ -289,8 +233,6 @@ bool FAutoRTFMTimerManagerTest::RunTest(const FString & Parameters)
 			{
 				TimerManager.SetTimer(ClosedTimer.Handle, ClosedTimer.Callback(), /* InRate */ 1.0f, FTimerManagerTimerParameters{});
 				check(TimerManager.IsTimerActive(ClosedTimer.Handle));
-				check(TimerManager.GetTimerRemaining(ClosedTimer.Handle) == 1.0f);
-				check(TimerManager.GetTimerElapsed(ClosedTimer.Handle) == 0.0f);
 			}
 			
 			AutoRTFM::Open([&]
@@ -299,8 +241,6 @@ bool FAutoRTFMTimerManagerTest::RunTest(const FString & Parameters)
 				{
 					TimerManager.SetTimer(OpenTimer.Handle, OpenTimer.Callback(), /* InRate */ 1.0f, FTimerManagerTimerParameters{});
 					check(TimerManager.IsTimerActive(OpenTimer.Handle));
-					check(TimerManager.GetTimerRemaining(OpenTimer.Handle) == 1.0f);
-					check(TimerManager.GetTimerElapsed(OpenTimer.Handle) == 0.0f);
 				}
 			});
 			
@@ -312,15 +252,11 @@ bool FAutoRTFMTimerManagerTest::RunTest(const FString & Parameters)
 		for (FTimer& ClosedTimer : ClosedTimers)
 		{
 			check(!ClosedTimer.bCalled);
-			check(TimerManager.GetTimerRemaining(ClosedTimer.Handle) == -1.0f);
-			check(TimerManager.GetTimerElapsed(ClosedTimer.Handle) == -1.0f);
 		}
 
 		for (FTimer& OpenTimer : OpenTimers)
 		{
 			check(OpenTimer.bCalled);
-			check(TimerManager.GetTimerRemaining(OpenTimer.Handle) == -1.0f);
-			check(TimerManager.GetTimerElapsed(OpenTimer.Handle) == -1.0f);
 		}
 
 		return true;
@@ -341,8 +277,6 @@ bool FAutoRTFMTimerManagerTest::RunTest(const FString & Parameters)
 				{
 					TimerManager.SetTimer(OpenTimer.Handle, OpenTimer.Callback(), /* InRate */ 1.0f, FTimerManagerTimerParameters{});
 					check(TimerManager.IsTimerActive(OpenTimer.Handle));
-					check(TimerManager.GetTimerRemaining(OpenTimer.Handle) == 1.0f);
-					check(TimerManager.GetTimerElapsed(OpenTimer.Handle) == 0.0f);
 				}
 			});
 
@@ -350,8 +284,6 @@ bool FAutoRTFMTimerManagerTest::RunTest(const FString & Parameters)
 			{
 				TimerManager.SetTimer(ClosedTimer.Handle, ClosedTimer.Callback(), /* InRate */ 1.0f, FTimerManagerTimerParameters{});
 				check(TimerManager.IsTimerActive(ClosedTimer.Handle));
-				check(TimerManager.GetTimerRemaining(ClosedTimer.Handle) == 1.0f);
-				check(TimerManager.GetTimerElapsed(ClosedTimer.Handle) == 0.0f);
 			}
 			
 			AutoRTFM::Open([&]
@@ -361,8 +293,6 @@ bool FAutoRTFMTimerManagerTest::RunTest(const FString & Parameters)
 					check(TimerManager.IsTimerActive(OpenTimer.Handle));
 					ClearTimer(TimerManager, OpenTimer.Handle);
 					check(!TimerManager.IsTimerActive(OpenTimer.Handle));
-					check(TimerManager.GetTimerRemaining(OpenTimer.Handle) == -1.0f);
-					check(TimerManager.GetTimerElapsed(OpenTimer.Handle) == -1.0f);
 				}
 			});
 		});
@@ -372,15 +302,11 @@ bool FAutoRTFMTimerManagerTest::RunTest(const FString & Parameters)
 		for (FTimer& OpenTimer : OpenTimers)
 		{
 			check(!OpenTimer.bCalled);
-			check(TimerManager.GetTimerRemaining(OpenTimer.Handle) == -1.0f);
-			check(TimerManager.GetTimerElapsed(OpenTimer.Handle) == -1.0f);
 		}
 
 		for (FTimer& ClosedTimer : ClosedTimers)
 		{
 			check(ClosedTimer.bCalled);
-			check(TimerManager.GetTimerRemaining(ClosedTimer.Handle) == -1.0f);
-			check(TimerManager.GetTimerElapsed(ClosedTimer.Handle) == -1.0f);
 		}
 
 		return true;
@@ -401,8 +327,6 @@ bool FAutoRTFMTimerManagerTest::RunTest(const FString & Parameters)
 				{
 					TimerManager.SetTimer(OpenTimer.Handle, OpenTimer.Callback(), /* InRate */ 1.0f, FTimerManagerTimerParameters{});
 					check(TimerManager.IsTimerActive(OpenTimer.Handle));
-					check(TimerManager.GetTimerRemaining(OpenTimer.Handle) == 1.0f);
-					check(TimerManager.GetTimerElapsed(OpenTimer.Handle) == 0.0f);
 				}
 			});
 
@@ -410,8 +334,6 @@ bool FAutoRTFMTimerManagerTest::RunTest(const FString & Parameters)
 			{
 				TimerManager.SetTimer(ClosedTimer.Handle, ClosedTimer.Callback(), /* InRate */ 1.0f, FTimerManagerTimerParameters{});
 				check(TimerManager.IsTimerActive(ClosedTimer.Handle));
-				check(TimerManager.GetTimerRemaining(ClosedTimer.Handle) == 1.0f);
-				check(TimerManager.GetTimerElapsed(ClosedTimer.Handle) == 0.0f);
 			}
 			
 			AutoRTFM::Open([&]
@@ -421,8 +343,6 @@ bool FAutoRTFMTimerManagerTest::RunTest(const FString & Parameters)
 					check(TimerManager.IsTimerActive(OpenTimer.Handle));
 					ClearTimer(TimerManager, OpenTimer.Handle);
 					check(!TimerManager.IsTimerActive(OpenTimer.Handle));
-					check(TimerManager.GetTimerRemaining(OpenTimer.Handle) == -1.0f);
-					check(TimerManager.GetTimerElapsed(OpenTimer.Handle) == -1.0f);
 				}
 			});
 			
@@ -434,15 +354,11 @@ bool FAutoRTFMTimerManagerTest::RunTest(const FString & Parameters)
 		for (FTimer& OpenTimer : OpenTimers)
 		{
 			check(!OpenTimer.bCalled);
-			check(TimerManager.GetTimerRemaining(OpenTimer.Handle) == -1.0f);
-			check(TimerManager.GetTimerElapsed(OpenTimer.Handle) == -1.0f);
 		}
 
 		for (FTimer& ClosedTimer : ClosedTimers)
 		{
 			check(!ClosedTimer.bCalled);
-			check(TimerManager.GetTimerRemaining(ClosedTimer.Handle) == -1.0f);
-			check(TimerManager.GetTimerElapsed(ClosedTimer.Handle) == -1.0f);
 		}
 
 		return true;
@@ -457,8 +373,6 @@ bool FAutoRTFMTimerManagerTest::RunTest(const FString & Parameters)
 		{
 			TimerManager.SetTimer(Timer.Handle, Timer.Callback(), /* InRate */ 1.0f, FTimerManagerTimerParameters{});
 			check(TimerManager.IsTimerActive(Timer.Handle));
-			check(TimerManager.GetTimerRemaining(Timer.Handle) == 1.0f);
-			check(TimerManager.GetTimerElapsed(Timer.Handle) == 0.0f);
 		}
 
 		AutoRTFM::Testing::Commit([&]
@@ -471,8 +385,6 @@ bool FAutoRTFMTimerManagerTest::RunTest(const FString & Parameters)
 					check(TimerManager.IsTimerActive(Timer.Handle));
 					ClearTimer(TimerManager, Timer.Handle);
 					check(!TimerManager.IsTimerActive(Timer.Handle));
-					check(TimerManager.GetTimerRemaining(Timer.Handle) == -1.0f);
-					check(TimerManager.GetTimerElapsed(Timer.Handle) == -1.0f);
 				}
 			}
 		});
@@ -484,8 +396,6 @@ bool FAutoRTFMTimerManagerTest::RunTest(const FString & Parameters)
 			FTimer& Timer = Timers[I];
 			bool const bCleared = (I & 1) == 0;
 			check(Timer.bCalled == !bCleared);
-			check(TimerManager.GetTimerRemaining(Timer.Handle) == -1.0f);
-			check(TimerManager.GetTimerElapsed(Timer.Handle) == -1.0f);
 		}
 
 		return true;
@@ -500,8 +410,6 @@ bool FAutoRTFMTimerManagerTest::RunTest(const FString & Parameters)
 		{
 			TimerManager.SetTimer(Timer.Handle, Timer.Callback(), /* InRate */ 1.0f, FTimerManagerTimerParameters{});
 			check(TimerManager.IsTimerActive(Timer.Handle));
-			check(TimerManager.GetTimerRemaining(Timer.Handle) == 1.0f);
-			check(TimerManager.GetTimerElapsed(Timer.Handle) == 0.0f);
 		}
 
 		AutoRTFM::Testing::Abort([&]
@@ -510,27 +418,16 @@ bool FAutoRTFMTimerManagerTest::RunTest(const FString & Parameters)
 			{
 				FTimer& Timer = Timers[I];
 				ClearTimer(TimerManager, Timer.Handle);
-				check(TimerManager.GetTimerRemaining(Timer.Handle) == -1.0f);
-				check(TimerManager.GetTimerElapsed(Timer.Handle) == -1.0f);
 			}
 			AutoRTFM::AbortTransaction();
 		});
 
-		for (size_t I = 0; I < NumTimers; I++)
-		{
-			FTimer& Timer = Timers[I];
-			check(TimerManager.GetTimerRemaining(Timer.Handle) == 1.0f);
-			check(TimerManager.GetTimerElapsed(Timer.Handle) == 0.0f);
-		}
-
 		Tick(TimerManager);
 
 		for (size_t I = 0; I < NumTimers; I++)
 		{
 			FTimer& Timer = Timers[I];
 			check(Timer.bCalled);
-			check(TimerManager.GetTimerRemaining(Timer.Handle) == -1.0f);
-			check(TimerManager.GetTimerElapsed(Timer.Handle) == -1.0f);
 		}
 
 		return true;
@@ -545,8 +442,6 @@ bool FAutoRTFMTimerManagerTest::RunTest(const FString & Parameters)
 		{
 			TimerManager.SetTimer(Timer.Handle, Timer.Callback(), /* InRate */ 1.0f, FTimerManagerTimerParameters{});
 			check(TimerManager.IsTimerActive(Timer.Handle));
-			check(TimerManager.GetTimerRemaining(Timer.Handle) == 1.0f);
-			check(TimerManager.GetTimerElapsed(Timer.Handle) == 0.0f);
 		}
 
 		AutoRTFM::Testing::Commit([&]
@@ -557,12 +452,8 @@ bool FAutoRTFMTimerManagerTest::RunTest(const FString & Parameters)
 				{
 					FTimer& Timer = Timers[I];
 					check(!TimerManager.IsTimerPaused(Timer.Handle));
-					check(TimerManager.GetTimerRemaining(Timer.Handle) == 1.0f);
-					check(TimerManager.GetTimerElapsed(Timer.Handle) == 0.0f);
 					TimerManager.PauseTimer(Timer.Handle);
 					check(TimerManager.IsTimerPaused(Timer.Handle));
-					check(TimerManager.GetTimerRemaining(Timer.Handle) == 1.0f);
-					check(TimerManager.GetTimerElapsed(Timer.Handle) == 0.0f);
 				}
 			}
 		});
@@ -575,8 +466,6 @@ bool FAutoRTFMTimerManagerTest::RunTest(const FString & Parameters)
 			bool const bPaused = (I & 1) == 0;
 			check(Timer.bCalled == !bPaused);
 			check(TimerManager.IsTimerPaused(Timer.Handle) == bPaused);
-			check(TimerManager.GetTimerRemaining(Timer.Handle) == (bPaused ? 1.0f : -1.0f));
-			check(TimerManager.GetTimerElapsed(Timer.Handle) == (bPaused ? 0.0f : -1.0f));
 		}
 
 		return true;
@@ -591,8 +480,6 @@ bool FAutoRTFMTimerManagerTest::RunTest(const FString & Parameters)
 		{
 			TimerManager.SetTimer(Timer.Handle, Timer.Callback(), /* InRate */ 1.0f, FTimerManagerTimerParameters{});
 			check(TimerManager.IsTimerActive(Timer.Handle));
-			check(TimerManager.GetTimerRemaining(Timer.Handle) == 1.0f);
-			check(TimerManager.GetTimerElapsed(Timer.Handle) == 0.0f);
 		}
 
 		AutoRTFM::Testing::Abort([&]
@@ -604,8 +491,6 @@ bool FAutoRTFMTimerManagerTest::RunTest(const FString & Parameters)
 					FTimer& Timer = Timers[I];
 					TimerManager.PauseTimer(Timer.Handle);
 					check(TimerManager.IsTimerPaused(Timer.Handle));
-					check(TimerManager.GetTimerRemaining(Timer.Handle) == 1.0f);
-					check(TimerManager.GetTimerElapsed(Timer.Handle) == 0.0f);
 				}
 			}
 			AutoRTFM::AbortTransaction();
@@ -618,8 +503,6 @@ bool FAutoRTFMTimerManagerTest::RunTest(const FString & Parameters)
 			FTimer& Timer = Timers[I];
 			check(Timer.bCalled);
 			check(!TimerManager.IsTimerPaused(Timer.Handle));
-			check(TimerManager.GetTimerRemaining(Timer.Handle) == -1.0f);
-			check(TimerManager.GetTimerElapsed(Timer.Handle) == -1.0f);
 		}
 
 		return true;
@@ -634,12 +517,8 @@ bool FAutoRTFMTimerManagerTest::RunTest(const FString & Parameters)
 		{
 			TimerManager.SetTimer(Timer.Handle, Timer.Callback(), /* InRate */ 1.0f, FTimerManagerTimerParameters{});
 			check(TimerManager.IsTimerActive(Timer.Handle));
-			check(TimerManager.GetTimerRemaining(Timer.Handle) == 1.0f);
-			check(TimerManager.GetTimerElapsed(Timer.Handle) == 0.0f);
 			TimerManager.PauseTimer(Timer.Handle);
 			check(TimerManager.IsTimerPaused(Timer.Handle));
-			check(TimerManager.GetTimerRemaining(Timer.Handle) == 1.0f);
-			check(TimerManager.GetTimerElapsed(Timer.Handle) == 0.0f);
 		}
 
 		AutoRTFM::Testing::Commit([&]
@@ -650,12 +529,8 @@ bool FAutoRTFMTimerManagerTest::RunTest(const FString & Parameters)
 				{
 					FTimer& Timer = Timers[I];
 					check(TimerManager.IsTimerPaused(Timer.Handle));
-					check(TimerManager.GetTimerRemaining(Timer.Handle) == 1.0f);
-					check(TimerManager.GetTimerElapsed(Timer.Handle) == 0.0f);
 					TimerManager.UnPauseTimer(Timer.Handle);
 					check(!TimerManager.IsTimerPaused(Timer.Handle));
-					check(TimerManager.GetTimerRemaining(Timer.Handle) == 1.0f);
-					check(TimerManager.GetTimerElapsed(Timer.Handle) == 0.0f);
 				}
 			}
 		});
@@ -668,8 +543,6 @@ bool FAutoRTFMTimerManagerTest::RunTest(const FString & Parameters)
 			bool const bPaused = (I & 1) != 0;
 			check(Timer.bCalled == !bPaused);
 			check(TimerManager.IsTimerPaused(Timer.Handle) == bPaused);
-			check(TimerManager.GetTimerRemaining(Timer.Handle) == (bPaused ? 1.0f : -1.0f));
-			check(TimerManager.GetTimerElapsed(Timer.Handle) == (bPaused ? 0.0f : -1.0f));
 		}
 
 		return true;
@@ -684,12 +557,8 @@ bool FAutoRTFMTimerManagerTest::RunTest(const FString & Parameters)
 		{
 			TimerManager.SetTimer(Timer.Handle, Timer.Callback(), /* InRate */ 1.0f, FTimerManagerTimerParameters{});
 			check(TimerManager.IsTimerActive(Timer.Handle));
-			check(TimerManager.GetTimerRemaining(Timer.Handle) == 1.0f);
-			check(TimerManager.GetTimerElapsed(Timer.Handle) == 0.0f);
 			TimerManager.PauseTimer(Timer.Handle);
 			check(TimerManager.IsTimerPaused(Timer.Handle));
-			check(TimerManager.GetTimerRemaining(Timer.Handle) == 1.0f);
-			check(TimerManager.GetTimerElapsed(Timer.Handle) == 0.0f);
 		}
 
 		AutoRTFM::Testing::Abort([&]
@@ -700,12 +569,8 @@ bool FAutoRTFMTimerManagerTest::RunTest(const FString & Parameters)
 				{
 					FTimer& Timer = Timers[I];
 					check(TimerManager.IsTimerPaused(Timer.Handle));
-					check(TimerManager.GetTimerRemaining(Timer.Handle) == 1.0f);
-					check(TimerManager.GetTimerElapsed(Timer.Handle) == 0.0f);
 					TimerManager.UnPauseTimer(Timer.Handle);
 					check(!TimerManager.IsTimerPaused(Timer.Handle));
-					check(TimerManager.GetTimerRemaining(Timer.Handle) == 1.0f);
-					check(TimerManager.GetTimerElapsed(Timer.Handle) == 0.0f);
 				}
 			}
 			AutoRTFM::AbortTransaction();
@@ -718,8 +583,6 @@ bool FAutoRTFMTimerManagerTest::RunTest(const FString & Parameters)
 			FTimer& Timer = Timers[I];
 			check(!Timer.bCalled);
 			check(TimerManager.IsTimerPaused(Timer.Handle));
-			check(TimerManager.GetTimerRemaining(Timer.Handle) == 1.0f);
-			check(TimerManager.GetTimerElapsed(Timer.Handle) == 0.0f);
 		}
 
 		return true;
@@ -737,8 +600,6 @@ bool FAutoRTFMTimerManagerTest::RunTest(const FString & Parameters)
 			{
 				TimerManager.SetTimer(Timer.Handle, Timer.Callback(), /* InRate */ 1.0f, FTimerManagerTimerParameters{});
 				check(TimerManager.IsTimerActive(Timer.Handle));
-				check(TimerManager.GetTimerRemaining(Timer.Handle) == 1.0f);
-				check(TimerManager.GetTimerElapsed(Timer.Handle) == 0.0f);
 			}
 
 			for (FTimer& Timer : ClosedTimers)
@@ -746,8 +607,6 @@ bool FAutoRTFMTimerManagerTest::RunTest(const FString & Parameters)
 				check(TimerManager.IsTimerActive(Timer.Handle));
 				ClearTimer(TimerManager, Timer.Handle);
 				check(!TimerManager.IsTimerActive(Timer.Handle));
-				check(TimerManager.GetTimerRemaining(Timer.Handle) == -1.0f);
-				check(TimerManager.GetTimerElapsed(Timer.Handle) == -1.0f);
 			}
 
 			AutoRTFM::Open([&]
@@ -755,8 +614,6 @@ bool FAutoRTFMTimerManagerTest::RunTest(const FString & Parameters)
 				for (FTimer& Timer : OpenTimers)
 				{
 					TimerManager.SetTimer(Timer.Handle, Timer.Callback(), /* InRate */ 1.0f, FTimerManagerTimerParameters{});
-					check(TimerManager.GetTimerRemaining(Timer.Handle) == 1.0f);
-					check(TimerManager.GetTimerElapsed(Timer.Handle) == 0.0f);
 				}
 			});
 
@@ -768,14 +625,10 @@ bool FAutoRTFMTimerManagerTest::RunTest(const FString & Parameters)
 		for (FTimer& Timer : ClosedTimers)
 		{
 			check(!Timer.bCalled);
-			check(TimerManager.GetTimerRemaining(Timer.Handle) == -1.0f);
-			check(TimerManager.GetTimerElapsed(Timer.Handle) == -1.0f);
 		}
 		for (FTimer& Timer : OpenTimers)
 		{
 			check(Timer.bCalled);
-			check(TimerManager.GetTimerRemaining(Timer.Handle) == -1.0f);
-			check(TimerManager.GetTimerElapsed(Timer.Handle) == -1.0f);
 		}
 
 		return true;
diff --git a/Engine/Source/Runtime/Engine/Private/TimerManager.cpp b/Engine/Source/Runtime/Engine/Private/TimerManager.cpp
--- a/Engine/Source/Runtime/Engine/Private/TimerManager.cpp
+++ b/Engine/Source/Runtime/Engine/Private/TimerManager.cpp
@@ -912,9 +912,6 @@ float FTimerManager::InternalGetTimerRemaining(FTimerData const* const TimerData
 			case ETimerStatus::Executing:
 				return 0.0f;
 
-			case ETimerStatus::ReservedForTransaction:
-				break;
-
 			default:
 				// ExpireTime is time remaining for paused timers
 				return TimerData->ExpireTime;
diff --git a/Engine/Source/Runtime/Engine/Public/TimerManager.h b/Engine/Source/Runtime/Engine/Public/TimerManager.h
--- a/Engine/Source/Runtime/Engine/Public/TimerManager.h
+++ b/Engine/Source/Runtime/Engine/Public/TimerManager.h
@@ -417,19 +417,6 @@ public:
 	 */
 	inline float GetTimerElapsed(FTimerHandle InHandle) const
 	{
-#if UE_AUTORTFM
-		if (AutoRTFM::IsClosed())
-		{
-			if (const FTransactionalTimer* Timer = TransactionalTimers.Find(InHandle))
-			{
-				// A transactionally created timer cannot be executed until the
-				// transaction commits, so will have an elapsed time of 0,
-				// unless it was cancelled, in which case we return -1 (same as
-				// InternalGetTimerElapsed()).
-				return Timer->Status == FTransactionalTimer::Canceled ? -1.0f : 0.0f;
-			}
-		}
-#endif  // UE_AUTORTFM
 		FTimerData const* const TimerData = FindTimer(InHandle);
 		return InternalGetTimerElapsed(TimerData);
 	}
@@ -442,22 +429,6 @@ public:
 	 */
 	inline float GetTimerRemaining(FTimerHandle InHandle) const
 	{
-#if UE_AUTORTFM
-		if (AutoRTFM::IsClosed())
-		{
-			if (const FTransactionalTimer* Timer = TransactionalTimers.Find(InHandle))
-			{
-				if (Timer->NewData.IsSet())
-				{
-					return Timer->NewData->ExpireTime;
-				}
-				else if (Timer->Status == FTransactionalTimer::Canceled)
-				{
-					return -1.0f;
-				}
-			}
-		}
-#endif  // UE_AUTORTFM
 		FTimerData const* const TimerData = FindTimer(InHandle);
 		return InternalGetTimerRemaining(TimerData);
 	}
