Fix for merge mesh with unified bone indices.

#Synced-CL 52391835

diff --git a/Engine/Source/Runtime/Engine/Private/SkeletalMeshMerge.cpp b/Engine/Source/Runtime/Engine/Private/SkeletalMeshMerge.cpp
--- a/Engine/Source/Runtime/Engine/Private/SkeletalMeshMerge.cpp
+++ b/Engine/Source/Runtime/Engine/Private/SkeletalMeshMerge.cpp
@@ -5,6 +5,7 @@
 =============================================================================*/
 
 #include "SkeletalMeshMerge.h"
+#include "Algo/Unique.h"
 #include "Animation/Skeleton.h"
 #include "Engine/SkeletalMesh.h"
 #include "Engine/SkeletalMeshSocket.h"
@@ -242,6 +243,14 @@ bool FSkeletalMeshMerge::FinalizeMesh()
 				GENERATE_LOD_MODEL(TGPUSkinVertexFloat32Uvs, PerLODNumUVSets[LODIdx]);
 			}
 		}
+
+		// Initialize the render data's unified bone map from the LOD storage.
+		FSkeletalMeshRenderData* MergeResource = MergeMesh->GetResourceForRendering();
+		if (MergeResource)
+		{
+			MergeResource->InitUnifiedBoneMapFromLODs();
+		}
+
 		// update the merge skel mesh entries
 		if (!ProcessMergeMesh())
 		{
@@ -745,6 +754,70 @@ void FSkeletalMeshMerge::GenerateLODModel( int32 LODIdx )
 		}
 	}
 
+	// If any source mesh uses unified bone indices, unify the merged output's bone maps too.
+	// This ensures the merged mesh is compatible with GPU instance LOD and single-pass skinning.
+	{
+		bool bAnySourceHasUnifiedBoneMap = false;
+		for (USkeletalMesh* SrcMesh : SrcMeshList)
+		{
+			if (SrcMesh && SrcMesh->GetResourceForRendering() && SrcMesh->GetResourceForRendering()->HasUnifiedBoneMap())
+			{
+				bAnySourceHasUnifiedBoneMap = true;
+				break;
+			}
+		}
+
+		if (bAnySourceHasUnifiedBoneMap)
+		{
+			// Build unified bone map from the union of all merged section bone maps.
+			TArray<FBoneIndexType> UnifiedBoneMap;
+			for (const FSkelMeshRenderSection& Section : MergeLODData.RenderSections)
+			{
+				UnifiedBoneMap.Append(Section.BoneMap);
+			}
+			UnifiedBoneMap.Sort();
+			UnifiedBoneMap.SetNum(Algo::Unique(UnifiedBoneMap));
+
+			// Remap each section's vertex bone indices from per-section bone map space to unified space,
+			// then overwrite the section's bone map with the unified map.
+			for (FSkelMeshRenderSection& Section : MergeLODData.RenderSections)
+			{
+				// Build remap: OldBoneMapIndex -> UnifiedBoneMapIndex
+				TArray<uint32> RemapTable;
+				RemapTable.SetNumUninitialized(Section.BoneMap.Num());
+				for (int32 OldIndex = 0; OldIndex < Section.BoneMap.Num(); ++OldIndex)
+				{
+					const int32 UnifiedIndex = UnifiedBoneMap.Find(Section.BoneMap[OldIndex]);
+					check(UnifiedIndex != INDEX_NONE);
+					RemapTable[OldIndex] = (uint32)UnifiedIndex;
+				}
+
+				// Remap vertex bone indices in the merged vertex buffer for this section's vertices.
+				for (uint32 VertIdx = 0; VertIdx < Section.NumVertices; ++VertIdx)
+				{
+					FSkinWeightInfo& Weight = MergedSkinWeightBuffer[Section.BaseVertexIndex + VertIdx];
+
+					for (int32 InfluenceIndex = 0; InfluenceIndex < MAX_TOTAL_INFLUENCES; ++InfluenceIndex)
+					{
+						if (Weight.InfluenceWeights[InfluenceIndex] > 0)
+						{
+							Weight.InfluenceBones[InfluenceIndex] = (FBoneIndexType)RemapTable[Weight.InfluenceBones[InfluenceIndex]];
+						}
+					}
+				}
+
+				Section.BoneMap = UnifiedBoneMap;
+				Section.bUnifiedBoneMap = true;
+			}
+
+			// Set the unified bone map on the LOD render data (points into the first section's BoneMap storage).
+			if (!MergeLODData.RenderSections.IsEmpty())
+			{
+				MergeLODData.UnifiedBoneMap = MergeLODData.RenderSections[0].BoneMap;
+			}
+		}
+	}
+
     const bool bNeedsCPUAccess = (MeshBufferAccess == EMeshBufferAccess::ForceCPUAndGPU);
 
 	// sort required bone array in strictly increasing order
diff --git a/Engine/Source/Runtime/Engine/Private/SkeletalMeshRenderData.cpp b/Engine/Source/Runtime/Engine/Private/SkeletalMeshRenderData.cpp
--- a/Engine/Source/Runtime/Engine/Private/SkeletalMeshRenderData.cpp
+++ b/Engine/Source/Runtime/Engine/Private/SkeletalMeshRenderData.cpp
@@ -302,7 +302,7 @@ void FSkeletalMeshRenderData::Cache(const ITargetPlatform* TargetPlatform, USkin
 		// If we have an initial build, the ddc key will be computed only after the build. Some structure are missing until we first build the asset to get the drived data key
 		
 		TArray64<uint8> DerivedData;
-		if(bAllowDDCFetch && DDCUtils64Bit::GetSynchronous(DerivedDataKey, Owner, DerivedData))
+		if (bAllowDDCFetch && DDCUtils64Bit::GetSynchronous(DerivedDataKey, Owner, DerivedData))
 		{
 			COOK_STAT(Timer.AddHit(DerivedData.Num()));
 
@@ -406,8 +406,8 @@ void FSkeletalMeshRenderData::Cache(const ITargetPlatform* TargetPlatform, USkin
 				LODData->BuildFromLODModel(LODModel, LODInfo->VertexAttributes, BuildSettings);
 			}
 
-			// Assign to the last LOD storage.
-			UnifiedBoneMap = !LODRenderData.IsEmpty() ? LODRenderData.Last().UnifiedBoneMap : TConstArrayView<FBoneIndexType>{};
+			// Point the render data's unified bone map at the last LOD's storage.
+			InitUnifiedBoneMapFromLODs();
 
 			FLargeMemoryWriter Ar(0, /*bIsPersistent=*/ true);
 
@@ -636,7 +636,7 @@ void FSkeletalMeshRenderData::Serialize(FArchive& Ar, USkinnedAsset* Owner)
 
 	if (Ar.IsLoading())
 	{
-		UnifiedBoneMap = !LODRenderData.IsEmpty() ? LODRenderData.Last().UnifiedBoneMap : TConstArrayView<FBoneIndexType>();
+		InitUnifiedBoneMapFromLODs();
 	}
 
 	check(NaniteResourcesPtr.IsValid());
diff --git a/Engine/Source/Runtime/Engine/Public/Rendering/SkeletalMeshRenderData.h b/Engine/Source/Runtime/Engine/Public/Rendering/SkeletalMeshRenderData.h
--- a/Engine/Source/Runtime/Engine/Public/Rendering/SkeletalMeshRenderData.h
+++ b/Engine/Source/Runtime/Engine/Public/Rendering/SkeletalMeshRenderData.h
@@ -148,6 +148,12 @@ public:
 		return UnifiedBoneMap;
 	}
 
+	/** Initialize the render data's unified bone map from the last LOD's storage. Call after LOD render data is fully built. */
+	void InitUnifiedBoneMapFromLODs()
+	{
+		UnifiedBoneMap = !LODRenderData.IsEmpty() ? LODRenderData.Last().GetUnifiedBoneMap() : TConstArrayView<FBoneIndexType>{};
+	}
+
 	ENGINE_API bool HasValidNaniteData() const;
 
 	ENGINE_API bool IsLODDataReady() const;
