ISKM now supports bRenderStatic on the !Nanite path. This mirrors how it works on SKM. This allows Nanite foliage to be used in LOD0 and WPO to be used on LOD1 -> N for fallback meshes.

#Synced-CL 52384345

diff --git a/Engine/Source/Runtime/Engine/Private/InstancedSkeletalMeshObject.cpp b/Engine/Source/Runtime/Engine/Private/InstancedSkeletalMeshObject.cpp
--- a/Engine/Source/Runtime/Engine/Private/InstancedSkeletalMeshObject.cpp
+++ b/Engine/Source/Runtime/Engine/Private/InstancedSkeletalMeshObject.cpp
@@ -25,7 +25,7 @@ FInstancedSkeletalMeshObject::FInstancedSkeletalMeshObject(
 		LODs.Emplace(InRenderData, LODIndex, InFeatureLevel);
 	}
 
-	bUseGpuLodSelection = NumLODs > 1 && InMeshDesc.bUseGpuLodSelection && bIsGPUSkin && InRenderData->HasUnifiedBoneMap();
+	bUseGpuLodSelection = NumLODs > 1 && InMeshDesc.bUseGpuLodSelection && ((bIsGPUSkin && InRenderData->HasUnifiedBoneMap()) || Type == EType::Static);
 	bForceAnimateSockets = InMeshDesc.bForceAnimateSockets;
 	BoneTransformStorageMode = (bIsGPUSkin || InRenderData->HasUnifiedBoneMap()) ? EBoneTransformStorageMode::BoneMap : EBoneTransformStorageMode::Direct;
 
@@ -96,10 +96,11 @@ void FInstancedSkeletalMeshObject::Update(
 
 FSkinningSceneExtensionProxy* FInstancedSkeletalMeshObject::CreateSceneExtensionProxy(const USkinnedAsset* InSkinnedAsset, bool bAllowScaling) const
 {
-	if (Type == EType::Nanite || IsGPUSkinSceneExtensionEnabled())
+	if (Type == EType::Nanite || (Type == EType::GPUSkin && IsGPUSkinSceneExtensionEnabled()))
 	{
 		return new FInstancedSkinningSceneExtensionProxy(TransformProvider, this, InSkinnedAsset, bAllowScaling, bForceAnimateSockets);
 	}
+	// EType::Static does not register with the skinning extension (uses FLocalVertexFactory)
 	return nullptr;
 }
 
@@ -149,6 +150,11 @@ const FVertexFactory* FInstancedSkeletalMeshObject::GetSkinVertexFactory(const F
 {
 	const FLOD& LOD = LODs[LODIndex];
 
+	if (Type == EType::Static)
+	{
+		return &LOD.LocalVertexFactory;
+	}
+
 #if RHI_RAYTRACING
 	if (VFMode == ESkinVertexFactoryMode::RayTracing)
 	{
@@ -258,8 +264,12 @@ void FInstancedSkeletalMeshObject::FLOD::InitResources(const FSkelMeshComponentL
 			}
 		}
 
+		// Initialize LocalVertexFactory for Static mode (all rendering) and for ray tracing in other modes
+		if (InType == EType::Static
 #if RHI_RAYTRACING
-		if (bRayTracingEnabled)
+			|| bRayTracingEnabled
+#endif
+		)
 		{
 			FPositionVertexBuffer* PositionVertexBufferPtr = &LODData.StaticVertexBuffers.PositionVertexBuffer;
 			FStaticMeshVertexBuffer* StaticMeshVertexBufferPtr = &LODData.StaticVertexBuffers.StaticMeshVertexBuffer;
@@ -273,7 +283,6 @@ void FInstancedSkeletalMeshObject::FLOD::InitResources(const FSkelMeshComponentL
 			LocalVertexFactory.SetData(RHICmdList, Data);
 			LocalVertexFactory.InitResource(RHICmdList);
 		}
-#endif
 	});
 }
 
diff --git a/Engine/Source/Runtime/Engine/Private/InstancedSkeletalMeshObject.h b/Engine/Source/Runtime/Engine/Private/InstancedSkeletalMeshObject.h
--- a/Engine/Source/Runtime/Engine/Private/InstancedSkeletalMeshObject.h
+++ b/Engine/Source/Runtime/Engine/Private/InstancedSkeletalMeshObject.h
@@ -16,7 +16,8 @@ public:
 	enum class EType : uint8
 	{
 		Nanite,
-		GPUSkin
+		GPUSkin,
+		Static
 	};
 
 	FInstancedSkeletalMeshObject(
diff --git a/Engine/Source/Runtime/Engine/Private/InstancedSkinnedMeshSceneProxy.cpp b/Engine/Source/Runtime/Engine/Private/InstancedSkinnedMeshSceneProxy.cpp
--- a/Engine/Source/Runtime/Engine/Private/InstancedSkinnedMeshSceneProxy.cpp
+++ b/Engine/Source/Runtime/Engine/Private/InstancedSkinnedMeshSceneProxy.cpp
@@ -73,7 +73,7 @@ FInstancedSkinnedMeshSceneProxy::FInstancedSkinnedMeshSceneProxy(const FInstance
 
 	bAlwaysHasVelocity = true;
 	bInstancedSkinnedMesh = true;
-	bUseGpuLodSelection = NumLODs > 1 && InMeshDesc.bUseGpuLodSelection && InRenderData->HasUnifiedBoneMap();
+	bUseGpuLodSelection = NumLODs > 1 && InMeshDesc.bUseGpuLodSelection && (InMeshDesc.bRenderStatic || InRenderData->HasUnifiedBoneMap());
 
 	SetupInstanceSceneDataBuffers(Data.InstanceDataSceneProxy->GeInstanceSceneDataBuffers());
 }
\ No newline at end of file
diff --git a/Engine/Source/Runtime/Engine/Private/InstancedSkinnedMeshSceneProxyDesc.cpp b/Engine/Source/Runtime/Engine/Private/InstancedSkinnedMeshSceneProxyDesc.cpp
--- a/Engine/Source/Runtime/Engine/Private/InstancedSkinnedMeshSceneProxyDesc.cpp
+++ b/Engine/Source/Runtime/Engine/Private/InstancedSkinnedMeshSceneProxyDesc.cpp
@@ -13,6 +13,10 @@ FSkeletalMeshObject* FInstancedSkinnedMeshSceneProxyDesc::CreateMeshObject(const
 	{
 		return new FInstancedSkeletalMeshObject(InMeshDesc, InRenderData, InFeatureLevel, FInstancedSkeletalMeshObject::EType::Nanite);
 	}
+	else if (InMeshDesc.bRenderStatic)
+	{
+		return new FInstancedSkeletalMeshObject(InMeshDesc, InRenderData, InFeatureLevel, FInstancedSkeletalMeshObject::EType::Static);
+	}
 	else if (!InMeshDesc.ShouldCPUSkin() && IsGPUSkinSceneExtensionEnabled())
 	{
 		return new FInstancedSkeletalMeshObject(InMeshDesc, InRenderData, InFeatureLevel, FInstancedSkeletalMeshObject::EType::GPUSkin);
@@ -49,7 +53,7 @@ FPrimitiveSceneProxy* FInstancedSkinnedMeshSceneProxyDesc::CreateSceneProxy(cons
 					Result = ::new FNaniteInstancedSkinnedMeshSceneProxy(NaniteMaterials, Desc, SkelMeshRenderData);
 				}
 			}
-			else if (MeshObject->IsGPUSkinMesh())
+			else if (!MeshObject->IsNaniteMesh())
 			{
 				Result = ::new FInstancedSkinnedMeshSceneProxy(Desc, SkelMeshRenderData);
 			}
