[shaders] fix up handling of "bBypassCache" in the in-memory job cache. this was previously only serving to prevent jobs with this flag set from being added to DDC; extend its behaviour to prevent querying such jobs from any cache, and to skip adding such jobs to the in-memory cache. this fixes the behaviour of the retry mechanism when r.DumpShaderDebugInfo is set to 3 - previously the resubmitted job was still hitting the cache, but since the output struct was reset the warnings that should be logged from the retry job were missing. this ensures the resubmitted job actually executes again.

#Synced-CL 52384904

diff --git a/Engine/Source/Runtime/Engine/Private/ShaderCompiler/ShaderCompilerJobCache.cpp b/Engine/Source/Runtime/Engine/Private/ShaderCompiler/ShaderCompilerJobCache.cpp
--- a/Engine/Source/Runtime/Engine/Private/ShaderCompiler/ShaderCompilerJobCache.cpp
+++ b/Engine/Source/Runtime/Engine/Private/ShaderCompiler/ShaderCompilerJobCache.cpp
@@ -1385,7 +1385,7 @@ void FShaderJobCache::SubmitJob(FShaderCommonCompileJob* Job)
 		FShaderCacheLoadContext LoadContext = FindOrAdd(InputHash, Job, bCheckDDC);
 
 		// see if there are already cached results for this job that were returned synchronously by FindOrAdd
-		if (LoadContext.HasData())
+		if (LoadContext.HasData() && !Job->bBypassCache)
 		{
 			Unlink(*Job);		// from PendingSubmitJobTaskJobs
 
@@ -1622,9 +1622,9 @@ void FShaderJobCache::AddToCacheAndProcessPending(FShaderCommonCompileJob* Finis
 
 		JobData.DuplicateJobsWaitList = nullptr;
 
-		if (FinishedJob->bSucceeded)
+		if (FinishedJob->bSucceeded && !FinishedJob->bBypassCache)
 		{
-			const bool bAddToDDC = !FinishedJob->bBypassCache && (GShaderCompilerPerShaderDDCGlobal || !(FinishedJob->bIsDefaultMaterial || FinishedJob->bIsGlobalShader));
+			const bool bAddToDDC = GShaderCompilerPerShaderDDCGlobal || !(FinishedJob->bIsDefaultMaterial || FinishedJob->bIsGlobalShader);
 			// we only cache jobs that succeded
 			AddJobOutput(JobData, FinishedJob, InputHash, SaveContext, FinishedDuplicateJobs.Num(), bAddToDDC);
 		}
