Material Translator: Add whether we use opacity to the compilatin output.
This is necessary to determine whether we write every pixel.
When we have translucent materials with dynamic shadow masking we are effectively masked by opacity output rather than opacity mask output.
#rnx
#Synced-CL 52391845

diff --git a/Engine/Source/Runtime/Engine/Private/Materials/HLSLMaterialTranslator.cpp b/Engine/Source/Runtime/Engine/Private/Materials/HLSLMaterialTranslator.cpp
--- a/Engine/Source/Runtime/Engine/Private/Materials/HLSLMaterialTranslator.cpp
+++ b/Engine/Source/Runtime/Engine/Private/Materials/HLSLMaterialTranslator.cpp
@@ -1719,7 +1719,10 @@ void FHLSLMaterialTranslator::TranslateMaterial()
 		Material->IsTessellationEnabled() &&
 		IsMaterialPropertyUsed(MP_Displacement, Chunk[MP_Displacement], FLinearColor(-1, 0, 0, 0), 1);
 
+	// Check if opacity mask is used. This is used by the FMaterialShaderParameters to determine whether to generate masked shaders.
 	const bool bUsesOpacityMask = IsMaterialPropertyUsed(MP_OpacityMask, Chunk[MP_OpacityMask], FLinearColor(1, 0, 0, 0), 1);
+	// Check if opacity is used. This is can also generate masked shaders when casting shadows as masked from a translucent material.
+	const bool bUsesOpacity = IsMaterialPropertyUsed(MP_Opacity, Chunk[MP_Opacity], FLinearColor(1, 0, 0, 0), 1);
 
 	static const FGuid FirstPersonInterpolationAlphaGuid = FMaterialAttributeDefinitionMap::GetCustomAttributeID(TEXT("FirstPersonInterpolationAlpha"));
 	const bool bHasFirstPersonOutput = (CustomOutputsBitmask & FMaterialAttributeDefinitionMap::GetBitmask(FirstPersonInterpolationAlphaGuid)) != 0;
@@ -1730,6 +1733,7 @@ void FHLSLMaterialTranslator::TranslateMaterial()
 	MaterialCompilationOutput.bUsesDisplacement = bUsesDisplacement;
 	MaterialCompilationOutput.bUsesCustomizedUVs = Material->GetNumCustomizedUVs() > 0;
 	MaterialCompilationOutput.bUsesOpacityMask = bUsesOpacityMask;
+	MaterialCompilationOutput.bUsesOpacity = bUsesOpacity;
 
 	static const FGuid TemporalResponsivenessGuid = FMaterialAttributeDefinitionMap::GetCustomAttributeID(TEXT("TemporalResponsiveness"));
 	static const FGuid MotionVectorWorldOffsetGuid = FMaterialAttributeDefinitionMap::GetCustomAttributeID(TEXT("MotionVectorWorldOffset"));
diff --git a/Engine/Source/Runtime/Engine/Private/Materials/MaterialShared.cpp b/Engine/Source/Runtime/Engine/Private/Materials/MaterialShared.cpp
--- a/Engine/Source/Runtime/Engine/Private/Materials/MaterialShared.cpp
+++ b/Engine/Source/Runtime/Engine/Private/Materials/MaterialShared.cpp
@@ -6107,7 +6107,7 @@ void FMaterialShaderParameters::InitializeShaderLayoutParamsFromCompOutput(const
 	// Consider whether opacity mask is actually used when calculating WritesEveryPixel().
 	// Treat the case where a custom HLSL expression is discarding pixels as equivalent to the use of opacity mask here.
 	const bool bVFTypeSupportsNullPixelShader = !InMaterial->IsUsedWithInstancedStaticMeshes();
-	const bool bUsesOpacityMaskOrDiscard = InCompilationOutput.bUsesOpacityMask || InCompilationOutput.bUsesCustomHLSLPixelDiscard;
+	const bool bUsesOpacityMaskOrDiscard = InCompilationOutput.bUsesOpacityMask || InCompilationOutput.bUsesCustomHLSLPixelDiscard || (bShouldCastDynamicShadows && InCompilationOutput.bUsesOpacity);
 	bWritesEveryPixel = InMaterial->WritesEveryPixel(/*bShadowPass*/false, bVFTypeSupportsNullPixelShader, !bUsesOpacityMaskOrDiscard);
 	bWritesEveryPixelShadowPass = InMaterial->WritesEveryPixel(/*bShadowPass*/true, bVFTypeSupportsNullPixelShader, !bUsesOpacityMaskOrDiscard);
 }
diff --git a/Engine/Source/Runtime/Engine/Public/MaterialShared.h b/Engine/Source/Runtime/Engine/Public/MaterialShared.h
--- a/Engine/Source/Runtime/Engine/Public/MaterialShared.h
+++ b/Engine/Source/Runtime/Engine/Public/MaterialShared.h
@@ -910,7 +910,9 @@ public:
 		bUsesDisplacement(false),
 		bUsedWithNeuralNetworks(false),
 		bUsesCustomizedUVs(false),
-		bUsesOpacityMask(false)
+		bUsesOpacityMask(false),
+		bUsesCustomHLSLPixelDiscard(false),
+		bUsesOpacity(false)
 	{
 #if WITH_EDITOR
 		FMemory::Memzero(EstimatedLWCFuncUsagesVS);
@@ -1085,9 +1087,11 @@ public:
 	/** True if the opacity mask isn't constant 1.0 (after constant folding). */
 	LAYOUT_BITFIELD(uint8, bUsesOpacityMask, 1);
 
-	/** True if any UMaterialExpressionCustom compiled into this material uses clip() or discard
-	    to conditionally discard pixels, making it behave like an opacity mask. */
+	/** True if any UMaterialExpressionCustom compiled into this material uses clip() or discard to conditionally discard pixels, making it behave like an opacity mask. */
 	LAYOUT_BITFIELD(uint8, bUsesCustomHLSLPixelDiscard, 1);
+
+	/** True if the opacity isn't constant 1.0 (after constant folding). */
+	LAYOUT_BITFIELD(uint8, bUsesOpacity, 1);
 };
 
 struct FDebugShaderPipelineInfo
