SimpleStreamableAssetManager - Fix memory leak in SSAM when objects are unregistered

Issue:
When SSAM is enabled, unregistering objects (proxies) while the assets they reference remain alive causes unbounded memory growth in AssetIndexToBounds4Index. Assets (e.g. textures) may remain registered in SSAM because the underlying UStreamableRenderAsset is still alive - kept from garbage collection by persistent references. Since UnregisterAsset() only runs when the asset is destroyed, the inner sparse array buffers allocated for each asset's object references are never freed. Over time - for example across repeated level transitions - the memory footprint grows linearly with the total number of unique assets encountered.

Fix:
When RemoveRenderAssetElements empties an asset's sparse array (all referencing objects removed), the asset slot is now fully reclaimed: inner buffers freed, AssetUsedIndices bit cleared, and the asset's shared index pointer invalidated back to INDEX_NONE. This allows the index to be reused by future registrations instead of accumulating as zombie entries.

To support this, a new parallel array AssetIndexToSharedIndex stores each asset's TSharedPtr<int32> so SSAM can invalidate it without needing a back-pointer to the UStreamableRenderAsset.

Changes:
- SimpleStreamableAssetManager.h: Added AssetIndexToSharedIndex array parallel to AssetIndexToBounds4Index
- SimpleStreamableAssetManager.cpp:
- AddRenderAssetElements: Populate shared_ptr on new asset index allocation; re-populate when reusing a previously reclaimed slot
- RemoveRenderAssetElements: Full slot reclamation when Num() == 0 (free inner buffer, release index, invalidate shared_ptr)
- Pending_RemovedAssetRecords: Clear shared_ptr entry on explicit asset unregistration
- GetAllocatedSize_Impl: Track new array in memory accounting
- SSAMCoreTests.cpp: Re-enabled MemoryBloatOnLevelTransition test

#Synced-CL 52335816

diff --git a/Engine/Source/Runtime/Engine/Private/Streaming/SimpleStreamableAssetManager.cpp b/Engine/Source/Runtime/Engine/Private/Streaming/SimpleStreamableAssetManager.cpp
--- a/Engine/Source/Runtime/Engine/Private/Streaming/SimpleStreamableAssetManager.cpp
+++ b/Engine/Source/Runtime/Engine/Private/Streaming/SimpleStreamableAssetManager.cpp
@@ -219,6 +219,7 @@ uint32 FSimpleStreamableAssetManager::GetAllocatedSize_Impl() const
 	{
 		AllocatedSize += AssetBoundElements.GetAllocatedSize();
 	}
+	AllocatedSize += AssetIndexToSharedIndex.GetAllocatedSize();
 
 	AllocatedSize += BoundsViewInfos.GetAllocatedSize();
 
@@ -331,6 +332,7 @@ void FSimpleStreamableAssetManager::UpdateTask_Async()
 					--UsedAssetIndices;
 					FreeAssetIndexHint = FMath::Min(FreeAssetIndexHint, AssetIndex);
 					AssetIndexToBounds4Index[AssetIndex].Empty();
+					AssetIndexToSharedIndex[AssetIndex].Reset();
 					*Record.SimpleStreamableAssetManagerIndex = INDEX_NONE;
 				}
 			}
@@ -550,6 +552,12 @@ void FSimpleStreamableAssetManager::AddRenderAssetElements(const TArrayView<FReg
 		{
 			check(AssetIndex == AssetIndexToBounds4Index.Num());
 			AssetIndexToBounds4Index.Add({});
+			AssetIndexToSharedIndex.Add(RegisteredAsset.SimpleStreamableAssetManagerIndex);
+		}
+		else if (!AssetIndexToSharedIndex[AssetIndex])
+		{
+			// Re-populate after a previous zombie cleanup reclaimed this slot
+			AssetIndexToSharedIndex[AssetIndex] = RegisteredAsset.SimpleStreamableAssetManagerIndex;
 		}
 		
 		const int32 AssetElementIndex = AssetIndexToBounds4Index[AssetIndex].Add(FAssetBoundElement{
@@ -584,6 +592,22 @@ void FSimpleStreamableAssetManager::RemoveRenderAssetElements(int32 ObjectRegist
 		check(AssetElementIndex != INDEX_NONE);
 
 		AssetIndexToBounds4Index[AssetIndex].Reset(AssetElementIndex);
+
+		// When all objects referencing an asset are removed, fully
+		// reclaim the asset slot to prevent unbounded memory growth when
+		// assets outlive the objects that reference them.
+		if (AssetIndexToBounds4Index[AssetIndex].Num() == 0)
+		{
+			AssetIndexToBounds4Index[AssetIndex].Empty();
+			AssetUsedIndices[AssetIndex] = false;
+			--UsedAssetIndices;
+			FreeAssetIndexHint = FMath::Min(FreeAssetIndexHint, AssetIndex);
+			if (AssetIndexToSharedIndex[AssetIndex])
+			{
+				*AssetIndexToSharedIndex[AssetIndex] = INDEX_NONE;
+				AssetIndexToSharedIndex[AssetIndex].Reset();
+			}
+		}
 	}
 	ObjectAssets.Reset();
 #if DO_CHECK
diff --git a/Engine/Source/Runtime/Engine/Private/Tests/Streaming/SSAMCoreTests.cpp b/Engine/Source/Runtime/Engine/Private/Tests/Streaming/SSAMCoreTests.cpp
--- a/Engine/Source/Runtime/Engine/Private/Tests/Streaming/SSAMCoreTests.cpp
+++ b/Engine/Source/Runtime/Engine/Private/Tests/Streaming/SSAMCoreTests.cpp
@@ -1210,15 +1210,12 @@ namespace SSAMTests
 		return true;
 	}
 
-// DISABLED - Currently fails - Fix pending
-#ifdef TEST_SSAM_MEMORY_BLOAT_ON_LEVEL_TRANSITION
 	// -----------------------------------------------------------------------
 	// UE-363476: Simulates level transitions where each "level" introduces
 	// unique assets with many objects. Objects are unregistered between
-	// levels but assets stay registered (cached). Verifies that the inner
-	// sparse arrays in AssetIndexToBounds4Index release memory when their
-	// element count drops to zero, preventing unbounded growth across
-	// transitions.
+	// levels but assets stay registered (cached). Verifies that zombie
+	// asset registrations are fully reclaimed when their element count
+	// drops to zero, preventing unbounded growth across transitions.
 	// -----------------------------------------------------------------------
 	IMPLEMENT_SIMPLE_AUTOMATION_TEST(FSSAMMemoryBloatOnLevelTransition, SSAM_TEST_NAME ".Lifecycle.MemoryBloatOnLevelTransition", TestFlags)
 
@@ -1274,14 +1271,14 @@ namespace SSAMTests
 		}
 		const uint32 SizeAfterAllLevels = FSimpleStreamableAssetManager::GetAllocatedSize();
 
-		// Expected growth: the outer TArray<TSimpleSparseArray> gains 90 new
-		// entries (9 levels × 10 assets). Each TSimpleSparseArray is ~24 bytes
-		// + outer TArray slack ≈ ~8KB total. If inner sparse arrays are not
-		// freed, each retains its peak capacity (TBitArray + TArray of
-		// FAssetBoundElement for ~10+ elements each), adding ~75KB.
+		// With the fix, zombie asset slots are fully reclaimed when all objects
+		// are removed: inner buffers freed, index returned to pool, shared_ptr
+		// invalidated. Each subsequent level reuses the freed indices, so
+		// growth should be near zero. Without the fix, each level's inner
+		// sparse arrays retain peak capacity (~75KB+ total bloat).
 		//
-		// Threshold: 16KB allows for outer array growth + alignment but
-		// catches the ~75KB+ bloat from unfixed inner array retention.
+		// Threshold: 16KB catches unfixed bloat while allowing for minor
+		// TArray slack and alignment overhead.
 		const int64 MemoryGrowth = (int64)SizeAfterAllLevels - (int64)SizeAfterLevel1;
 
 		UE_LOG(LogTemp, Display, TEXT("UE-363476 test: SizeAfterLevel1=%u, SizeAfterAllLevels=%u, Growth=%lld bytes"),
@@ -1299,7 +1296,6 @@ namespace SSAMTests
 
 		return true;
 	}
-#endif // DISABLED_TEST_SSAM_MEMORY_BLOAT_ON_LEVEL_TRANSITION
 
 	// -----------------------------------------------------------------------
 	// Register 100 components, each referencing 1-3 of 3 shared assets.
diff --git a/Engine/Source/Runtime/Engine/Public/Streaming/SimpleStreamableAssetManager.h b/Engine/Source/Runtime/Engine/Public/Streaming/SimpleStreamableAssetManager.h
--- a/Engine/Source/Runtime/Engine/Public/Streaming/SimpleStreamableAssetManager.h
+++ b/Engine/Source/Runtime/Engine/Public/Streaming/SimpleStreamableAssetManager.h
@@ -328,6 +328,10 @@ private:
 	int32 UsedAssetIndices = 0;
 	TBitArray<> AssetUsedIndices;
 	TArray<TSimpleSparseArray<FAssetBoundElement>> AssetIndexToBounds4Index;
+	// Parallel to AssetIndexToBounds4Index. Stores the asset's shared index ptr
+	// so SSAM can invalidate the asset registration when all objects are removed,
+	// allowing the index slot to be reclaimed.
+	TArray<TSharedPtr<int32, ESPMode::ThreadSafe>> AssetIndexToSharedIndex;
 	void AddRenderAssetElements(const TArrayView<FRegisteredAsset>& RegisteredAssets, int32 ObjectRegistrationIndex, bool bForceMipStreaming);
 	void RemoveRenderAssetElements(int32 ObjectRegistrationIndex);
 
