Slate: Deprecate TouchpadIndex

[RN] Deprecated TouchpadIndex (there are no supported platforms with multiple multi-touch pads per single device)
#Synced-CL 52395449

diff --git a/Engine/Plugins/Runtime/CommonUI/Source/CommonUI/Private/CommonGameViewportClient.cpp b/Engine/Plugins/Runtime/CommonUI/Source/CommonUI/Private/CommonGameViewportClient.cpp
--- a/Engine/Plugins/Runtime/CommonUI/Source/CommonUI/Private/CommonGameViewportClient.cpp
+++ b/Engine/Plugins/Runtime/CommonUI/Source/CommonUI/Private/CommonGameViewportClient.cpp
@@ -76,7 +76,8 @@ bool UCommonGameViewportClient::InputAxis(const FInputKeyEventArgs& Args)
 	return Super::InputAxis(Args);
 }
 
-bool UCommonGameViewportClient::InputTouch(FViewport* InViewport, const FInputDeviceId DeviceId, uint32 Handle, ETouchType::Type Type, const FVector2D& TouchLocation, float Force, uint32 TouchpadIndex, const uint64 Timestamp)
+bool UCommonGameViewportClient::InputTouch(FViewport* const InViewport, const FInputDeviceId DeviceId, const ETouchIndex::Type TouchId, const ETouchType::Type Type,
+	const FVector2D& TouchLocation, const float Force, const uint64 Timestamp)
 {
 #if ALLOW_CONSOLE
 	if (ViewportConsole != nullptr && (ViewportConsole->ConsoleState != NAME_Typing) && (ViewportConsole->ConsoleState != NAME_Open))
@@ -90,10 +91,10 @@ bool UCommonGameViewportClient::InputTouch(FViewport* InViewport, const FInputDe
 		int32 ControllerId = IPlatformInputDeviceMapper::Get().GetUserIndexForPlatformUser(UserId);
 		 
 		PRAGMA_DISABLE_DEPRECATION_WARNINGS
-		if (!OnRerouteTouchInput().ExecuteIfBound(DeviceId, Handle, Type, TouchLocation, Result) &&
-			!OnRerouteTouch().ExecuteIfBound(ControllerId, Handle, Type, TouchLocation, Result))
+		if (!OnRerouteTouchInput().ExecuteIfBound(DeviceId, TouchId, Type, TouchLocation, Result) &&
+			!OnRerouteTouch().ExecuteIfBound(ControllerId, TouchId, Type, TouchLocation, Result))
 		{
-			HandleRerouteTouch(DeviceId, Handle, Type, TouchLocation, Result);
+			HandleRerouteTouch(DeviceId, TouchId, Type, TouchLocation, Result);
 		}
 		PRAGMA_ENABLE_DEPRECATION_WARNINGS
 
@@ -103,7 +104,7 @@ bool UCommonGameViewportClient::InputTouch(FViewport* InViewport, const FInputDe
 		}
 	}
 
-	return Super::InputTouch(InViewport, DeviceId, Handle, Type, TouchLocation, Force, TouchpadIndex, Timestamp);
+	return Super::InputTouch(InViewport, DeviceId, TouchId, Type, TouchLocation, Force, Timestamp);
 }
 
 void UCommonGameViewportClient::MouseMove(FViewport* InViewport, int32 X, int32 Y)
diff --git a/Engine/Plugins/Runtime/CommonUI/Source/CommonUI/Public/CommonGameViewportClient.h b/Engine/Plugins/Runtime/CommonUI/Source/CommonUI/Public/CommonGameViewportClient.h
--- a/Engine/Plugins/Runtime/CommonUI/Source/CommonUI/Public/CommonGameViewportClient.h
+++ b/Engine/Plugins/Runtime/CommonUI/Source/CommonUI/Public/CommonGameViewportClient.h
@@ -32,7 +32,7 @@ public:
 	// UGameViewportClient interface begin
 	UE_API virtual bool InputKey(const FInputKeyEventArgs& EventArgs) override;
 	UE_API virtual bool InputAxis(const FInputKeyEventArgs& Args) override;
-	UE_API virtual bool InputTouch(FViewport* Viewport, const FInputDeviceId DeviceId, uint32 Handle, ETouchType::Type Type, const FVector2D& TouchLocation, float Force, uint32 TouchpadIndex, const uint64 Timestamp) override;
+	UE_API virtual bool InputTouch(FViewport* const InViewport, const FInputDeviceId DeviceId, const ETouchIndex::Type TouchId, const ETouchType::Type Type, const FVector2D& TouchLocation, const float Force, const uint64 Timestamp) override;
 	UE_API virtual void MouseMove(FViewport* InViewport, int32 X, int32 Y) override;
 	UE_API virtual void CapturedMouseMove(FViewport* InViewport, int32 X, int32 Y) override;
 	// UGameViewportClient interface end
diff --git a/Engine/Source/Runtime/Engine/Classes/Engine/Console.h b/Engine/Source/Runtime/Engine/Classes/Engine/Console.h
--- a/Engine/Source/Runtime/Engine/Classes/Engine/Console.h
+++ b/Engine/Source/Runtime/Engine/Classes/Engine/Console.h
@@ -272,8 +272,29 @@ class UConsole
 	ENGINE_API virtual bool InputKey(FInputDeviceId DeviceId, FKey Key, EInputEvent Event, float AmountDepressed=1.f, bool bGamepad=false);
 	
 	virtual bool InputAxis(FInputDeviceId DevideId, FKey Key, float Delta, float DeltaTime, int32 NumSamples=1, bool bGamepad=false) { return false; };
-	
-	virtual bool InputTouch(FInputDeviceId DevideId, uint32 Handle, ETouchType::Type Type, const FVector2D& TouchLocation, float Force, FDateTime DeviceTimestamp, uint32 TouchpadIndex) { return false; }
+
+	UE_DEPRECATED(5.8, "Use the overload that omits TouchpadIndex instead")
+	virtual bool InputTouch(const FInputDeviceId DevideId, const uint32 Handle, const ETouchType::Type Type, const FVector2D& TouchLocation, const float Force, const FDateTime DeviceTimestamp, const uint32 TouchpadIndex) final
+	{
+		return InputTouch(DevideId, static_cast<ETouchIndex::Type>(Handle), Type, TouchLocation, Force, DeviceTimestamp.GetTicks());
+	}
+
+	/**
+	 * Allow the console to process a touch event received by the viewport-client.
+	 * If the console uses the event, it should return true.  If true is returned, the viewport-client will not offer it to the next candidate.
+	 * 
+	 * @param	DeviceId - The input device that was touched.
+	 * @param	TouchId - Identifier unique to this ongoing touch (i.e. finger).
+	 * @param	Type - What kind of touch event this is (see ETouchType).
+	 * @param	TouchLocation - Viewport-relative position of the touch.
+	 * @param	Force - How hard the touch is.
+	 * @param	Timestamp - The timestamp of when this input event was originally polled in the form of CPU cycles.
+	 * @return	True to consume the key event, false to pass it on.
+	 */
+	virtual bool InputTouch(const FInputDeviceId DeviceId, const ETouchIndex::Type TouchId, const ETouchType::Type Type, const FVector2D& TouchLocation, const float Force, const uint64 Timestamp)
+	{
+		return false;
+	}
 
 	ENGINE_API virtual void MouseMove(FViewport* Viewport, int32 X, int32 Y);
 	ENGINE_API virtual void CapturedMouseMove(FViewport* InViewport, int32 X, int32 Y);
diff --git a/Engine/Source/Runtime/Engine/Classes/Engine/GameViewportClient.h b/Engine/Source/Runtime/Engine/Classes/Engine/GameViewportClient.h
--- a/Engine/Source/Runtime/Engine/Classes/Engine/GameViewportClient.h
+++ b/Engine/Source/Runtime/Engine/Classes/Engine/GameViewportClient.h
@@ -172,7 +172,7 @@ public:
 	ENGINE_API virtual bool InputKey(const FInputKeyEventArgs& EventArgs) override;
 	ENGINE_API virtual bool InputAxis(const FInputKeyEventArgs& EventArgs) override;
 	ENGINE_API virtual bool InputChar(FViewport* Viewport,int32 ControllerId, TCHAR Character) override;
-	ENGINE_API virtual bool InputTouch(FViewport* InViewport, const FInputDeviceId DeviceId, uint32 Handle, ETouchType::Type Type, const FVector2D& TouchLocation, float Force, uint32 TouchpadIndex, const uint64 Timestamp) override;
+	ENGINE_API virtual bool InputTouch(FViewport* const InViewport, const FInputDeviceId DeviceId, const ETouchIndex::Type TouchId, const ETouchType::Type Type, const FVector2D& TouchLocation, const float Force, const uint64 Timestamp) override;
 	ENGINE_API virtual bool InputMotion(FViewport* Viewport, const FInputDeviceId DeviceId, const FVector& Tilt, const FVector& RotationRate, const FVector& Gravity, const FVector& Acceleration, const uint64 Timestamp) override;
 	ENGINE_API virtual EMouseCursor::Type GetCursor(FViewport* Viewport, int32 X, int32 Y ) override;
 	ENGINE_API virtual TOptional<TSharedRef<SWidget>> MapCursor(FViewport* Viewport, const FCursorReply& CursorReply) override;
diff --git a/Engine/Source/Runtime/Engine/Classes/GameFramework/PlayerController.h b/Engine/Source/Runtime/Engine/Classes/GameFramework/PlayerController.h
--- a/Engine/Source/Runtime/Engine/Classes/GameFramework/PlayerController.h
+++ b/Engine/Source/Runtime/Engine/Classes/GameFramework/PlayerController.h
@@ -1789,19 +1789,37 @@ public:
 	ENGINE_API virtual bool InputKey(const FInputKeyEventArgs& Params);
 
 	/** Handles a touch screen action */
-	UE_DEPRECATED(5.6, "Use the version which takes in a DeviceId and Timestamp instead.")
-	ENGINE_API virtual bool InputTouch(uint32 Handle, ETouchType::Type Type, const FVector2D& TouchLocation, float Force, FDateTime DeviceTimestamp, uint32 TouchpadIndex) final;
+	UE_DEPRECATED(5.8, "Use the overload that omits TouchpadIndex instead")
+	ENGINE_API virtual bool InputTouch(
+		const FInputDeviceId DeviceId,
+		const uint32 Handle,
+		const ETouchType::Type Type,
+		const FVector2D& TouchLocation,
+		const float Force,
+		const uint32 TouchpadIndex,
+		const uint64 Timestamp) final;
 
-	/** Handles a touch screen action */
+	/**
+	 * Allow the player-controller to process a touch event received by the viewport-client.
+	 * If the player-controller uses the event, it should return true.  If true is returned, the viewport-client will not offer it to the next
+	 * candidate.
+	 * 
+	 * @param	DeviceId - The input device that was touched.
+	 * @param	TouchId - Identifier unique to this ongoing touch (i.e. finger).
+	 * @param	Type - What kind of touch event this is (see ETouchType).
+	 * @param	TouchLocation - Viewport-relative position of the touch.
+	 * @param	Force - How hard the touch is.
+	 * @param	Timestamp - The timestamp of when this input event was originally polled in the form of CPU cycles.
+	 * @return	True to consume the event, false to pass it on.
+	 */
 	ENGINE_API virtual bool InputTouch(
 		const FInputDeviceId DeviceId,
-		uint32 Handle,
-		ETouchType::Type Type,
+		const ETouchIndex::Type TouchId,
+		const ETouchType::Type Type,
 		const FVector2D& TouchLocation,
-		float Force,		
-		uint32 TouchpadIndex,
+		const float Force,
 		const uint64 Timestamp);
-	
+
 	/** Handles motion control */
 	UE_DEPRECATED(5.6, "Use the version which takes in a DeviceId and Timestamp instead.")
 	ENGINE_API virtual bool InputMotion(const FVector& Tilt, const FVector& RotationRate, const FVector& Gravity, const FVector& Acceleration) final;
diff --git a/Engine/Source/Runtime/Engine/Classes/GameFramework/PlayerInput.h b/Engine/Source/Runtime/Engine/Classes/GameFramework/PlayerInput.h
--- a/Engine/Source/Runtime/Engine/Classes/GameFramework/PlayerInput.h
+++ b/Engine/Source/Runtime/Engine/Classes/GameFramework/PlayerInput.h
@@ -605,17 +605,23 @@ public:
 	ENGINE_API virtual bool InputKey(const FInputKeyEventArgs& Params);
 
 	/** Handles a touch input event.  Returns true. */
-	UE_DEPRECATED(5.6, "Use the version which takes a FInputDeviceId and Timestamp instead.")
-	ENGINE_API bool InputTouch(uint32 Handle, ETouchType::Type Type, const FVector2D& TouchLocation, float Force, FDateTime DeviceTimestamp, uint32 TouchpadIndex);
-	
+	UE_DEPRECATED(5.8, "Use the overload that omits TouchpadIndex instead")
+	ENGINE_API bool InputTouch(
+		const FInputDeviceId DeviceId,
+		const uint32 Handle,
+		const ETouchType::Type Type,
+		const FVector2D& TouchLocation,
+		const float Force,
+		const uint32 TouchpadIndex,
+		const uint64 Timestamp);
+
 	/** Handles a touch input event.  Returns true. */
 	ENGINE_API bool InputTouch(
 		const FInputDeviceId DeviceId,
-		uint32 Handle,
-		ETouchType::Type Type,
+		const ETouchIndex::Type TouchId,
+		const ETouchType::Type Type,
 		const FVector2D& TouchLocation,
-		float Force,
-		uint32 TouchpadIndex,
+		const float Force,
 		const uint64 Timestamp);
 
 	/** Handles a motion input event.  Returns true. */
diff --git a/Engine/Source/Runtime/Engine/Classes/Kismet/KismetInputLibrary.h b/Engine/Source/Runtime/Engine/Classes/Kismet/KismetInputLibrary.h
--- a/Engine/Source/Runtime/Engine/Classes/Kismet/KismetInputLibrary.h
+++ b/Engine/Source/Runtime/Engine/Classes/Kismet/KismetInputLibrary.h
@@ -307,6 +307,7 @@ class UKismetInputLibrary : public UBlueprintFunctionLibrary
 	static ENGINE_API int32 PointerEvent_GetPointerIndex(const FPointerEvent& Input);
 
 	/** Returns the index of the touch pad that generated this event (for platforms with multiple touch pads per user) */
+	UE_DEPRECATED(5.8, "TouchpadIndex is was never fully implemented and has been retired.")
 	UFUNCTION(BlueprintPure, meta=( DisplayName = "Get Touchpad Index" ), Category="Input|PointerEvent")
 	static ENGINE_API int32 PointerEvent_GetTouchpadIndex(const FPointerEvent& Input);
 
diff --git a/Engine/Source/Runtime/Engine/Private/GameViewportClient.cpp b/Engine/Source/Runtime/Engine/Private/GameViewportClient.cpp
--- a/Engine/Source/Runtime/Engine/Private/GameViewportClient.cpp
+++ b/Engine/Source/Runtime/Engine/Private/GameViewportClient.cpp
@@ -901,15 +901,8 @@ bool UGameViewportClient::InputChar(FViewport* InViewport, int32 ControllerId, T
 	return bResult;
 }
 
-bool UGameViewportClient::InputTouch(
-	FViewport* InViewport,
-	const FInputDeviceId DeviceId,
-	uint32 Handle,
-	ETouchType::Type Type,
-	const FVector2D& TouchLocation,
-	float Force,
-	uint32 TouchpadIndex,
-	const uint64 Timestamp)
+bool UGameViewportClient::InputTouch(FViewport* const InViewport, const FInputDeviceId DeviceId, const ETouchIndex::Type TouchId, const ETouchType::Type Type,
+	const FVector2D& TouchLocation, const float Force, const uint64 Timestamp)
 {
 	if (IgnoreInput())
 	{
@@ -917,13 +910,13 @@ bool UGameViewportClient::InputTouch(
 	}
 
 	// route to subsystems that care
-	bool bResult = (ViewportConsole ? ViewportConsole->InputTouch(DeviceId, Handle, Type, TouchLocation, Force, Timestamp, TouchpadIndex) : false);
+	bool bResult = (ViewportConsole ? ViewportConsole->InputTouch(DeviceId, TouchId, Type, TouchLocation, Force, Timestamp) : false);
 	if (!bResult)
 	{
 		ULocalPlayer* const TargetPlayer = GEngine->GetLocalPlayerFromInputDevice(this, DeviceId);
 		if (TargetPlayer && TargetPlayer->PlayerController)
 		{
-			bResult = TargetPlayer->PlayerController->InputTouch(DeviceId, Handle, Type, TouchLocation, Force, TouchpadIndex, Timestamp);
+			bResult = TargetPlayer->PlayerController->InputTouch(DeviceId, TouchId, Type, TouchLocation, Force, Timestamp);
 		}
 	}
 
diff --git a/Engine/Source/Runtime/Engine/Private/KismetInputLibrary.cpp b/Engine/Source/Runtime/Engine/Private/KismetInputLibrary.cpp
--- a/Engine/Source/Runtime/Engine/Private/KismetInputLibrary.cpp
+++ b/Engine/Source/Runtime/Engine/Private/KismetInputLibrary.cpp
@@ -304,7 +304,9 @@ int32 UKismetInputLibrary::PointerEvent_GetPointerIndex(const FPointerEvent& Inp
 
 int32 UKismetInputLibrary::PointerEvent_GetTouchpadIndex(const FPointerEvent& Input)
 {
+PRAGMA_DISABLE_DEPRECATION_WARNINGS
 	return Input.GetTouchpadIndex();
+PRAGMA_ENABLE_DEPRECATION_WARNINGS
 }
 
 bool UKismetInputLibrary::PointerEvent_IsTouchEvent(const FPointerEvent& Input)
diff --git a/Engine/Source/Runtime/Engine/Private/PlayerController.cpp b/Engine/Source/Runtime/Engine/Private/PlayerController.cpp
--- a/Engine/Source/Runtime/Engine/Private/PlayerController.cpp
+++ b/Engine/Source/Runtime/Engine/Private/PlayerController.cpp
@@ -2484,40 +2484,32 @@ bool APlayerController::InputKey(const FInputKeyEventArgs& Params)
 	return bResult;
 }
 
-PRAGMA_DISABLE_DEPRECATION_WARNINGS
-bool APlayerController::InputTouch(uint32 Handle, ETouchType::Type Type, const FVector2D& TouchLocation, float Force, FDateTime DeviceTimestamp, uint32 TouchpadIndex)
-{
-	// Use the primary input device id for this player controller. The deprecated code path doesnt have any other way to grab it
-	// and this is a sensible default
-	const FInputDeviceId DeviceId = IPlatformInputDeviceMapper::Get().GetPrimaryInputDeviceForUser(GetPlatformUserId());
-
-	// A timestamp of 0 is invalid because we don't have any way to accurately get one here from deprecated callsites.
-	constexpr uint64 Timestamp = 0u;
-	
-	return InputTouch(DeviceId, Handle, Type, TouchLocation, Force, TouchpadIndex, Timestamp);
-}
-PRAGMA_ENABLE_DEPRECATION_WARNINGS
-
 bool APlayerController::InputTouch(
 	const FInputDeviceId DeviceId,
-	uint32 Handle,
-	ETouchType::Type Type,
+	const uint32 Handle,
+	const ETouchType::Type Type,
 	const FVector2D& TouchLocation,
-	float Force,		
-	uint32 TouchpadIndex,
+	const float Force,
+	const uint32 TouchpadIndex,
 	const uint64 Timestamp)
+{
+	return InputTouch(DeviceId, static_cast<ETouchIndex::Type>(Handle), Type, TouchLocation, Force, Timestamp);
+}
+
+bool APlayerController::InputTouch(const FInputDeviceId DeviceId, const ETouchIndex::Type TouchId, const ETouchType::Type Type, const FVector2D& TouchLocation,
+                                   const float Force, const uint64 Timestamp)
 {
 	bool bResult = false;
 	if (PlayerInput)
 	{
-		bResult = PlayerInput->InputTouch(DeviceId, Handle, Type, TouchLocation, Force, TouchpadIndex, Timestamp);
+		bResult = PlayerInput->InputTouch(DeviceId, TouchId, Type, TouchLocation, Force, Timestamp);
 
 		if (bEnableTouchEvents || bEnableTouchOverEvents)
 		{
 			FHitResult HitResult;
 			const bool bHit = GetHitResultAtScreenPosition(TouchLocation, CurrentClickTraceChannel, true, HitResult);
 
-			ProcessTouchHitResult(DeviceId, Handle, Type, TouchLocation, HitResult, bHit);
+			ProcessTouchHitResult(DeviceId, TouchId, Type, TouchLocation, HitResult, bHit);
 		}
 	}
 
diff --git a/Engine/Source/Runtime/Engine/Private/Slate/SceneViewport.cpp b/Engine/Source/Runtime/Engine/Private/Slate/SceneViewport.cpp
--- a/Engine/Source/Runtime/Engine/Private/Slate/SceneViewport.cpp
+++ b/Engine/Source/Runtime/Engine/Private/Slate/SceneViewport.cpp
@@ -906,7 +906,10 @@ FReply FSceneViewport::OnTouchStarted( const FGeometry& MyGeometry, const FPoint
 
 		const FVector2D TouchPosition = CachedCursorPos;
 
-		if(ViewportClient->InputTouch( this, TouchEvent.GetInputDeviceId(), TouchEvent.GetPointerIndex(), ETouchType::Began, TouchPosition, TouchEvent.GetTouchForce(), TouchEvent.GetTouchpadIndex(), TouchEvent.GetEventTimestamp()) )
+		const bool bEventConsumed = ViewportClient->InputTouch(this, TouchEvent.GetInputDeviceId(),
+			static_cast<ETouchIndex::Type>(TouchEvent.GetPointerIndex()), ETouchType::Began, TouchPosition, TouchEvent.GetTouchForce(),
+			TouchEvent.GetEventTimestamp());
+		if (bEventConsumed)
 		{
 			const bool bTemporaryCapture = ViewportClient->GetMouseCaptureMode() == EMouseCaptureMode::CaptureDuringMouseDown;
 			if (bTemporaryCapture)
@@ -938,8 +941,10 @@ FReply FSceneViewport::OnTouchMoved( const FGeometry& MyGeometry, const FPointer
 		// Switch to the viewport clients world before processing input
 		FScopedConditionalWorldSwitcher WorldSwitcher( ViewportClient );
 
-
-		if( !ViewportClient->InputTouch( this, TouchEvent.GetInputDeviceId(), TouchEvent.GetPointerIndex(), ETouchType::Moved, CachedCursorPos, TouchEvent.GetTouchForce(), TouchEvent.GetTouchpadIndex(), TouchEvent.GetEventTimestamp()) )
+		const bool bEventConsumed = ViewportClient->InputTouch(this, TouchEvent.GetInputDeviceId(),
+			static_cast<ETouchIndex::Type>(TouchEvent.GetPointerIndex()), ETouchType::Moved, CachedCursorPos, TouchEvent.GetTouchForce(),
+			TouchEvent.GetEventTimestamp());
+		if (!bEventConsumed)
 		{
 			CurrentReplyState = FReply::Unhandled(); 
 		}
@@ -974,7 +979,10 @@ FReply FSceneViewport::OnTouchEnded( const FGeometry& MyGeometry, const FPointer
 		// Switch to the viewport clients world before processing input
 		FScopedConditionalWorldSwitcher WorldSwitcher( ViewportClient );
 
-		if( !ViewportClient->InputTouch( this, TouchEvent.GetInputDeviceId(), TouchEvent.GetPointerIndex(), ETouchType::Ended, CurCursorPos, 0.0f, TouchEvent.GetTouchpadIndex(), TouchEvent.GetEventTimestamp()) )
+		const bool bEventConsumed = ViewportClient->InputTouch(this, TouchEvent.GetInputDeviceId(),
+			static_cast<ETouchIndex::Type>(TouchEvent.GetPointerIndex()), ETouchType::Ended, CurCursorPos, /* Force */ 0.0f,
+			TouchEvent.GetEventTimestamp());
+		if (!bEventConsumed)
 		{
 			CurrentReplyState = FReply::Unhandled(); 
 		}
@@ -1003,7 +1011,10 @@ FReply FSceneViewport::OnTouchForceChanged(const FGeometry& MyGeometry, const FP
 
 		const FVector2D TouchPosition = MyGeometry.AbsoluteToLocal(TouchEvent.GetScreenSpacePosition()) * MyGeometry.Scale;
 
-		if (!ViewportClient->InputTouch(this, TouchEvent.GetInputDeviceId(), TouchEvent.GetPointerIndex(), ETouchType::ForceChanged, TouchPosition, TouchEvent.GetTouchForce(), TouchEvent.GetTouchpadIndex(), TouchEvent.GetEventTimestamp()))
+		const bool bEventConsumed = ViewportClient->InputTouch(this, TouchEvent.GetInputDeviceId(),
+			static_cast<ETouchIndex::Type>(TouchEvent.GetPointerIndex()), ETouchType::ForceChanged, TouchPosition, TouchEvent.GetTouchForce(), 
+			TouchEvent.GetEventTimestamp());
+		if (!bEventConsumed)
 		{
 			CurrentReplyState = FReply::Unhandled();
 		}
@@ -1027,7 +1038,10 @@ FReply FSceneViewport::OnTouchFirstMove(const FGeometry& MyGeometry, const FPoin
 
 		const FVector2D TouchPosition = MyGeometry.AbsoluteToLocal(TouchEvent.GetScreenSpacePosition()) * MyGeometry.Scale;
 
-		if (!ViewportClient->InputTouch(this, TouchEvent.GetInputDeviceId(), TouchEvent.GetPointerIndex(), ETouchType::FirstMove, TouchPosition, TouchEvent.GetTouchForce(), TouchEvent.GetTouchpadIndex(), TouchEvent.GetEventTimestamp()))
+		const bool bEventConsumed = ViewportClient->InputTouch(this, TouchEvent.GetInputDeviceId(),
+			static_cast<ETouchIndex::Type>(TouchEvent.GetPointerIndex()), ETouchType::FirstMove, TouchPosition, TouchEvent.GetTouchForce(),
+			TouchEvent.GetEventTimestamp());
+		if (!bEventConsumed)
 		{
 			CurrentReplyState = FReply::Unhandled();
 		}
diff --git a/Engine/Source/Runtime/Engine/Private/UnrealClient.cpp b/Engine/Source/Runtime/Engine/Private/UnrealClient.cpp
--- a/Engine/Source/Runtime/Engine/Private/UnrealClient.cpp
+++ b/Engine/Source/Runtime/Engine/Private/UnrealClient.cpp
@@ -2474,87 +2474,6 @@ void FViewportClient::RequestInvalidateHitProxy(FViewport* Viewport)
 	Viewport->InvalidateHitProxy();
 }
 
-PRAGMA_DISABLE_DEPRECATION_WARNINGS
-bool FViewportClient::InputAxis(
-		FViewport* Viewport,
-		FInputDeviceId InputDevice,
-		FKey Key,
-		float Delta,
-		float DeltaTime,
-		int32 NumSamples,
-		bool bGamepad)
-{
-	// Note: Set the timestamp to 0u here to represent an "invalid" timestamp, since this older input axis
-	// function did not account for timestamps. If we initialized it to the current cycle count here,
-	// then deprecated code paths would always be considered "newer" then anything calling the newer API, where the timestamp
-	// is calculated when the actual slate input event is constructed. 
-	FInputKeyEventArgs Args(Viewport, InputDevice, Key, Delta, DeltaTime, NumSamples, 0u);
-
-	return InputAxis(Args);
-}
-
-bool FViewportClient::InputTouch(FViewport* Viewport, int32 ControllerId, uint32 Handle, ETouchType::Type Type, const FVector2D& TouchLocation, float Force, FDateTime DeviceTimestamp, uint32 TouchpadIndex)
-{
-	// Get the newer FInputDeviceId from the legayc ControllerId
-	FPlatformUserId UserId = PLATFORMUSERID_NONE;
-	FInputDeviceId DeviceId = INPUTDEVICEID_NONE;
-	IPlatformInputDeviceMapper::Get().RemapControllerIdToPlatformUserAndDevice(ControllerId, OUT UserId, OUT DeviceId);
-	
-	// Note: Set the timestamp to 0u here to represent an "invalid" timestamp
-	constexpr uint64 Timestamp = 0u;
-	
-	return InputTouch(
-		Viewport,
-		DeviceId,
-		Handle,
-		Type,
-		TouchLocation,
-		Force,
-		TouchpadIndex,
-		Timestamp);
-}
-
-bool FViewportClient::InputGesture(FViewport* Viewport, EGestureEvent GestureType, const FVector2D& GestureDelta, bool bIsDirectionInvertedFromDevice)
-{
-	// Use the default platform user's primary device for any upgrade path to try and get the best behavior.
-	// Legacy behavior is that all touch/gesture events go through this primary user
-	const IPlatformInputDeviceMapper& Mapper = IPlatformInputDeviceMapper::Get();
-	const FPlatformUserId PrimaryUser = Mapper.GetPrimaryPlatformUser();
-	const FInputDeviceId DeviceId = IPlatformInputDeviceMapper::Get().GetPrimaryInputDeviceForUser(PrimaryUser);
-	
-	// Note: Set the timestamp to 0u here to represent an "invalid" timestamp
-	constexpr uint64 Timestamp = 0u;
-	
-	return InputGesture(
-		Viewport,
-		DeviceId,
-		GestureType,
-		GestureDelta,
-		bIsDirectionInvertedFromDevice,
-		Timestamp);
-}
-
-bool FViewportClient::InputMotion(FViewport* Viewport, int32 ControllerId, const FVector& Tilt, const FVector& RotationRate, const FVector& Gravity, const FVector& Acceleration)
-{
-	FPlatformUserId UserId = PLATFORMUSERID_NONE;
-	FInputDeviceId DeviceId = INPUTDEVICEID_NONE;
-	IPlatformInputDeviceMapper::Get().RemapControllerIdToPlatformUserAndDevice(ControllerId, OUT UserId, OUT DeviceId);
-	
-	// Note: Set the timestamp to 0u here to represent an "invalid" timestamp
-	constexpr uint64 Timestamp = 0u;
-
-	return InputMotion(
-		Viewport,
-		DeviceId,
-		Tilt,
-		RotationRate,
-		Gravity,
-		Acceleration,
-		Timestamp);
-}
-
-PRAGMA_ENABLE_DEPRECATION_WARNINGS
-
 TOptional<TSharedRef<SWidget>> FViewportClient::MapCursor(FViewport* Viewport, const FCursorReply& CursorReply)
 {
 	return TOptional<TSharedRef<SWidget>>();
diff --git a/Engine/Source/Runtime/Engine/Private/UserInterface/PlayerInput.cpp b/Engine/Source/Runtime/Engine/Private/UserInterface/PlayerInput.cpp
--- a/Engine/Source/Runtime/Engine/Private/UserInterface/PlayerInput.cpp
+++ b/Engine/Source/Runtime/Engine/Private/UserInterface/PlayerInput.cpp
@@ -478,27 +478,22 @@ bool UPlayerInput::InputKey(const FInputKeyEventArgs& Params)
 	}
 }
 
-PRAGMA_DISABLE_DEPRECATION_WARNINGS
-bool UPlayerInput::InputTouch(uint32 Handle, ETouchType::Type Type, const FVector2D& TouchLocation, float Force, FDateTime DeviceTimestamp, uint32 TouchpadIndex)
-{
-	// No way to get the input device from deprecated code here
-	const FInputDeviceId DeviceId = INPUTDEVICEID_NONE;
-
-	// A timestamp of 0 is invalid because we don't have any way to accurately get one here from deprecated callsites.
-	constexpr uint64 Timestamp = 0u;
-	
-	return InputTouch(DeviceId, Handle, Type, TouchLocation, Force, TouchpadIndex, Timestamp);
-}
-PRAGMA_ENABLE_DEPRECATION_WARNINGS
-
 bool UPlayerInput::InputTouch(
 		const FInputDeviceId DeviceId,
-		uint32 Handle,
-		ETouchType::Type Type,
+		const uint32 Handle,
+		const ETouchType::Type Type,
 		const FVector2D& TouchLocation,
-		float Force,
-		uint32 TouchpadIndex,
+		const float Force,
+		const uint32 TouchpadIndex,
 		const uint64 Timestamp)
+{
+	checkf(TouchpadIndex == 0, TEXT("We currently assume one touchpad in UPlayerInput::InputTouch. If this triggers, add support for multiple pads"));
+
+	return InputTouch(DeviceId, static_cast<ETouchIndex::Type>(Handle), Type, TouchLocation, Force, Timestamp);
+}
+
+bool UPlayerInput::InputTouch(const FInputDeviceId DeviceId, const ETouchIndex::Type Handle, const ETouchType::Type Type, const FVector2D& TouchLocation,
+	const float Force, const uint64 Timestamp)
 {
 	// if touch disabled from the command line, consume the input and do nothing
 	static bool bTouchDisabled = FParse::Param(FCommandLine::Get(), TEXT("DisableTouch"));
@@ -508,8 +503,6 @@ bool UPlayerInput::InputTouch(
 	}
 
 	// get the current state of each finger
-	checkf(TouchpadIndex == 0, TEXT("We currently assume one touchpad in UPlayerInput::InputTouch. If this triggers, add support for multiple pads"));
-
 	// if the handle is out of bounds, we can't handle it
 	if (Handle >= EKeys::NUM_TOUCH_KEYS)
 	{
diff --git a/Engine/Source/Runtime/Engine/Public/ViewportClient.h b/Engine/Source/Runtime/Engine/Public/ViewportClient.h
--- a/Engine/Source/Runtime/Engine/Public/ViewportClient.h
+++ b/Engine/Source/Runtime/Engine/Public/ViewportClient.h
@@ -5,7 +5,6 @@
 #include "Engine/EngineBaseTypes.h"
 #include "GenericPlatform/ICursor.h"
 #include "InputCoreTypes.h"
-#include "Misc/DateTime.h"
 
 class FCanvas;
 class FCursorReply;
@@ -38,8 +37,8 @@ public:
 	virtual struct FEngineShowFlags* GetEngineShowFlags() { return NULL; }
 
 	/**
-	 * Check a key event received by the viewport.
-	 * If the viewport client uses the event, it should return true to consume it.
+	 * Allow the viewport-client to process a key event received by the viewport.
+	 * If the viewport-client uses the event, it should return true.  If true is returned, the viewport should not bubble the event.
 	 * 
 	 * @param	EventArgs - The Input event args.
 	 * @return	True to consume the key event, false to pass it on.
@@ -47,31 +46,8 @@ public:
 	virtual bool InputKey(const FInputKeyEventArgs& EventArgs) { return false; }
 
 	/**
-	 * Check an axis movement received by the viewport.
-	 * If the viewport client uses the movement, it should return true to consume it.
-	 * @param	Viewport - The viewport which the axis movement is from.
-	 * @param	InputDevice - The input device that triggered this axis movement
-	 * @param	Key - The name of the axis which moved.
-	 * @param	Delta - The axis movement delta.
-	 * @param	DeltaTime - The time since the last axis update.
-	 * @param	NumSamples - The number of device samples that contributed to this Delta, useful for things like smoothing
-	 * @param	bGamepad - input came from gamepad (ie xbox controller)
-	 * @return	True to consume the axis movement, false to pass it on.
-	 */
-	UE_DEPRECATED(5.6, "Use the version of InputAxis which takes in a FInputKeyEventArgs instead.")
-	ENGINE_API virtual bool InputAxis(
-	 	FViewport* Viewport,
-	 	FInputDeviceId InputDevice,
-	 	FKey Key,
-	 	float Delta,
-	 	float DeltaTime,
-	 	int32 NumSamples = 1,
-	 	bool bGamepad = false) final;
-
-	/**
-	 * Check an axis movement received by the viewport. If this viewport client uses this input
-	 * and wants to consume it (stop it from being passed along to further viewports) then it should
-	 * return true.
+	 * Allow the viewport-client to process an axis event received by the viewport.
+	 * If the viewport-client uses the event, it should return true.  If true is returned, the viewport should not bubble the event.
 	 * 
 	 * @param Args The input event arguments for this axis movement.
 	 * @return True to consume the axis movement, false to pass it on.
@@ -79,10 +55,11 @@ public:
 	virtual bool InputAxis(const FInputKeyEventArgs& Args) { return false; }
 	
 	/**
-	 * Check a character input received by the viewport.
-	 * If the viewport client uses the character, it should return true to consume it.
-	 * @param	Viewport - The viewport which the axis movement is from.
-	 * @param	ControllerId - The controller which the axis movement is from.
+	 * Allow the viewport-client to process a character event received by the viewport.
+	 * If the viewport-client uses the event, it should return true.  If true is returned, the viewport should not bubble the event.
+	 * 
+	 * @param	Viewport - The viewport that received the event (is focused by the owner of the axis).
+	 * @param	ControllerId - The controller from which the character event was emitted.
 	 * @param	Character - The character.
 	 * @return	True to consume the character, false to pass it on.
 	 */
@@ -92,65 +69,63 @@ public:
 	 * Check a key event received by the viewport.
 	 * If the viewport client uses the event, it should return true to consume it.
 	 * @param	Viewport - The viewport which the event is from.
-	 * @param	ControllerId - The controller which the key event is from.
+	 * @param	DeviceId - The input device id from which this event is from
 	 * @param	Handle - Identifier unique to this touch event
 	 * @param	Type - What kind of touch event this is (see ETouchType)
 	 * @param	TouchLocation - Screen position of the touch
 	 * @param	Force - How hard the touch is
-	 * @param	DeviceTimestamp - Timestamp of the event
 	 * @param	TouchpadIndex - For devices with multiple touchpads, this is the index of which one
+	 * @param	Timestamp - The timestamp of when this input event was originally polled in the form of CPU cycles.
 	 * @return	True to consume the key event, false to pass it on.
 	 */
-	UE_DEPRECATED(5.6, "Use the version which takes a uint64 timestamp instead")
-	ENGINE_API virtual bool InputTouch(FViewport* Viewport, int32 ControllerId, uint32 Handle, ETouchType::Type Type, const FVector2D& TouchLocation, float Force, FDateTime DeviceTimestamp, uint32 TouchpadIndex) final;
+	UE_DEPRECATED(5.8, "Use the overload that omits TouchpadIndex instead")
+	virtual bool InputTouch(
+		FViewport* const Viewport,
+		const FInputDeviceId DeviceId,
+		const uint32 Handle,
+		const ETouchType::Type Type,
+		const FVector2D& TouchLocation,
+		const float Force,
+		const uint32 TouchpadIndex,
+		const uint64 Timestamp) final
+	{
+		return InputTouch(Viewport, DeviceId, static_cast<ETouchIndex::Type>(Handle), Type, TouchLocation, Force, Timestamp);
+	}
 
 	/**
-	 * Check a key event received by the viewport.
-	 * If the viewport client uses the event, it should return true to consume it.
-	 * @param	Viewport - The viewport which the event is from.
-	 * @param	DeviceId - The input device id from which this event is from
-	 * @param	Handle - Identifier unique to this touch event
-	 * @param	Type - What kind of touch event this is (see ETouchType)
-	 * @param	TouchLocation - Screen position of the touch
-	 * @param	Force - How hard the touch is
-	 * @param	TouchpadIndex - For devices with multiple touchpads, this is the index of which one
+	 * Allow the viewport client to process a touch event received by the viewport.
+	 * If the viewport client uses the event, it should return true.  If true is returned, the viewport should not bubble the event.
+	 * 
+	 * @param	Viewport - The viewport that was touched.
+	 * @param	DeviceId - The input device that was touched.
+	 * @param	TouchId - Identifier unique to this ongoing touch (i.e. finger).
+	 * @param	Type - What kind of touch event this is (see ETouchType).
+	 * @param	TouchLocation - Viewport-relative position of the touch.
+	 * @param	Force - How hard the touch is.
 	 * @param	Timestamp - The timestamp of when this input event was originally polled in the form of CPU cycles.
-	 * @return	True to consume the key event, false to pass it on.
+	 * @return	True to consume the event, false to pass it on.
 	 */
 	virtual bool InputTouch(
-		FViewport* Viewport,
+		FViewport* const Viewport,
 		const FInputDeviceId DeviceId,
-		uint32 Handle,
-		ETouchType::Type Type,
+		const ETouchIndex::Type TouchId,
+		const ETouchType::Type Type,
 		const FVector2D& TouchLocation,
-		float Force,		
-		uint32 TouchpadIndex,
+		const float Force,
 		const uint64 Timestamp)
 	{
 		return false;
 	}
 
 	/**
-	 * Check a gesture event received by the viewport.
-	 * If the viewport client uses the event, it should return true to consume it.
-	 * @param	Viewport - The viewport which the event is from.
-	 * @param	GestureType - @todo desc
-	 * @param	GestureDelta - @todo desc
-	 * @return	True to consume the gesture event, false to pass it on.
-	 */
-	UE_DEPRECATED(5.6, "Use the version which takes in a FInputDeviceId and uint64 Timestamp instead.")
-	ENGINE_API virtual bool InputGesture(FViewport* Viewport, EGestureEvent GestureType, const FVector2D& GestureDelta, bool bIsDirectionInvertedFromDevice) final;
-
-	/**
-	 * Check a gesture event received by the viewport.
-	 * If the viewport client uses the event, it should return true to consume it.
+	 * Allow the viewport-client to process a gesture event received by the viewport.
+	 * If the viewport-client uses the event, it should return true.  If true is returned, the viewport should not bubble the event.
 	 * 
-	 * @param	Viewport - The viewport which the event is from.
-	 * @param	DeviceId - The input device id from which this gesture came from
-	 * @param	GestureType - The type of gesture being input
-	 * @param	GestureDelta - Delta values of this gesture from the previous instance of it
+	 * @param	Viewport - The viewport on which the gesture was performed.
+	 * @param	DeviceId - The input device on which the gesture was performed.
+	 * @param	GestureType - The type of the gesture.
+	 * @param	GestureDelta - Delta values of this gesture from the previous instance of it.
 	 * @param	Timestamp - The timestamp of when this input event was originally polled in the form of CPU cycles.
-	 * 
 	 * @return	True to consume the gesture event, false to pass it on.
 	 */
 	virtual bool InputGesture(
@@ -165,29 +140,16 @@ public:
 	}
 
 	/**
-	 * Each frame, the input system will update the motion data.
+	 * Allow the viewport-client to process a motion event received by the viewport.
+	 * If the viewport-client uses the event, it should return true.  If true is returned, the viewport should not bubble the event.
 	 *
-	 * @param Viewport - The viewport which the key event is from.
-	 * @param ControllerId - The controller which the key event is from.
-	 * @param Tilt			The current orientation of the device
-	 * @param RotationRate	How fast the tilt is changing
-	 * @param Gravity		Describes the current gravity of the device
-	 * @param Acceleration  Describes the acceleration of the device
-	 * @return	True to consume the motion event, false to pass it on.
-	 */
-	UE_DEPRECATED(5.6, "Use the version which takes in a FInputDeviceId and a timestamp instead")
-	ENGINE_API virtual bool InputMotion(FViewport* Viewport, int32 ControllerId, const FVector& Tilt, const FVector& RotationRate, const FVector& Gravity, const FVector& Acceleration) final;
-
-	/**
-	 * Each frame, the input system will update the motion data.
-	 *
-	 * @param Viewport - The viewport which the key event is from.
-	 * @param DeviceId - The input device which this event is from
-	 * @param Tilt			The current orientation of the device
-	 * @param RotationRate	How fast the tilt is changing
-	 * @param Gravity		Describes the current gravity of the device
-	 * @param Acceleration  Describes the acceleration of the device
-	 * @param Timestamp The timestamp of when this input event was originally polled in the form of CPU cycles.
+	 * @param Viewport - The viewport that received the event (is focused by the owner of the accelerometer).
+	 * @param DeviceId - The input device that was moved.
+	 * @param Tilt - The current orientation of the device.
+	 * @param RotationRate - How fast the tilt is changing.
+	 * @param Gravity - Describes the current gravity of the device.
+	 * @param Acceleration - Describes the acceleration of the device.
+	 * @param Timestamp - The timestamp of when this input event was originally polled in the form of CPU cycles.
 	 * @return	True to consume the motion event, false to pass it on.
 	 */
 	virtual bool InputMotion(
@@ -204,9 +166,6 @@ public:
 
 	virtual void SetIsSimulateInEditorViewport(bool bInIsSimulateInEditorViewport) { };
 
-	UE_DEPRECATED(5.6, "This function will be removed in a future release.")
-	virtual bool WantsPollingMouseMovement(void) const { return true; }
-
 	virtual void MouseEnter( FViewport* Viewport,int32 x, int32 y ) {}
 
 	virtual void MouseLeave( FViewport* Viewport ) {}
diff --git a/Engine/Source/Runtime/SlateCore/Public/Input/Events.h b/Engine/Source/Runtime/SlateCore/Public/Input/Events.h
--- a/Engine/Source/Runtime/SlateCore/Public/Input/Events.h
+++ b/Engine/Source/Runtime/SlateCore/Public/Input/Events.h
@@ -853,19 +853,20 @@ public:
 		, bIsTouchFirstMove(false)
 	{ }
 
+	UE_DEPRECATED(5.8, "TouchpadIndex is was never fully implemented and has been retired.")
 	FPointerEvent(
-		uint32 InUserIndex,
-		uint32 InPointerIndex,
+		const uint32 InUserIndex,
+		const uint32 InPointerIndex,
 		const UE::Slate::FDeprecateVector2DParameter& InScreenSpacePosition,
 		const UE::Slate::FDeprecateVector2DParameter& InLastScreenSpacePosition,
-		float InForce,
-		bool bPressLeftMouseButton,
-		bool bInIsForceChanged = false,
-		bool bInIsFirstMove = false,
-		const FModifierKeysState& InModifierKeys = FModifierKeysState(),
-		uint32 InTouchpadIndex=0
-		)
-	: FInputEvent(InModifierKeys, InUserIndex, false)
+		const float InForce,
+		const bool bPressLeftMouseButton,
+		const bool bInIsForceChanged,
+		const bool bInIsFirstMove,
+		const FModifierKeysState& InModifierKeys,
+		const uint32 InTouchpadIndex
+	)
+		: FInputEvent(InModifierKeys, InUserIndex, false)
 		, ScreenSpacePosition(InScreenSpacePosition)
 		, LastScreenSpacePosition(InLastScreenSpacePosition)
 		, CursorDelta(FVector2f(InScreenSpacePosition) - FVector2f(InLastScreenSpacePosition))
@@ -881,18 +882,82 @@ public:
 		, bIsTouchForceChanged(bInIsForceChanged)
 		, bIsTouchFirstMove(bInIsFirstMove)
 	{ }
-	
+
 	FPointerEvent(
-		FInputDeviceId InDeviceId,
-		uint32 InPointerIndex,
+		const uint32 InUserIndex,
+		const uint32 InPointerIndex,
 		const UE::Slate::FDeprecateVector2DParameter& InScreenSpacePosition,
 		const UE::Slate::FDeprecateVector2DParameter& InLastScreenSpacePosition,
-		float InForce,
-		bool bPressLeftMouseButton,
-		bool bInIsForceChanged = false,
-		bool bInIsFirstMove = false,
+		const float InForce,
+		const bool bPressLeftMouseButton,
+		const bool bInIsForceChanged = false,
+		const bool bInIsFirstMove = false,
+		const FModifierKeysState& InModifierKeys = FModifierKeysState()
+	)
+		: FInputEvent(InModifierKeys, InUserIndex, false)
+		, ScreenSpacePosition(InScreenSpacePosition)
+		, LastScreenSpacePosition(InLastScreenSpacePosition)
+		, CursorDelta(FVector2f(InScreenSpacePosition) - FVector2f(InLastScreenSpacePosition))
+		, PressedButtons(bPressLeftMouseButton ? FTouchKeySet::StandardSet : FTouchKeySet::EmptySet)
+		, EffectingButton(EKeys::LeftMouseButton)
+		, PointerIndex(InPointerIndex)
+		, TouchpadIndex(0)
+		, Force(InForce)
+		, bIsTouchEvent(true)
+		, GestureType(EGestureEvent::None)
+		, WheelOrGestureDelta(0.0f, 0.0f)
+		, bIsDirectionInvertedFromDevice(false)
+		, bIsTouchForceChanged(bInIsForceChanged)
+		, bIsTouchFirstMove(bInIsFirstMove)
+	{
+	}
+
+	FPointerEvent(
+		const FInputDeviceId InDeviceId,
+		const uint32 InPointerIndex,
+		const UE::Slate::FDeprecateVector2DParameter& InScreenSpacePosition,
+		const UE::Slate::FDeprecateVector2DParameter& InLastScreenSpacePosition,
+		const float InForce,
+		const bool bPressLeftMouseButton,
+		const bool bInIsForceChanged = false,
+		const bool bInIsFirstMove = false,
 		const FModifierKeysState& InModifierKeys = FModifierKeysState(),
-		uint32 InTouchpadIndex=0,
+		const TOptional<int32> InOptionalSlateUserIndex = TOptional<int32>()
+	)
+	: FInputEvent(InModifierKeys, InDeviceId, false)
+		, ScreenSpacePosition(InScreenSpacePosition)
+		, LastScreenSpacePosition(InLastScreenSpacePosition)
+		, CursorDelta(FVector2f(InScreenSpacePosition) - FVector2f(InLastScreenSpacePosition))
+		, PressedButtons(bPressLeftMouseButton ? FTouchKeySet::StandardSet : FTouchKeySet::EmptySet)
+		, EffectingButton(EKeys::LeftMouseButton)
+		, PointerIndex(InPointerIndex)
+		, TouchpadIndex(0)
+		, Force(InForce)
+		, bIsTouchEvent(true)
+		, GestureType(EGestureEvent::None)
+		, WheelOrGestureDelta(0.0f, 0.0f)
+		, bIsDirectionInvertedFromDevice(false)
+		, bIsTouchForceChanged(bInIsForceChanged)
+		, bIsTouchFirstMove(bInIsFirstMove)
+	{
+		if (InOptionalSlateUserIndex.IsSet())
+		{
+			UserIndex = InOptionalSlateUserIndex.GetValue();
+		}
+	}
+
+	UE_DEPRECATED(5.8, "TouchpadIndex is was never fully implemented and has been retired.")
+	FPointerEvent(
+		const FInputDeviceId InDeviceId,
+		const uint32 InPointerIndex,
+		const UE::Slate::FDeprecateVector2DParameter& InScreenSpacePosition,
+		const UE::Slate::FDeprecateVector2DParameter& InLastScreenSpacePosition,
+		const float InForce,
+		const bool bPressLeftMouseButton,
+		const bool bInIsForceChanged,
+		const bool bInIsFirstMove,
+		const FModifierKeysState& InModifierKeys,
+		const uint32 InTouchpadIndex,
 		const TOptional<int32> InOptionalSlateUserIndex = TOptional<int32>()
 	)
 	: FInputEvent(InModifierKeys, InDeviceId, false)
@@ -979,7 +1044,7 @@ public:
 	/** Returns the unique identifier of the pointer (e.g., finger index) */
 	uint32 GetPointerIndex() const { return PointerIndex; }
 
-	/** Returns the index of the touch pad that generated this event (for platforms with multiple touch pads per user) */
+	UE_DEPRECATED(5.8, "TouchpadIndex is was never fully implemented and has been retired.")
 	uint32 GetTouchpadIndex() const { return TouchpadIndex; }
 
 	/** Returns the force of a touch (1.0f is mapped to an general touch force, < 1 is "light", > 1 is "heavy", and 10 is the max force possible) */
diff --git a/Engine/Source/Runtime/UMG/Private/Components/WidgetInteractionComponent.cpp b/Engine/Source/Runtime/UMG/Private/Components/WidgetInteractionComponent.cpp
--- a/Engine/Source/Runtime/UMG/Private/Components/WidgetInteractionComponent.cpp
+++ b/Engine/Source/Runtime/UMG/Private/Components/WidgetInteractionComponent.cpp
@@ -515,8 +515,7 @@ void UWidgetInteractionComponent::PressPointerKey(FKey Key)
 			false,
 			false,
 			FModifierKeysState(),
-			0,
-			VirtualUser->GetUserIndex());		
+			TOptional(VirtualUser->GetUserIndex()));
 	}
 	else
 	{
@@ -588,8 +587,7 @@ void UWidgetInteractionComponent::ReleasePointerKey(FKey Key)
 			false,
 			false,
 			FModifierKeysState(),
-			0,
-			VirtualUser->GetUserIndex());
+			TOptional(VirtualUser->GetUserIndex()));
 	}
 	else
 	{
