Back out 52358026 causing some loss of tranform information

#rnx
#Synced-CL 52399147

diff --git a/Engine/Source/Runtime/CoreUObject/Private/UObject/UObjectArchetype.cpp b/Engine/Source/Runtime/CoreUObject/Private/UObject/UObjectArchetype.cpp
--- a/Engine/Source/Runtime/CoreUObject/Private/UObject/UObjectArchetype.cpp
+++ b/Engine/Source/Runtime/CoreUObject/Private/UObject/UObjectArchetype.cpp
@@ -149,17 +149,6 @@ FORCENOINLINE UObject* FindArchetypeFromRequiredInfoImpl(const UClass* Class, co
 			{
 				Result = ArchetypeToSearch->GetClass()->FindArchetype(Class, Name);
 			}
-
-			if(!Result && (ObjectFlags & RF_ArchetypeObject))
-			{
-				if(const UClass* OuterClass = Cast<UClass>(Outer))
-				{
-					if(const UClass* SuperClass = OuterClass->GetSuperClass())
-					{
-						Result = SuperClass->FindArchetype(Class, Name);
-					}
-				}
-			}
 		}
 
 		if (!Result)
diff --git a/Engine/Source/Runtime/Engine/Private/BlueprintGeneratedClass.cpp b/Engine/Source/Runtime/Engine/Private/BlueprintGeneratedClass.cpp
--- a/Engine/Source/Runtime/Engine/Private/BlueprintGeneratedClass.cpp
+++ b/Engine/Source/Runtime/Engine/Private/BlueprintGeneratedClass.cpp
@@ -1572,30 +1572,6 @@ UObject* UBlueprintGeneratedClass::FindArchetype(const UClass* ArchetypeClass, c
 				Class = nullptr;
 			}
 		}
-		
-		if(Archetype == nullptr)
-		{
-			// We failed to find an archetype by name - but the ICH mangles
-			// the DefaultSceneRoot name *and* tries to cull unused templates
-			// so do a relaxed compare to find the real default scene root:
-			static const FName DefaultSceneRoot("DefaultSceneRoot");
-			static const FName DefaultSceneRootSCS("DefaultSceneRoot_GEN_VARIABLE");
-			static const FName ICHDefaultSceneRoot("ICH-DefaultSceneRoot_GEN_VARIABLE");
-			if(ArchetypeName == DefaultSceneRootSCS || ArchetypeName == ICHDefaultSceneRoot)
-			{
-				Class = const_cast<UBlueprintGeneratedClass*>(this);
-				while (Class && !Archetype)
-				{
-					// first look for an SCS template, then fall back to an unmangled DefaultSceneRoot.
-					Archetype = static_cast<UObject*>(FindObjectWithOuter(Class, ArchetypeClass, DefaultSceneRootSCS));
-					if(Archetype == nullptr)
-					{
-						Archetype = static_cast<UObject*>(FindObjectWithOuter(Class, ArchetypeClass, DefaultSceneRoot));
-					}
-					Class = Cast<UBlueprintGeneratedClass>(Class->GetSuperClass());
-				}
-			}
-		}
 	}
 
 	return Archetype;
