Add support to FScreenshotRequest::RequestScreenshot for only capturing the game viewport so that the entire editor is not captured when in PIE

#tests Testing of new view model and capture functionality to ensure it works as expected
#rnx

#Synced-CL 52420473

diff --git a/Engine/Source/Runtime/Engine/Private/GameViewportClient.cpp b/Engine/Source/Runtime/Engine/Private/GameViewportClient.cpp
--- a/Engine/Source/Runtime/Engine/Private/GameViewportClient.cpp
+++ b/Engine/Source/Runtime/Engine/Private/GameViewportClient.cpp
@@ -2247,6 +2247,15 @@ bool UGameViewportClient::ProcessScreenShots(FViewport* InViewport)
 			ensureMsgf(FScreenshotRequest::GetViewRect().IsEmpty(), TEXT("Specifying a screenshot view rect is not supported for slate widget screenshots (bShowUI=true). Using the full window size."));
 
 			TSharedRef<SWidget> WindowRef = WindowPtr.ToSharedRef();
+			if (FScreenshotRequest::ShouldRestrictToGameViewport())
+			{
+				TSharedPtr<SViewport> GameViewportWidgetPtr = GetGameViewportWidget();
+				if (GameViewportWidgetPtr.IsValid())
+				{
+					WindowRef = GameViewportWidgetPtr.ToSharedRef();
+				}
+			}
+
 			if (bHdrEnabled)
 			{
 				bScreenshotSuccessful = FSlateApplication::Get().TakeHDRScreenshot(WindowRef, BitmapHDR, Size);
diff --git a/Engine/Source/Runtime/Engine/Private/UnrealClient.cpp b/Engine/Source/Runtime/Engine/Private/UnrealClient.cpp
--- a/Engine/Source/Runtime/Engine/Private/UnrealClient.cpp
+++ b/Engine/Source/Runtime/Engine/Private/UnrealClient.cpp
@@ -208,13 +208,13 @@ const FTextureRHIRef& FRenderTarget::GetShaderResourceTexture() const
 	return GetRenderTargetTexture();
 }
 
-void FScreenshotRequest::RequestScreenshot(bool bInShowUI)
+void FScreenshotRequest::RequestScreenshot(bool bInShowUI, const bool bInRestrictToGameViewport)
 {
 	// empty string means we'll later pick the name
-	RequestScreenshot(TEXT(""), bInShowUI, /*bAddUniqueSuffix*/ true, /*bHdrScreenshot*/ false, FIntRect());
+	RequestScreenshot(TEXT(""), bInShowUI, /*bAddUniqueSuffix*/ true, /*bHdrScreenshot*/ false, FIntRect(), bInRestrictToGameViewport);
 }
 
-void FScreenshotRequest::RequestScreenshot(const FString& InFilename, bool bInShowUI, bool bAddUniqueSuffix, bool bHdrScreenshot, const FIntRect& InViewRect)
+void FScreenshotRequest::RequestScreenshot(const FString& InFilename, bool bInShowUI, bool bAddUniqueSuffix, bool bHdrScreenshot, const FIntRect& InViewRect, const bool bInRestrictToGameViewport)
 {
 	FString GeneratedFilename = InFilename;
 	CreateViewportScreenShotFilename(GeneratedFilename);
@@ -248,6 +248,7 @@ void FScreenshotRequest::RequestScreenshot(const FString& InFilename, bool bInSh
 	if (!Filename.IsEmpty())
 	{
 		bShowUI = bInShowUI;
+		bRestrictToGameViewport = bInRestrictToGameViewport;
 		ViewRect = InViewRect;
 		bIsScreenshotRequested = true;
 	}
@@ -267,6 +268,7 @@ void FScreenshotRequest::Reset()
 	bIsScreenshotRequested = false;
 	Filename.Empty();
 	bShowUI = false;
+	bRestrictToGameViewport = false;
 	ViewRect = FIntRect();
 }
 
@@ -343,6 +345,7 @@ bool FScreenshotRequest::bIsScreenshotRequested = false;
 FString FScreenshotRequest::Filename;
 FString FScreenshotRequest::NextScreenshotName;
 bool FScreenshotRequest::bShowUI = false;
+bool FScreenshotRequest::bRestrictToGameViewport = false;
 TArray<FColor> FScreenshotRequest::HighresScreenshotMaskColorArray;
 FIntPoint FScreenshotRequest::HighresScreenshotMaskExtents;
 FIntRect FScreenshotRequest::ViewRect = FIntRect();
diff --git a/Engine/Source/Runtime/Engine/Public/UnrealClient.h b/Engine/Source/Runtime/Engine/Public/UnrealClient.h
--- a/Engine/Source/Runtime/Engine/Public/UnrealClient.h
+++ b/Engine/Source/Runtime/Engine/Public/UnrealClient.h
@@ -198,8 +198,9 @@ struct FScreenshotRequest
 	 * to the ViewPort's OnScreenshotCaptured delegate.
 	 *
 	 * @param bInShowUI				Whether or not to show Slate UI
+	 * @param bInRestrictToGameViewport If true then the capture will be restricted to the game viewport, so will not incldue the entire editor in PIE for example
 	 */
-	static ENGINE_API void RequestScreenshot(bool bInShowUI);
+	static ENGINE_API void RequestScreenshot(bool bInShowUI, const bool bInRestrictToGameViewport = false);
 
 	/**
 	 * Requests a new screenshot with a specific filename
@@ -209,8 +210,9 @@ struct FScreenshotRequest
 	 * @param bAddFilenameSuffix	Whether an auto-generated unique suffix should be added to the supplied filename
 	 * @param bHdrScreenshot		If true, the screenshot is saved as .exr rather than .png
 	 * @param InViewRect			If non-empty, within the viewport extent, and bInShowUI=false, only this view rect will be included in the screenshot
+	 * @param bInRestrictToGameViewport If true then the capture will be restricted to the game viewport, so will not incldue the entire editor in PIE for example
 	 */
-	static ENGINE_API void RequestScreenshot(const FString& InFilename, bool bInShowUI, bool bAddFilenameSuffix, bool bHdrScreenshot = false, const FIntRect& InViewRect = FIntRect());
+	static ENGINE_API void RequestScreenshot(const FString& InFilename, bool bInShowUI, bool bAddFilenameSuffix, bool bHdrScreenshot = false, const FIntRect& InViewRect = FIntRect(), const bool bInRestrictToGameViewport = false);
 
 	/**
 	 * Resets a screenshot request
@@ -232,6 +234,11 @@ struct FScreenshotRequest
 	 */
 	static bool ShouldShowUI() { return bShowUI; }
 
+	/**
+	 * @return True if capture should be restricted to the game viewport
+	 */
+	static bool ShouldRestrictToGameViewport() { return bRestrictToGameViewport; }
+
 	/**
 	 * @return The view rect for the next screenshot
 	 */
@@ -275,6 +282,7 @@ private:
 	static ENGINE_API FString NextScreenshotName;
 	static ENGINE_API FString Filename;
 	static ENGINE_API bool bShowUI;
+	static ENGINE_API bool bRestrictToGameViewport;
 	static ENGINE_API TArray<FColor> HighresScreenshotMaskColorArray;
 	static ENGINE_API FIntPoint HighresScreenshotMaskExtents;
 	static ENGINE_API FIntRect ViewRect;
