Revert change that uses the scene extension if all sections are disabled.

#Synced-CL 52384345

diff --git a/Engine/Source/Runtime/Engine/Private/SkeletalRenderGPUSkin.cpp b/Engine/Source/Runtime/Engine/Private/SkeletalRenderGPUSkin.cpp
--- a/Engine/Source/Runtime/Engine/Private/SkeletalRenderGPUSkin.cpp
+++ b/Engine/Source/Runtime/Engine/Private/SkeletalRenderGPUSkin.cpp
@@ -623,8 +623,19 @@ FSkeletalMeshObjectGPUSkin::FSkeletalMeshObjectGPUSkin(const FSkinnedMeshScenePr
 
 		FSkeletalMeshObjectLOD& LOD = LODs.Emplace_GetRef(SkeletalMeshRenderData, LODIndex, InFeatureLevel, RecreateMorphVertexBuffer, GPUSkinTechnique);
 
-		// Fall back to vertex factory buffer allocation mode if any of the LOD's require it.
-		if (GPUSkinTechnique != ESkeletalMeshGPUSkinTechnique::Inline)
+		bool bHasValidSections = false;
+
+		for (const FSkelMeshRenderSection& Section : GetRenderSections(LODIndex))
+		{
+			if (Section.IsValid())
+			{
+				bHasValidSections = true;
+				break;
+			}
+		}
+
+		// Fall back to vertex factory buffer allocation mode if any of the LOD's require it or if an LOD has no valid sections in it.
+		if (!bHasValidSections || GPUSkinTechnique != ESkeletalMeshGPUSkinTechnique::Inline)
 		{
 			BoneTransformAllocationMode = EBoneTransformAllocationMode::VertexFactory;
 		}
