[HWRT] Refactored handling of evicted ray tracing geometry in FTextRenderSceneProxy.
- mark geometry as referenced and skip adding ray tracing instance when geometry is evicted instead of relying on lower level code handling it.

#Synced-CL 52420933

diff --git a/Engine/Source/Runtime/Engine/Private/Components/TextRenderComponent.cpp b/Engine/Source/Runtime/Engine/Private/Components/TextRenderComponent.cpp
--- a/Engine/Source/Runtime/Engine/Private/Components/TextRenderComponent.cpp
+++ b/Engine/Source/Runtime/Engine/Private/Components/TextRenderComponent.cpp
@@ -973,6 +973,12 @@ void FTextRenderSceneProxy::GetDynamicRayTracingInstances(FRayTracingInstanceCol
 			}
 		);
 	}
+	else if (RayTracingGeometry.IsEvicted())
+	{
+		Collector.AddReferencedGeometry(RayTracingGeometry.GetGeometryHandle());
+		return;
+	}
+
 	check(CachedRayTracingMaterials.Num() == RayTracingInstance.GetMaterials().Num());
 	checkf(RayTracingInstance.Geometry->GetInitializer().Segments.Num() == CachedRayTracingMaterials.Num(), TEXT("Segments/Materials mismatch. Number of segments: %d. Number of Materials: %d."),
 		RayTracingInstance.Geometry->GetInitializer().Segments.Num(),
