Early out when trying to precache default PSOs with null ptr - can happen during shutdown request at boot

#Synced-CL 52340434

diff --git a/Engine/Source/Runtime/Engine/Private/Materials/Material.cpp b/Engine/Source/Runtime/Engine/Private/Materials/Material.cpp
--- a/Engine/Source/Runtime/Engine/Private/Materials/Material.cpp
+++ b/Engine/Source/Runtime/Engine/Private/Materials/Material.cpp
@@ -715,7 +715,7 @@ void UMaterialInterface::PrecacheDefaultMaterialPSOs()
 
 void UMaterialInterface::PrecacheFallbackMaterialPSOs(UMaterialInterface* MaterialInterface)
 {
-	if (!GIsRHIInitialized)
+	if (!GIsRHIInitialized || (MaterialInterface == nullptr))
 	{
 		return;
 	}
