UX/UI improvements when working with Mirror Data Tables. Small fixes too.

Main changes:
- Deprecated UMirrorDataTable::FindReplaceMirroredNames() in favor of UMirrorDataTable::UpdateFromFindReplaceExpressions().
- New method handles more than just adding missing rows, it takes care of the following:
- Updating current entries.
- Disabling stale entries.
- Fixed computing runtime mirror arrays when updating table from skeleton using find/replace expressions.
- Added BoneScope property to control which bones are targeted when refreshing/reimporting from skeleton changes.
- Entries in mirror data tables can now be enabled/disabled.
- Added toolbar buttons in Data Table Editor when editing a Mirror Data Table:
- Reimport from Skeleton: Refresh table data using the associated skeleton.
- Validate: Checks if table is properly setup (missing entries, stale entries, partially completed entries, duplicate entries, non-reversible mappings, etc)
- Auto Pair: Update entries name pair using the find/replace expressions.
- Clear Sheet: Removes all rows from the table
- Added info icon for regex to inform users that expression matching is done from top to bottom with the first one winning.
- Added skeleton picker when using ExplicitBoneList.

Other:
- Removed deprecated UID-based curve mirror arrays.

#rnx minor

#Synced-CL 52307970

diff --git a/Engine/Plugins/Editor/EngineAssetDefinitions/Source/Private/Table/AssetDefinition_MirrorDataTable.cpp b/Engine/Plugins/Editor/EngineAssetDefinitions/Source/Private/Table/AssetDefinition_MirrorDataTable.cpp
--- a/Engine/Plugins/Editor/EngineAssetDefinitions/Source/Private/Table/AssetDefinition_MirrorDataTable.cpp
+++ b/Engine/Plugins/Editor/EngineAssetDefinitions/Source/Private/Table/AssetDefinition_MirrorDataTable.cpp
@@ -24,7 +24,7 @@ namespace MenuExtension_MirrorDataTable
 		FScopedTransaction Transaction(LOCTEXT("MirrorDataTable_FindReplace", "Update using Find Replace"));
 		for (UMirrorDataTable* MirrorDataTable : CBContext->LoadSelectedObjects<UMirrorDataTable>())
 		{
-			MirrorDataTable->FindReplaceMirroredNames(); 
+			MirrorDataTable->UpdateFromFindReplaceExpressions(UMirrorDataTable::FFindReplaceOptions::UpdateExisting().WithNotification()); 
 		}
 	}
 
@@ -40,7 +40,7 @@ namespace MenuExtension_MirrorDataTable
 						const UContentBrowserAssetContextMenuContext* CBContext = UContentBrowserAssetContextMenuContext::FindContextWithAssets(InSection);
 						{
 							const TAttribute<FText> Label = LOCTEXT("MirrorDataTable_FindReplace", "Update using Find Replace");
-							const TAttribute<FText> ToolTip = LOCTEXT("MirrorDataTable_FindReplaceTooltip", "Run find replace on all items, adding any new entries based on updated strings");
+							const TAttribute<FText> ToolTip = LOCTEXT("MirrorDataTable_FindReplaceTooltip", "Updates existing mirror pairs and disables stale ones. Missing entries are not added.");
 							const FSlateIcon Icon = FSlateIcon();
 
 							FToolUIAction UIAction;
diff --git a/Engine/Source/Editor/MirrorDataTableEditor/MirrorDataTableEditor.Build.cs b/Engine/Source/Editor/MirrorDataTableEditor/MirrorDataTableEditor.Build.cs
new file mode 100644
--- /dev/null
+++ b/Engine/Source/Editor/MirrorDataTableEditor/MirrorDataTableEditor.Build.cs
@@ -0,0 +1,37 @@
+﻿// Copyright Epic Games, Inc. All Rights Reserved.
+
+using UnrealBuildTool;
+
+[SupportedTargetTypes(TargetType.Editor, TargetType.Program)]
+public class MirrorDataTableEditor : ModuleRules
+{
+    public MirrorDataTableEditor(ReadOnlyTargetRules Target) : base(Target)
+    {
+        PCHUsage = ModuleRules.PCHUsageMode.UseExplicitOrSharedPCHs;
+
+        PublicDependencyModuleNames.AddRange(
+            new string[]
+            {
+                "Core", "SlateCore",
+            }
+        );
+
+        PrivateDependencyModuleNames.AddRange(
+            new string[]
+            {
+                "CoreUObject",
+                "Engine",
+                "Slate",
+                "SlateCore",
+                
+                "DataTableEditor",
+                "ToolMenus",
+                "UnrealEd",
+                "AssetTools",
+                "EditorStyle",
+                
+                "SkeletonEditor"
+            }
+        );
+    }
+}
\ No newline at end of file
diff --git a/Engine/Source/Editor/MirrorDataTableEditor/Private/MirrorDataTableCustomization.cpp b/Engine/Source/Editor/MirrorDataTableEditor/Private/MirrorDataTableCustomization.cpp
new file mode 100644
--- /dev/null
+++ b/Engine/Source/Editor/MirrorDataTableEditor/Private/MirrorDataTableCustomization.cpp
@@ -0,0 +1,730 @@
+﻿// Copyright Epic Games, Inc. All Rights Reserved.
+
+#include "MirrorDataTableCustomization.h"
+
+#include "DetailCategoryBuilder.h"
+#include "DetailLayoutBuilder.h"
+#include "DetailWidgetRow.h"
+#include "ISkeletonEditorModule.h"
+#include "ISkeletonTree.h"
+#include "SkeletonTreeBuilder.h"
+#include "SkeletonTreeItem.h"
+#include "Modules/ModuleManager.h"
+#include "ScopedTransaction.h"
+#include "Widgets/Images/SImage.h"
+#include "Widgets/Layout/SExpandableArea.h"
+#include "Widgets/Layout/SScrollBox.h"
+#include "Styling/AppStyle.h"
+#include "Widgets/Input/SCheckBox.h"
+#include "Widgets/Text/STextBlock.h"
+
+#define LOCTEXT_NAMESPACE "MirrorDataTableCustomization"
+
+/** Denotes a bone item is included in MirrorTable bone scope. */
+class FSkeletonTreeMirrorScopeItem : public FSkeletonTreeItem
+{
+public:
+	SKELETON_TREE_ITEM_TYPE(FSkeletonTreeMirrorScopeItem, FSkeletonTreeItem);
+
+	FSkeletonTreeMirrorScopeItem(const TWeakObjectPtr<UMirrorDataTable>& InMirrorDataTable, const FName& InBoneName, bool bInIsVirtualBone, const TSharedRef<ISkeletonTree>& InSkeletonTree)
+		: FSkeletonTreeItem(InSkeletonTree)
+		, MirrorTable(InMirrorDataTable)
+		, bIsVirtualBone(bInIsVirtualBone)
+		, BoneName(InBoneName)
+	{
+	};
+
+	/** ISkeletonTreeItem interface begin */
+	virtual void GenerateWidgetForNameColumn(TSharedPtr<SHorizontalBox> Box, const TAttribute<FText>& FilterText, FIsSelected InIsSelected) override
+	{
+		Box->AddSlot()
+		.AutoWidth()
+		.Padding(FMargin(0.0f, 1.0f))
+		.VAlign(VAlign_Center)
+		[
+			SNew(SCheckBox)
+			.IsChecked(this, &FSkeletonTreeMirrorScopeItem::GetCheckBoxState)
+			.OnCheckStateChanged(this, &FSkeletonTreeMirrorScopeItem::OnCheckBoxStateChanged)
+			.ToolTipText(this, &FSkeletonTreeMirrorScopeItem::GetBoneToolTip)
+		];
+
+		Box->AddSlot()
+		.AutoWidth()
+		.Padding(2, 0, 0, 0)
+		.VAlign(VAlign_Center)
+		[
+			SNew(STextBlock)
+			.ColorAndOpacity(this, &FSkeletonTreeMirrorScopeItem::GetBoneTextColor, InIsSelected)
+			.Text(FText::FromName(GetRowItemName()))
+			.HighlightText(FilterText)
+			.Font(this, &FSkeletonTreeMirrorScopeItem::GetBoneTextFont)
+			.ToolTipText(this, &FSkeletonTreeMirrorScopeItem::GetBoneToolTip)
+		];
+	}
+	
+	virtual TSharedRef<SWidget> GenerateWidgetForDataColumn(const FName& DataColumnName, FIsSelected InIsSelected) override
+	{
+		return SNullWidget::NullWidget;
+	}
+	
+	virtual FName GetRowItemName() const override
+	{
+		return BoneName;
+	}
+	/** ISkeletonTreeItem interface end */
+
+protected:
+
+	ECheckBoxState GetCheckBoxState() const
+	{
+		return (MirrorTable.IsValid() && MirrorTable->BoneScopeNameList.Contains(BoneName))
+			? ECheckBoxState::Checked
+			: ECheckBoxState::Unchecked;
+	}
+
+	void OnCheckBoxStateChanged(ECheckBoxState NewState)
+	{
+		if (MirrorTable.IsValid())
+		{
+			if (NewState == ECheckBoxState::Checked)
+			{
+				MirrorTable->BoneScopeNameList.Add(BoneName);
+			}
+			else
+			{
+				MirrorTable->BoneScopeNameList.Remove(BoneName);
+			}
+			
+			MirrorTable->InvalidateCachedSkeletonData();
+		}
+	}
+	
+	FSlateFontInfo GetBoneTextFont() const
+	{
+		const bool bInScope = MirrorTable.IsValid() && MirrorTable->BoneScopeNameList.Contains(BoneName);
+		return bInScope
+			? FAppStyle::GetWidgetStyle<FTextBlockStyle>("SkeletonTree.NormalFont").Font
+			: FAppStyle::GetWidgetStyle<FTextBlockStyle>("SkeletonTree.ItalicFont").Font;
+	}
+
+	FSlateColor GetBoneTextColor(FIsSelected InIsSelected) const
+	{
+		const bool bInScope = MirrorTable.IsValid() && MirrorTable->BoneScopeNameList.Contains(GetRowItemName());
+		if (bInScope)
+		{
+			return !bIsVirtualBone ? FSlateColor::UseForeground() : FSlateColor(FLinearColor(0.4f, 0.4f, 1.f));
+		}
+		return !bIsVirtualBone ? FSlateColor::UseSubduedForeground() : FSlateColor(FLinearColor(0.4f, 0.4f, 1.f, 0.5f));
+	};
+	
+	FText GetBoneToolTip() const
+	{
+		if (MirrorTable.IsValid() && MirrorTable->BoneScopeNameList.Contains(BoneName))
+		{
+			return LOCTEXT("IncludedBone_ToolTip", "This bone will be included when refreshing the mirror table. Right-click to remove it.");
+		}
+
+		return LOCTEXT("ExcludedBone_ToolTip", "This bone will be skipped when refreshing the mirror table. Right-click to include it.");
+	}
+	
+	TWeakObjectPtr<UMirrorDataTable> MirrorTable;
+
+	bool bIsVirtualBone;
+	
+	/** The name of the item in the tree */
+	FName BoneName;
+};
+
+/** Same as FSkeletonTreeBuilder but allows me to use FSkeletonTreeMirrorScopeItem. */
+class FMirrorDataTableEditorSkeletonTreeBuilder : public FSkeletonTreeBuilder
+{
+public:
+	FMirrorDataTableEditorSkeletonTreeBuilder(const TWeakObjectPtr<UMirrorDataTable>& InMirrorDataTable, const FSkeletonTreeBuilderArgs& InBuilderArgs)
+		: FSkeletonTreeBuilder(InBuilderArgs)
+		, MirrorTable(InMirrorDataTable)
+	{
+	}
+	
+	/** ISkeletonTreeBuilder interface begin */
+	virtual void Build(FSkeletonTreeBuilderOutput& Output) override
+	{
+		TSharedPtr<IEditableSkeleton> EditableSkeleton = EditableSkeletonPtr.Pin();
+		if (!EditableSkeleton || !EditableSkeleton->IsSkeletonValid())
+		{
+			UE_LOG(LogAnimation, Warning, TEXT("Skeleton tree builder - skeleton is invalid"));
+			return;
+		}
+
+		const USkeleton& Skeleton = EditableSkeleton->GetSkeleton();
+		const FReferenceSkeleton& RefSkeleton = Skeleton.GetReferenceSkeleton();
+
+		struct FBoneInfo
+		{
+			FBoneInfo(const FName& InBoneName, const FName& InParentName, int32 InDepth)
+				: BoneName(InBoneName)
+				, ParentName(InParentName)
+				, Depth(InDepth)
+			{
+				SortString = BoneName.ToString();
+				SortNumber = 0;
+				SortLength = SortString.Len();
+
+				// Split the bone name into string prefix and numeric suffix for sorting (different from FName to support leading zeros in the numeric suffix)
+				int32 Index = SortLength - 1;
+				for (int32 PlaceValue = 1; Index >= 0 && FChar::IsDigit(SortString[Index]); --Index, PlaceValue *= 10)
+				{
+					SortNumber += static_cast<int32>(SortString[Index] - '0') * PlaceValue;
+				}
+				SortString.LeftInline(Index + 1, EAllowShrinking::No);
+			}
+
+			bool operator<(const FBoneInfo& RHS)
+			{
+				// Sort parents before children
+				if (Depth != RHS.Depth)
+				{
+					return Depth < RHS.Depth;
+				}
+
+				// Sort alphabetically by string prefix
+				if (int32 SplitNameComparison = SortString.Compare(RHS.SortString))
+				{
+					return SplitNameComparison < 0;
+				}
+
+				// Sort by number if the string prefixes match
+				if (SortNumber != RHS.SortNumber)
+				{
+					return SortNumber < RHS.SortNumber;
+				}
+
+				// Sort by length to give us the equivalent to alphabetical sorting if the numbers match (which gives us the following sort order: bone_, bone_0, bone_00, bone_000, bone_001, bone_01, bone_1, etc)
+				return (SortNumber == 0) ? SortLength < RHS.SortLength : SortLength > RHS.SortLength;
+			}
+
+			FName BoneName = NAME_None;
+			FName ParentName = NAME_None;
+			int32 Depth = 0;
+
+			FString SortString;
+			int32 SortNumber = 0;
+			int32 SortLength = 0;
+		};
+
+		TArray<FBoneInfo> Bones;
+		Bones.Reserve(RefSkeleton.GetRawBoneNum());
+
+		// Gather the bones from the skeleton
+		for (int32 BoneIndex = 0; BoneIndex < RefSkeleton.GetRawBoneNum(); ++BoneIndex)
+		{
+			const FName& BoneName = RefSkeleton.GetBoneName(BoneIndex);
+
+			FName ParentName = NAME_None;
+			int32 ParentIndex = RefSkeleton.GetParentIndex(BoneIndex);
+			int32 Depth = 0;
+			if (ParentIndex != INDEX_NONE)
+			{
+				ParentName = RefSkeleton.GetBoneName(ParentIndex);
+				Depth = Bones[ParentIndex].Depth + 1;
+			}
+
+			Bones.Emplace(BoneName, ParentName, Depth);
+		}
+
+		// Sort the bones lexically (and also by depth in order to maintain the invariant of parents before children)
+		Algo::Sort(Bones);
+		
+		// Add the sorted bones to the skeleton tree
+		for (const FBoneInfo& Bone : Bones)
+		{
+			TSharedRef<ISkeletonTreeItem> DisplayBone = MakeShareable(new FSkeletonTreeMirrorScopeItem(MirrorTable, Bone.BoneName, false, SkeletonTreePtr.Pin().ToSharedRef()));
+			Output.Add(DisplayBone, Bone.ParentName, FSkeletonTreeMirrorScopeItem::GetTypeId());
+		}
+		
+		// Now add virtual bones.
+		const TArray<FVirtualBone>& VirtualBones = EditableSkeleton->GetSkeleton().GetVirtualBones();
+		for (const FVirtualBone& VirtualBone : VirtualBones)
+		{
+			TSharedRef<ISkeletonTreeItem>  DisplayBone = MakeShareable(new FSkeletonTreeMirrorScopeItem(MirrorTable, VirtualBone.VirtualBoneName, true, SkeletonTreePtr.Pin().ToSharedRef()));
+			Output.Add(DisplayBone, VirtualBone.SourceBoneName, FSkeletonTreeMirrorScopeItem::GetTypeId(), true);
+		}
+	}
+	/** ISkeletonTreeBuilder interface end */
+	
+protected:
+	
+	TWeakObjectPtr<UMirrorDataTable> MirrorTable;
+};
+
+TSharedRef<IDetailCustomization> FMirrorDataTableCustomization::MakeInstance()
+{
+	return MakeShareable(new FMirrorDataTableCustomization);
+}
+
+void FMirrorDataTableCustomization::CustomizeDetails(IDetailLayoutBuilder& DetailBuilder)
+{
+	DetailLayoutBuilder = &DetailBuilder;
+
+	// Get mirror table being edited/customized
+	{
+		TArray<TWeakObjectPtr<UObject>> Objects;
+		DetailBuilder.GetObjectsBeingCustomized(Objects);
+
+		if (Objects.Num() > 0)
+		{
+			MirrorTable = Cast<UMirrorDataTable>(Objects[0].Get());
+		}
+	}
+
+	// Setup skeleton tree widget
+	{
+		FSkeletonTreeBuilderArgs BuilderArgs;
+		BuilderArgs.bShowBones = true;
+		BuilderArgs.bShowSockets = false;
+		BuilderArgs.bShowAttachedAssets = false;
+		BuilderArgs.bShowVirtualBones = true;
+		
+		FSkeletonTreeArgs TreeArgs;
+		TreeArgs.Mode = ESkeletonTreeMode::Picker;
+		TreeArgs.bAllowMeshOperations = false;
+		TreeArgs.bAllowSkeletonOperations = false;
+		TreeArgs.bShowBlendProfiles = false;
+		TreeArgs.bShowFilterMenu = false;
+		TreeArgs.bHideBonesByDefault = false;
+		TreeArgs.bAllowInvisibleItemSelection = true;
+		TreeArgs.ContextName = TEXT("MirrorDataTableCustomization");
+		TreeArgs.Builder = SkeletonTreeBuilder = MakeShared<FMirrorDataTableEditorSkeletonTreeBuilder>(MirrorTable, BuilderArgs);
+		TreeArgs.Extenders = MakeShared<FExtender>();
+		TreeArgs.Extenders->AddMenuExtension("SkeletonTreeContextMenu", EExtensionHook::After, nullptr, FMenuExtensionDelegate::CreateLambda([this](FMenuBuilder& InMenuBuilder)
+			{
+				// Returns the given bone and all its descendants from the skeleton hierarchy.
+				// Works for both regular and virtual bones since both are present in FReferenceSkeleton.
+				auto GetBoneAndDescendants = [this](FName BoneName) -> TArray<FName>
+				{
+					TArray<FName> Result;
+					if (!MirrorTable.IsValid() || !MirrorTable->Skeleton)
+					{
+						return Result;
+					}
+
+					const FReferenceSkeleton& RefSkel = MirrorTable->Skeleton->GetReferenceSkeleton();
+					const int32 RootIndex = RefSkel.FindBoneIndex(BoneName);
+					if (RootIndex == INDEX_NONE)
+					{
+						return Result;
+					}
+
+					Result.Add(BoneName);
+
+					const int32 NumBones = RefSkel.GetNum();
+					for (int32 i = RootIndex + 1; i < NumBones; ++i)
+					{
+						int32 ParentIndex = RefSkel.GetParentIndex(i);
+						while (ParentIndex != INDEX_NONE)
+						{
+							if (ParentIndex == RootIndex)
+							{
+								Result.Add(RefSkel.GetBoneName(i)); break;
+							}
+							
+							if (ParentIndex < RootIndex)
+							{
+								break;
+							}
+							
+							ParentIndex = RefSkel.GetParentIndex(ParentIndex);
+						}
+					}
+						
+					return Result;
+				};
+
+				InMenuBuilder.BeginSection("MirrorTableScope", LOCTEXT("MirrorTableScopeSection", "Scope"));
+
+				InMenuBuilder.AddMenuEntry(
+					LOCTEXT("MenuLabel", "Include bone"),
+					LOCTEXT("MenuTooltip", "This bone will be included the next time the mirror table is refreshed from the skeleton."),
+					FSlateIcon(),
+					FUIAction(
+						FExecuteAction::CreateLambda([this]()
+						{
+							if (!MirrorTable.IsValid() || !SkeletonTree.IsValid())
+							{
+								return;
+							}
+								
+							FScopedTransaction Transaction(LOCTEXT("IncludeBone", "Include bone when refreshing"));
+
+							MirrorTable->Modify();
+
+							for (TSharedPtr<ISkeletonTreeItem> SelectedItem : SkeletonTree->GetSelectedItems())
+							{
+								MirrorTable->BoneScopeNameList.Add(SelectedItem->GetRowItemName());
+							}
+								
+							MirrorTable->InvalidateCachedSkeletonData();
+						}),
+						FCanExecuteAction::CreateLambda([this]()
+						{
+							if (!MirrorTable.IsValid() || !SkeletonTree.IsValid())
+							{
+								return false;
+							}
+								
+							for (TSharedPtr<ISkeletonTreeItem> SelectedItem : SkeletonTree->GetSelectedItems())
+							{
+								if (!MirrorTable->BoneScopeNameList.Contains(SelectedItem->GetRowItemName()))
+								{
+									return true;
+								}
+							}
+								
+							return false;
+						})
+					)
+				);
+
+				InMenuBuilder.AddMenuEntry(
+					LOCTEXT("RemoveMenuLabel", "Exclude bone"),
+					LOCTEXT("RemoveMenuTooltip", "This bone will be skipped the next time the mirror table is refreshed from the skeleton."),
+					FSlateIcon(),
+					FUIAction(
+						FExecuteAction::CreateLambda([this]()
+						{
+							if (!MirrorTable.IsValid() || !SkeletonTree.IsValid())
+							{
+								return;
+							}
+								
+							FScopedTransaction Transaction(LOCTEXT("ExcludeBone", "Exclude bone when refreshing"));
+
+							MirrorTable->Modify();
+
+							for (TSharedPtr<ISkeletonTreeItem> SelectedItem : SkeletonTree->GetSelectedItems())
+							{
+								MirrorTable->BoneScopeNameList.Remove(SelectedItem->GetRowItemName());
+							}
+								
+							MirrorTable->InvalidateCachedSkeletonData();
+						}),
+						FCanExecuteAction::CreateLambda([this]()
+						{
+							if (!MirrorTable.IsValid() || !SkeletonTree.IsValid())
+							{
+								return false;
+							}
+								
+							for (TSharedPtr<ISkeletonTreeItem> SelectedItem : SkeletonTree->GetSelectedItems())
+							{
+								if (MirrorTable->BoneScopeNameList.Contains(SelectedItem->GetRowItemName()))
+								{
+									return true;
+								}
+							}
+								
+							return false;
+						})
+					)
+				);
+
+				InMenuBuilder.AddMenuEntry(
+					LOCTEXT("IncludeWithChildrenLabel", "Include bone and children"),
+					LOCTEXT("IncludeWithChildrenTooltip", "This bone and all its children will be included the next time the mirror table is refreshed from the skeleton."),
+					FSlateIcon(),
+					FUIAction(
+						FExecuteAction::CreateLambda([this, GetBoneAndDescendants]()
+						{
+							if (!MirrorTable.IsValid() || !SkeletonTree.IsValid())
+							{
+								return;
+							}
+								
+							FScopedTransaction Transaction(LOCTEXT("IncludeBoneAndChildren", "Include bone and children when refreshing"));
+
+							MirrorTable->Modify();
+
+							for (TSharedPtr<ISkeletonTreeItem> SelectedItem : SkeletonTree->GetSelectedItems())
+							{
+								for (const FName& BoneName : GetBoneAndDescendants(SelectedItem->GetRowItemName()))
+								{
+									MirrorTable->BoneScopeNameList.Add(BoneName);
+								}
+							}
+								
+							MirrorTable->InvalidateCachedSkeletonData();
+						}),
+						FCanExecuteAction::CreateLambda([this, GetBoneAndDescendants]()
+						{
+							if (!MirrorTable.IsValid() || !SkeletonTree.IsValid())
+							{
+								return false;
+							}
+								
+							for (TSharedPtr<ISkeletonTreeItem> SelectedItem : SkeletonTree->GetSelectedItems())
+							{
+								for (const FName& BoneName : GetBoneAndDescendants(SelectedItem->GetRowItemName()))
+								{
+									if (!MirrorTable->BoneScopeNameList.Contains(BoneName))
+									{
+										return true;
+									}
+								}
+							}
+								
+							return false;
+						})
+					)
+				);
+
+				InMenuBuilder.AddMenuEntry(
+					LOCTEXT("ExcludeWithChildrenLabel", "Exclude bone and children"),
+					LOCTEXT("ExcludeWithChildrenTooltip", "This bone and all its children will be excluded the next time the mirror table is refreshed from the skeleton."),
+					FSlateIcon(),
+					FUIAction(
+						FExecuteAction::CreateLambda([this, GetBoneAndDescendants]()
+						{
+							if (!MirrorTable.IsValid() || !SkeletonTree.IsValid())
+							{
+								return;
+							}
+								
+							FScopedTransaction Transaction(LOCTEXT("ExcludeBoneAndChildren", "Exclude bone and children when refreshing"));
+
+							MirrorTable->Modify();
+
+							for (TSharedPtr<ISkeletonTreeItem> SelectedItem : SkeletonTree->GetSelectedItems())
+							{
+								for (const FName& BoneName : GetBoneAndDescendants(SelectedItem->GetRowItemName()))
+								{
+									MirrorTable->BoneScopeNameList.Remove(BoneName);
+								}
+							}
+								
+							MirrorTable->InvalidateCachedSkeletonData();
+						}),
+						FCanExecuteAction::CreateLambda([this, GetBoneAndDescendants]()
+						{
+							if (!MirrorTable.IsValid() || !SkeletonTree.IsValid())
+							{
+								return false;
+							}
+								
+							for (TSharedPtr<ISkeletonTreeItem> SelectedItem : SkeletonTree->GetSelectedItems())
+							{
+								for (const FName& BoneName : GetBoneAndDescendants(SelectedItem->GetRowItemName()))
+								{
+									if (MirrorTable->BoneScopeNameList.Contains(BoneName))
+									{
+										return true;
+									}
+								}
+							}
+							return false;
+						})
+					)
+				);
+
+				InMenuBuilder.EndSection();
+			})
+		);
+
+		ISkeletonEditorModule& SkeletonEditorModule = FModuleManager::LoadModuleChecked<ISkeletonEditorModule>("SkeletonEditor");
+
+		if (MirrorTable.IsValid() && MirrorTable->Skeleton)
+		{
+			SkeletonTree = SkeletonEditorModule.CreateSkeletonTree(MirrorTable->Skeleton, TreeArgs);
+		}
+	}
+	
+	// Add info icon to MirrorFindReplaceExpressions property (regex run order).
+	{
+		IDetailCategoryBuilder& CreateTableCategory = DetailBuilder.EditCategory("CreateTable");
+		TSharedPtr<IPropertyHandle> FindReplaceHandle = DetailBuilder.GetProperty(
+			GET_MEMBER_NAME_CHECKED(UMirrorDataTable, MirrorFindReplaceExpressions));
+
+		if (FindReplaceHandle.IsValid())
+		{
+			IDetailPropertyRow& Row = CreateTableCategory.AddProperty(FindReplaceHandle);
+
+			static const FText InfoText = LOCTEXT("FindReplaceOrderInfo",
+				"Expressions are evaluated from top to bottom. The first matching expression is used.");
+
+			Row.CustomWidget(true)
+			.NameContent()
+			[
+				SNew(SHorizontalBox)
+				.ToolTipText(InfoText)
+
+				+SHorizontalBox::Slot()
+				.VAlign(VAlign_Center)
+				.AutoWidth()
+				[
+					FindReplaceHandle->CreatePropertyNameWidget()
+				]
+
+				+SHorizontalBox::Slot()
+				.VAlign(VAlign_Center)
+				.AutoWidth()
+				.Padding(8.0f, 0.0f, 0.0f, 0.0f)
+				[
+					SNew(SImage)
+					.Image(FAppStyle::GetBrush("Icons.Info"))
+					.ColorAndOpacity(FSlateColor::UseSubduedForeground())
+				]
+			];
+		}
+	}
+
+	// Pin the 'Import' category to the bottom of the details panel
+	IDetailCategoryBuilder& ImportCategory = DetailBuilder.EditCategory("Import");
+	ImportCategory.SetSortOrder(INT_MAX);
+
+	BoneScopeHandle = DetailBuilder.GetProperty(GET_MEMBER_NAME_CHECKED(UMirrorDataTable, BoneScope));
+
+	// Add the BoneScope enum dropdown
+	ImportCategory.AddProperty(BoneScopeHandle);
+	
+	// Refresh the panel when enum changes so the next row updates immediately
+	BoneScopeHandle->SetOnPropertyValueChanged(FSimpleDelegate::CreateLambda([this]()
+		{
+			if (DetailLayoutBuilder)
+			{
+				DetailLayoutBuilder->ForceRefreshDetails();
+			}
+		})
+	);
+	
+	// Show skeleton tree bone widget only for specific bone selection
+	ImportCategory.AddCustomRow(LOCTEXT("SpecificBoneListFilter", "Specific Bone List"))
+	.Visibility(TAttribute<EVisibility>::CreateLambda([this]()
+		{
+			if (!BoneScopeHandle.IsValid())
+			{
+				return EVisibility::Collapsed;
+			}
+			
+			uint8 RawValue = 0;
+			if (BoneScopeHandle->GetValue(RawValue) != FPropertyAccess::Success)
+			{
+				return EVisibility::Collapsed;
+			}
+			
+			const EMirrorTableBoneRefreshScope Scope = static_cast<EMirrorTableBoneRefreshScope>(RawValue);
+			return Scope == EMirrorTableBoneRefreshScope::ExplicitBoneList ? EVisibility::Visible : EVisibility::Collapsed;
+		})
+	)
+	.WholeRowContent()
+	[
+		SNew(SBox)
+		.Padding(FMargin(-8.f, 0.f, 0.f, 0.f))
+		[
+			SNew(SExpandableArea)
+			.InitiallyCollapsed(false)
+			.BorderImage(FAppStyle::Get().GetBrush("NoBorder"))
+			.Padding(FMargin(0.f, 4.f, 0.f, 0.f))
+			.HeaderContent()
+			[
+				SNew(SHorizontalBox)
+
+				+ SHorizontalBox::Slot()
+				.AutoWidth()
+				.VAlign(VAlign_Center)
+				[
+					SNew(STextBlock)
+					.Text(LOCTEXT("SpecificBoneListLabel", "Specific Bones"))
+					.Font(IDetailLayoutBuilder::GetDetailFontBold())
+				]
+
+				+ SHorizontalBox::Slot()
+				.AutoWidth()
+				.VAlign(VAlign_Center)
+				.Padding(6.f, 0.f, 0.f, 0.f)
+				[
+					SNew(STextBlock)
+					.Text_Lambda([this]()
+					{
+						if (!MirrorTable.IsValid() || MirrorTable->BoneScopeNameList.Num() == 0)
+						{
+							return LOCTEXT("SelectedBoneCountEmpty", "(none — add bones below)");
+						}
+						return FText::Format(LOCTEXT("SelectedBoneCount", "({0} selected)"), MirrorTable->BoneScopeNameList.Num());
+					})
+					.Font(IDetailLayoutBuilder::GetDetailFontItalic())
+					.ColorAndOpacity(FSlateColor::UseSubduedForeground())
+				]
+
+				+ SHorizontalBox::Slot()
+				.FillWidth(1.f)
+
+				+ SHorizontalBox::Slot()
+				.AutoWidth()
+				.VAlign(VAlign_Center)
+				.Padding(0.f, 0.f, 4.f, 0.f)
+				[
+					SNew(SButton)
+					.Text(LOCTEXT("SelectAllBonesLabel", "Select All"))
+					.ToolTipText(LOCTEXT("SelectAllBonesTooltip", "Add all bones from the skeleton to the list."))
+					.IsEnabled_Lambda([this]()
+					{
+						if (!MirrorTable.IsValid() || !MirrorTable->Skeleton)
+						{
+							return false;
+						}
+						const int32 TotalBones = MirrorTable->Skeleton->GetReferenceSkeleton().GetRawBoneNum();
+						return MirrorTable->BoneScopeNameList.Num() < TotalBones;
+					})
+					.OnClicked_Lambda([this]()
+					{
+						if (MirrorTable.IsValid() && MirrorTable->Skeleton)
+						{
+							const FReferenceSkeleton& RefSkeleton = MirrorTable->Skeleton->GetReferenceSkeleton();
+							for (int32 i = 0; i < RefSkeleton.GetRawBoneNum(); ++i)
+							{
+								MirrorTable->BoneScopeNameList.Add(RefSkeleton.GetBoneName(i));
+							}
+							MirrorTable->InvalidateCachedSkeletonData();
+						}
+						return FReply::Handled();
+					})
+				]
+
+				+ SHorizontalBox::Slot()
+				.AutoWidth()
+				.VAlign(VAlign_Center)
+				[
+					SNew(SButton)
+					.Text(LOCTEXT("ClearAllBonesLabel", "Clear All"))
+					.ToolTipText(LOCTEXT("ClearAllBonesTooltip", "Clear the bone list and switch back to using the full skeleton."))
+					.OnClicked_Lambda([this]()
+					{
+						if (MirrorTable.IsValid() && BoneScopeHandle.IsValid())
+						{
+							MirrorTable->BoneScopeNameList.Empty();
+							MirrorTable->InvalidateCachedSkeletonData();
+							
+							BoneScopeHandle->SetValue(static_cast<uint8>(EMirrorTableBoneRefreshScope::FullSkeleton));
+						}
+						return FReply::Handled();
+					})
+				]
+			]
+			.BodyContent()
+			[
+				SNew(SBox)
+				.MinDesiredWidth(300.f)
+				.MaxDesiredHeight(600.f)
+				[
+					SNew(SScrollBox)
+					+ SScrollBox::Slot()
+					[
+						SkeletonTree.IsValid()
+						? StaticCastSharedRef<SWidget>(SkeletonTree.ToSharedRef())
+						: StaticCastSharedRef<SWidget>(SNew(STextBlock).Text(LOCTEXT("NoSkeleton", "No skeleton available.")))
+					]
+				]
+			]
+		]
+	];
+}
+
+#undef LOCTEXT_NAMESPACE
\ No newline at end of file
diff --git a/Engine/Source/Editor/MirrorDataTableEditor/Private/MirrorDataTableCustomization.h b/Engine/Source/Editor/MirrorDataTableEditor/Private/MirrorDataTableCustomization.h
new file mode 100644
--- /dev/null
+++ b/Engine/Source/Editor/MirrorDataTableEditor/Private/MirrorDataTableCustomization.h
@@ -0,0 +1,32 @@
+﻿// Copyright Epic Games, Inc. All Rights Reserved.
+
+#pragma once
+
+#include "CoreMinimal.h"
+#include "IDetailCustomization.h"
+#include "Animation/MirrorDataTable.h"
+
+class FSkeletonTreeBuilder;
+class ISkeletonTree;
+class IDetailLayoutBuilder;
+class SWidget;
+class IPropertyHandle;
+
+class FMirrorDataTableCustomization : public IDetailCustomization
+{
+public:
+	
+	static TSharedRef<IDetailCustomization> MakeInstance();
+
+	// IDetailCustomization interface
+	virtual void CustomizeDetails(IDetailLayoutBuilder& DetailBuilder) override;
+	// End of IDetailCustomization interface
+
+protected:
+	
+	TSharedPtr<FSkeletonTreeBuilder> SkeletonTreeBuilder;
+	TSharedPtr<ISkeletonTree> SkeletonTree;
+	TWeakObjectPtr<UMirrorDataTable> MirrorTable;
+	TSharedPtr<IPropertyHandle> BoneScopeHandle;
+	IDetailLayoutBuilder* DetailLayoutBuilder = nullptr;
+};
\ No newline at end of file
diff --git a/Engine/Source/Editor/MirrorDataTableEditor/Private/MirrorDataTableEditorModule.cpp b/Engine/Source/Editor/MirrorDataTableEditor/Private/MirrorDataTableEditorModule.cpp
new file mode 100644
--- /dev/null
+++ b/Engine/Source/Editor/MirrorDataTableEditor/Private/MirrorDataTableEditorModule.cpp
@@ -0,0 +1,617 @@
+﻿// Copyright Epic Games, Inc. All Rights Reserved.
+
+#include "MirrorDataTableEditorModule.h"
+
+#include "DataTableEditorModule.h"
+#include "DataTableEditorUtils.h"
+#include "IEditableSkeleton.h"
+#include "ISkeletonEditorModule.h"
+#include "MirrorDataTableCustomization.h"
+#include "MirrorDataTableValidator.h"
+#include "PropertyEditorModule.h"
+#include "ScopedTransaction.h"
+#include "ToolMenu.h"
+#include "ToolMenus.h"
+#include "Animation/AnimationSettings.h"
+#include "Animation/AnimCurveMetaData.h"
+#include "Animation/MirrorDataTable.h"
+#include "Framework/Commands/Commands.h"
+#include "Framework/Notifications/NotificationManager.h"
+#include "Logging/StructuredLog.h"
+#include "Subsystems/AssetEditorSubsystem.h"
+#include "Toolkits/AssetEditorToolkitMenuContext.h"
+#include "Widgets/Notifications/SNotificationList.h"
+
+#define LOCTEXT_NAMESPACE "FMirrorDataTableEditorModule"
+
+/** Returns the mirror data table being currently edited, if any. */
+static UMirrorDataTable* GetMirrorDataTableFromContext(const FToolMenuContext& Context)
+{
+	if (const UAssetEditorToolkitMenuContext* ToolkitMenuContext = Context.FindContext<UAssetEditorToolkitMenuContext>())
+	{
+		if (ToolkitMenuContext && ToolkitMenuContext->Toolkit.IsValid())
+		{
+			for (UObject* Object : ToolkitMenuContext->GetEditingObjects())
+			{
+				if (Object && Object->IsA<UMirrorDataTable>())
+				{
+					return Cast<UMirrorDataTable>(Object);
+				}
+			}
+		}
+	}
+
+	return nullptr;
+}
+
+void FMirrorDataTableEditorModule::StartupModule()
+{
+	// Add mirror table specific menus to data table editor toolbar.
+    UToolMenus::RegisterStartupCallback(FSimpleMulticastDelegate::FDelegate::CreateRaw(this, &FMirrorDataTableEditorModule::RegisterMenus));
+
+	// Register customization for mirror table.
+	FPropertyEditorModule& PropertyModule = FModuleManager::LoadModuleChecked<FPropertyEditorModule>("PropertyEditor");
+	PropertyModule.RegisterCustomClassLayout("MirrorDataTable", FOnGetDetailCustomizationInstance::CreateStatic(&FMirrorDataTableCustomization::MakeInstance));
+
+	// Bind immediately — does not require GEditor.
+	OnObjectPropertyChangedHandle = FCoreUObjectDelegates::OnObjectPropertyChanged.AddRaw(this, &FMirrorDataTableEditorModule::OnPropertyChanged);
+
+	// UAssetEditorSubsystem requires GEditor, which is not yet valid at StartupModule time. Defer until PostEngineInit.
+	FCoreDelegates::GetOnPostEngineInit().AddLambda([this]()
+		{
+			if (GEditor)
+			{
+				if (UAssetEditorSubsystem* AssetEditorSubsystem = GEditor->GetEditorSubsystem<UAssetEditorSubsystem>())
+				{
+					OnAssetOpenedHandle = AssetEditorSubsystem->OnAssetOpenedInEditor().AddRaw(this, &FMirrorDataTableEditorModule::OnAssetOpenedInEditor);
+				}
+			}
+		}
+	);
+}
+
+void FMirrorDataTableEditorModule::ShutdownModule()
+{
+    UToolMenus::UnRegisterStartupCallback(this);
+
+	if (FModuleManager::Get().IsModuleLoaded("PropertyEditor"))
+	{
+		FPropertyEditorModule& PropertyModule = FModuleManager::GetModuleChecked<FPropertyEditorModule>("PropertyEditor");
+		PropertyModule.UnregisterCustomClassLayout("MirrorDataTable");
+	}
+
+	// Unbind all active skeleton callbacks.
+	TArray<TWeakObjectPtr<UMirrorDataTable>> Keys;
+	ActiveSkeletonCallbacks.GetKeys(Keys);
+	for (const TWeakObjectPtr<UMirrorDataTable>& WeakMirrorTable : Keys)
+	{
+		if (UMirrorDataTable* MirrorTable = WeakMirrorTable.Get())
+		{
+			UnbindSkeletonCallbacks(MirrorTable);
+		}
+	}
+	ActiveSkeletonCallbacks.Empty();
+
+	// Unbind global delegates.
+	FCoreUObjectDelegates::OnObjectPropertyChanged.Remove(OnObjectPropertyChangedHandle);
+
+	if (GEditor)
+	{
+		if (UAssetEditorSubsystem* AssetEditorSubsystem = GEditor->GetEditorSubsystem<UAssetEditorSubsystem>())
+		{
+			AssetEditorSubsystem->OnAssetOpenedInEditor().Remove(OnAssetOpenedHandle);
+		}
+	}
+}
+
+void FMirrorDataTableEditorModule::RegisterMenus()
+{
+	FToolMenuOwnerScoped OwnerScoped(this);
+	
+	if (UToolMenu* MirrorToolBarMenu = UToolMenus::Get()->ExtendMenu("AssetEditor.DataTableEditor.ToolBar"))
+	{
+		MirrorToolBarMenu->AddDynamicSection("MirrorDataTableDynamicSection", FNewSectionConstructChoice(
+			FNewToolMenuDelegate::CreateLambda([this](UToolMenu* InMenu)
+				{
+					if (!InMenu)
+					{
+						return;
+					}
+
+					// Don't extend anything if we're editing any other data table type.
+					UMirrorDataTable* MirrorTable = GetMirrorDataTableFromContext(InMenu->Context);
+					if (!MirrorTable)
+					{
+						return;
+					}
+				
+					// Our commands will go here.
+					FToolMenuSection& Section = InMenu->FindOrAddSection("MirrorDataTable");
+
+					// Do a full sync/reimport from skeleton.
+					FToolMenuEntry& ReimportFromSkeleton = Section.AddEntry(FToolMenuEntry::InitToolBarButton(
+						"MDT_ReimportFromSkeleton",
+						FToolMenuExecuteAction::CreateLambda([](const FToolMenuContext& Context)
+							{
+								ExecuteMenuAction_ReimportFromSkeleton(Context, UMirrorDataTable::FFindReplaceOptions::Sync().WithNotification());
+							}
+						),
+						LOCTEXT("ReimportLabel", "Reimport from Skeleton"),
+						TAttribute<FText>::CreateLambda([MirrorTable]() -> FText
+						{
+							static const FText OperationDescription = LOCTEXT("ReimportOperationDescription", "Syncs the table with the skeleton: adds missing entries, updates existing ones, and disables stale entries.");
+								
+							if (!MirrorTable)
+							{
+								return OperationDescription;
+							}
+								
+							switch (MirrorTable->GetSkeletonSyncStatus())
+							{
+							case UMirrorDataTable::ESyncStatus::Stale:
+								return FText::Format(LOCTEXT("ReimportToolTip_Stale", "(Out of sync): {0}"), OperationDescription);
+							case UMirrorDataTable::ESyncStatus::UpToDate:
+								return FText::Format(LOCTEXT("ReimportToolTip_UpToDate", "(Up to date): {0}"), OperationDescription);
+							default:
+								return FText::Format(LOCTEXT("ReimportToolTip_Unknown", "(Unknown): {0}"), OperationDescription);
+							}
+						}),
+						TAttribute<FSlateIcon>::CreateStatic(&FMirrorDataTableEditorModule::GetStatusIcon, MirrorTable)
+					));
+					ReimportFromSkeleton.StyleNameOverride = FName("CalloutToolbar");
+					
+					// Reimport from skeleton partially.
+					FToolMenuEntry& ReimportFromSkeletonOptions = Section.AddEntry(FToolMenuEntry::InitComboButton(
+						"MDT_ReimportFromSkeletonOptions",
+						FUIAction(),
+						FNewToolMenuChoice(FNewToolMenuWidget::CreateLambda([this](const FToolMenuContext& Context)
+							{
+								static const FName MenuName = "AssetEditor.MirrorDataTableEditor.ReimportFromSkeletonMenu";
+
+								if (!UToolMenus::Get()->IsMenuRegistered(MenuName))
+								{
+									UToolMenu* Menu = UToolMenus::Get()->RegisterMenu(MenuName);
+									FToolMenuSection& ActionsSection = Menu->FindOrAddSection("Modes", LOCTEXT("ModesLabel", "Actions"));
+
+									// Only add missing entries.
+									ActionsSection.AddMenuEntry(
+										NAME_None,
+										LOCTEXT("AddMissingLabel", "Add Missing Entries"),
+										LOCTEXT("AddMissingTooltip", "Only add rows for skeleton entries that do not yet exist in the mirror data table."),
+										FSlateIcon(),
+										FToolMenuExecuteAction::CreateLambda([](const FToolMenuContext& Context)
+											{
+												ExecuteMenuAction_ReimportFromSkeleton(Context,
+													UMirrorDataTable::FFindReplaceOptions::AddMissingOnly().WithNotification()
+												);
+											}
+										)
+									);
+
+									// Update existing entries.
+									ActionsSection.AddMenuEntry(
+										NAME_None,
+										LOCTEXT("UpdateExistingEntriesLabel", "Update Existing Entries"),
+										LOCTEXT("UpdateExistingEntriesTooltip", "Only update existing rows."),
+										FSlateIcon(),
+										FToolMenuExecuteAction::CreateLambda([](const FToolMenuContext& Context)
+											{
+												ExecuteMenuAction_ReimportFromSkeleton(Context,
+													UMirrorDataTable::FFindReplaceOptions::UpdateExisting().WithNotification()
+												);
+											}
+										)
+									);
+									
+									ActionsSection.AddMenuEntry(
+										NAME_None,
+										LOCTEXT("DisableStaleEntriesLabel", "Disable Stale Entries"),
+										LOCTEXT("DisableStaleEntriesToolTip", "Only disables entries found in table but not in skeleton."),
+										FSlateIcon(),
+										FToolMenuExecuteAction::CreateLambda([](const FToolMenuContext& Context)
+											{
+												UMirrorDataTable::FFindReplaceOptions Options;
+												Options.Flags |= UMirrorDataTable::FFindReplaceOptions::EFlags::DisableStaleRows;
+												Options.bShowNotification = true;
+												
+												ExecuteMenuAction_ReimportFromSkeleton(Context, Options);
+											}
+										)
+									);
+									
+									FToolMenuSection& OptionsSection = Menu->FindOrAddSection("Options", LOCTEXT("OptionsLabel", "Options"));
+									
+									OptionsSection.AddMenuEntry(
+										NAME_None,
+										LOCTEXT("MapToSelfLabel", "Default to self on find/replace failure"),
+										LOCTEXT("MapToSelfLabel", "Maps entry name to itself when the find/replace fails to find a mirror entry name."),
+										FSlateIcon(),
+										FUIAction(
+											FExecuteAction::CreateLambda([]()
+												{
+													if (UAnimationSettings* AnimationSettings = UAnimationSettings::Get())
+													{
+														AnimationSettings->bOnMirrorPairFindReplaceFailedMapToSelf = !AnimationSettings->bOnMirrorPairFindReplaceFailedMapToSelf;
+													}
+												}
+											),
+											FCanExecuteAction(),
+											FIsActionChecked::CreateLambda([]()
+												{
+													if (UAnimationSettings* AnimationSettings = UAnimationSettings::Get())
+													{
+														return AnimationSettings->bOnMirrorPairFindReplaceFailedMapToSelf;
+													}
+													
+													return false;
+												}
+											)
+										),
+										EUserInterfaceActionType::Check
+									);
+								}
+
+								return UToolMenus::Get()->GenerateWidget(MenuName, Context);
+							})
+						),
+						FText::GetEmpty(),
+						LOCTEXT("ReimportFromSkeletonOptions", "Choose a reimport action")
+					));
+					ReimportFromSkeletonOptions.StyleNameOverride = FName("CalloutToolbar");
+					ReimportFromSkeletonOptions.ToolBarData.bSimpleComboBox = true;;
+				
+					// Validate and inform user.
+					FToolMenuEntry& Validate = Section.AddEntry(FToolMenuEntry::InitToolBarButton(
+						"MDT_Validate",
+						FToolMenuExecuteAction::CreateStatic(&FMirrorDataTableEditorModule::ExecuteMenuAction_Validate),
+						LOCTEXT("ValidateLabel", "Validate"),
+						LOCTEXT("ValidateTooltip", "Validate mirror data table is properly set up."),
+						FSlateIcon(FAppStyle::Get().GetStyleSetName(), "Icons.Validate")
+					));
+					Validate.StyleNameOverride = FName("CalloutToolbar");
+					
+					// Validation options.
+					FToolMenuEntry& ValidateOptions = Section.AddEntry(FToolMenuEntry::InitComboButton(
+						"MDT_ValidateOptions",
+						FUIAction(),
+						FNewToolMenuChoice(FNewToolMenuWidget::CreateLambda([](const FToolMenuContext& Context)
+							{
+								static const FName MenuName = "AssetEditor.MirrorDataTableEditor.ValidateMenu";
+
+								if (!UToolMenus::Get()->IsMenuRegistered(MenuName))
+								{
+									UToolMenu* Menu = UToolMenus::Get()->RegisterMenu(MenuName);
+									FToolMenuSection& OptionsSection = Menu->FindOrAddSection("Options", LOCTEXT("ValidateOptionsLabel", "Options"));
+
+									OptionsSection.AddMenuEntry(
+										NAME_None,
+										LOCTEXT("ValidateOnOpenLabel", "Validate on Open"),
+										LOCTEXT("ValidateOnOpenTooltip", "Automatically validate the mirror data table when it is opened in the editor. Only shows a notification if issues are found."),
+										FSlateIcon(),
+										FUIAction(
+											FExecuteAction::CreateLambda([]()
+												{
+													if (UAnimationSettings* AnimationSettings = UAnimationSettings::Get())
+													{
+														AnimationSettings->bValidateMirrorDataTableOnOpen = !AnimationSettings->bValidateMirrorDataTableOnOpen;
+													}
+												}
+											),
+											FCanExecuteAction(),
+											FIsActionChecked::CreateLambda([]()
+												{
+													if (const UAnimationSettings* AnimationSettings = UAnimationSettings::Get())
+													{
+														return AnimationSettings->bValidateMirrorDataTableOnOpen;
+													}
+
+													return false;
+												}
+											)
+										),
+										EUserInterfaceActionType::Check
+									);
+								}
+
+								return UToolMenus::Get()->GenerateWidget(MenuName, Context);
+							})
+						),
+						FText::GetEmpty(),
+						LOCTEXT("ValidateOptionsTooltip", "Validation options")
+					));
+					ValidateOptions.StyleNameOverride = FName("CalloutToolbar");
+					ValidateOptions.ToolBarData.bSimpleComboBox = true;
+				
+					// Only update existing entries (no stale entries).
+					Section.AddEntry(FToolMenuEntry::InitToolBarButton(
+						"MDT_AutoPair",
+						FToolMenuExecuteAction::CreateStatic(&FMirrorDataTableEditorModule::ExecuteMenuAction_AutoPair),
+						LOCTEXT("AutoPairLabel", "Auto Pair"),
+						LOCTEXT("AutoPairTooltip", "Auto Pair using Find Replace Expressions"),
+						FSlateIcon(FAppStyle::Get().GetStyleSetName(), "AutomationTools.MenuIcon")
+					));
+
+					// Clear all rows from the mirror data table.
+					Section.AddEntry(FToolMenuEntry::InitToolBarButton(
+						"MDT_ClearAll",
+						FToolMenuExecuteAction::CreateStatic(&FMirrorDataTableEditorModule::ExecuteMenuAction_ClearSheet),
+						LOCTEXT("ClearAllLabel", "Clear Sheet"),
+						LOCTEXT("ClearAllTooltip", "Remove all rows from the mirror data table."),
+						FSlateIcon(FAppStyle::Get().GetStyleSetName(), "Icons.Delete")
+					));
+					
+					// @todo: Preview changes option. Useful for big skeletons like meta human.
+				})
+			));
+	}
+}
+
+void FMirrorDataTableEditorModule::BindSkeletonCallbacks(UMirrorDataTable* MirrorTable)
+{
+	if (!MirrorTable || !MirrorTable->Skeleton)
+	{
+		return;
+	}
+
+	// Unbind any existing callbacks first (e.g. skeleton changed, we're rebinding to a new one).
+	UnbindSkeletonCallbacks(MirrorTable);
+
+	ISkeletonEditorModule& SkeletonEditorModule = FModuleManager::LoadModuleChecked<ISkeletonEditorModule>("SkeletonEditor");
+	TSharedPtr<IEditableSkeleton> EditableSkeleton = SkeletonEditorModule.CreateEditableSkeleton(MirrorTable->Skeleton);
+
+	FSkeletonCallbackHandles Handles;
+	Handles.EditableSkeleton = EditableSkeleton;
+	Handles.BoundSkeleton = MirrorTable->Skeleton;
+
+	// Listen for bone hierarchy changes. Only sync if the data is actually stale to avoid redundant work.
+	MirrorTable->Skeleton->RegisterOnSkeletonHierarchyChanged(USkeleton::FOnSkeletonHierarchyChanged::CreateWeakLambda(MirrorTable, [MirrorTable]()
+		{
+			if (!MirrorTable->Skeleton)
+			{
+				return;
+			}
+
+			if (MirrorTable->GetSkeletonSyncStatus() == UMirrorDataTable::ESyncStatus::UpToDate)
+			{
+				return;
+			}
+
+			UE_LOGF(LogAnimation, Log, "MirrorDataTable '%ls': Skeleton hierarchy changed. Updating entries.", *MirrorTable->GetName());
+			MirrorTable->UpdateFromFindReplaceExpressions(UMirrorDataTable::FFindReplaceOptions::Sync());
+		}
+	));
+
+	// Listen for notifies and/or sync marker changes.
+	if (EditableSkeleton && EditableSkeleton->IsSkeletonValid())
+	{
+		EditableSkeleton->RegisterOnNotifiesChanged(FSimpleMulticastDelegate::FDelegate::CreateWeakLambda(MirrorTable, [MirrorTable]()
+			{
+				UE_LOGF(LogAnimation, Log, "MirrorDataTable '%ls': Skeleton notifies and/or sync markers changed. Updating entries.", *MirrorTable->GetName());
+				MirrorTable->UpdateFromFindReplaceExpressions(UMirrorDataTable::FFindReplaceOptions::Sync());
+			}
+		));
+	}
+
+	// Listen for curve name changes.
+	if (UAnimCurveMetaData* AnimCurveMetaData = MirrorTable->Skeleton->GetAssetUserData<UAnimCurveMetaData>())
+	{
+		Handles.CurveMetaDataChangedHandle = AnimCurveMetaData->RegisterOnCurveMetaDataChanged(FSimpleMulticastDelegate::FDelegate::CreateWeakLambda(MirrorTable, [MirrorTable]()
+			{
+				UE_LOGF(LogAnimation, Log, "MirrorDataTable '%ls': Skeleton curve names changed. Updating entries.", *MirrorTable->GetName());
+				MirrorTable->UpdateFromFindReplaceExpressions(UMirrorDataTable::FFindReplaceOptions::Sync());
+			}
+		));
+	}
+
+	ActiveSkeletonCallbacks.Add(MirrorTable, MoveTemp(Handles));
+}
+
+void FMirrorDataTableEditorModule::UnbindSkeletonCallbacks(UMirrorDataTable* MirrorTable)
+{
+	FSkeletonCallbackHandles* Handles = ActiveSkeletonCallbacks.Find(MirrorTable);
+	if (!Handles)
+	{
+		return;
+	}
+
+	// Use BoundSkeleton, not the mirror table's Skeleton, since the Skeleton property may have already changed by the time we unbind.
+	if (USkeleton* BoundSkeleton = Handles->BoundSkeleton.Get())
+	{
+		BoundSkeleton->UnregisterOnSkeletonHierarchyChanged(MirrorTable);
+
+		if (Handles->EditableSkeleton)
+		{
+			Handles->EditableSkeleton->UnregisterOnNotifiesChanged(MirrorTable);
+		}
+
+		if (Handles->CurveMetaDataChangedHandle.IsValid())
+		{
+			if (UAnimCurveMetaData* AnimCurveMetaData = BoundSkeleton->GetAssetUserData<UAnimCurveMetaData>())
+			{
+				AnimCurveMetaData->UnregisterOnCurveMetaDataChanged(Handles->CurveMetaDataChangedHandle);
+			}
+		}
+	}
+
+	ActiveSkeletonCallbacks.Remove(MirrorTable);
+}
+
+void FMirrorDataTableEditorModule::OnAssetOpenedInEditor(UObject* Asset, IAssetEditorInstance* EditorInstance)
+{
+	UMirrorDataTable* MirrorTable = Cast<UMirrorDataTable>(Asset);
+	if (!MirrorTable)
+	{
+		return;
+	}
+
+	BindSkeletonCallbacks(MirrorTable);
+
+	if (const UAnimationSettings* AnimationSettings = UAnimationSettings::Get())
+	{
+		if (AnimationSettings->bValidateMirrorDataTableOnOpen)
+		{
+			// Only notify if issues are found — clean pass is silent to avoid noise.
+			ValidateMirrorDataTable(MirrorTable, false);
+		}
+	}
+}
+
+void FMirrorDataTableEditorModule::OnPropertyChanged(UObject* Object, FPropertyChangedEvent& PropertyChangedEvent)
+{
+	UMirrorDataTable* MirrorTable = Cast<UMirrorDataTable>(Object);
+	if (!MirrorTable)
+	{
+		return;
+	}
+
+	const FName MemberPropertyName = PropertyChangedEvent.GetMemberPropertyName();
+
+	// Skeleton property changed - rebind to the new skeleton.
+	if (MemberPropertyName == GET_MEMBER_NAME_CHECKED(UMirrorDataTable, Skeleton))
+	{
+		BindSkeletonCallbacks(MirrorTable);
+	}
+	// Bone scope settings changed - mark table as stale.
+	else if (MemberPropertyName == GET_MEMBER_NAME_CHECKED(UMirrorDataTable, BoneScope) ||
+	         MemberPropertyName == GET_MEMBER_NAME_CHECKED(UMirrorDataTable, BoneScopeNameList))
+	{
+		MirrorTable->InvalidateCachedSkeletonData();
+	}
+}
+
+void FMirrorDataTableEditorModule::ExecuteMenuAction_ReimportFromSkeleton(const FToolMenuContext& Context, const UMirrorDataTable::FFindReplaceOptions& Options)
+{
+	if (UMirrorDataTable* MirrorTable = GetMirrorDataTableFromContext(Context))
+	{
+		FScopedTransaction Transaction(LOCTEXT("UndoAction_ReimportFromSkeleton", "Re-import missing data from the associated Skeleton."));
+		
+		MirrorTable->UpdateFromFindReplaceExpressions(Options);
+		
+		MirrorTable->MarkPackageDirty();
+	}
+}
+
+void FMirrorDataTableEditorModule::ExecuteMenuAction_Validate(const FToolMenuContext& Context)
+{
+	if (UMirrorDataTable* MirrorTable = GetMirrorDataTableFromContext(Context))
+	{
+		ValidateMirrorDataTable(MirrorTable, true);
+	}
+}
+
+void FMirrorDataTableEditorModule::ValidateMirrorDataTable(UMirrorDataTable* MirrorTable, bool bNotifyOnCleanPass)
+{
+	const FMirrorDataTableValidator::FResult Result = FMirrorDataTableValidator::Validate(*MirrorTable);
+
+	uint32 NumErrors = 0;
+	uint32 NumWarnings = 0;
+
+	for (const FMirrorDataTableValidator::FIssue& Issue : Result.Issues)
+	{
+		const TCHAR* LogFormatString = !Issue.RowName.IsNone() ? TEXT("MirrorDataTable '{0}': Row '{1}': {2}") : TEXT("MirrorDataTable '{0}': {2}");
+		const FString FormattedString = FString::Format(LogFormatString, { MirrorTable->GetName(), Issue.RowName.ToString(), Issue.Message.ToString() });
+
+		switch (Issue.Severity)
+		{
+		case FMirrorDataTableValidator::FIssue::ESeverity::Info:
+			UE_LOGF(LogAnimation, Log, "%ls", *FormattedString);
+			break;
+		case FMirrorDataTableValidator::FIssue::ESeverity::Warning:
+			UE_LOGF(LogAnimation, Warning, "%ls", *FormattedString);
+			++NumWarnings;
+			break;
+		case FMirrorDataTableValidator::FIssue::ESeverity::Error:
+			UE_LOGF(LogAnimation, Error, "%ls", *FormattedString);
+			++NumErrors;
+			break;
+		}
+	}
+
+	const bool bHasIssues = NumErrors > 0 || NumWarnings > 0;
+	if (!bHasIssues && !bNotifyOnCleanPass)
+	{
+		return;
+	}
+
+	const FText CleanPassedSummary = LOCTEXT("PassedSummary_Validate", "Validation passed.");
+	const FText PassWithWarningsSummary = FText::Format(LOCTEXT("PassedWithWarningsSummary_Validate", "Validation passed with {0} warnings. See log for details."), NumWarnings);
+	const FText FinalPassedSummary = NumWarnings == 0 ? CleanPassedSummary : PassWithWarningsSummary;
+	const FText FailedSummary = FText::Format(LOCTEXT("FailedSummary_Validate", "Validation failed. Found {0} errors and {1} warnings. See log for details."), NumErrors, NumWarnings);
+	const FText Summary = NumErrors <= 0 ? FinalPassedSummary : FailedSummary;
+
+	FNotificationInfo Info = FNotificationInfo(Summary);
+	Info.bFireAndForget = true;
+	Info.ExpireDuration = 4.0f;
+	Info.bUseSuccessFailIcons = true;
+
+	if (TSharedPtr<SNotificationItem> Notification = FSlateNotificationManager::Get().AddNotification(Info))
+	{
+		Notification->SetCompletionState(NumErrors > 0 ? SNotificationItem::CS_Fail : SNotificationItem::CS_Success);
+	}
+}
+
+void FMirrorDataTableEditorModule::ExecuteMenuAction_AutoPair(const FToolMenuContext& Context)
+{
+	if (UMirrorDataTable* MirrorTable = GetMirrorDataTableFromContext(Context))
+	{
+		FScopedTransaction Transaction(LOCTEXT("UndoAction_AutoPairMirrorEntries", "Automatically find mirror pairs for all entries in table."));
+		
+		UMirrorDataTable::FFindReplaceOptions Options { UMirrorDataTable::FFindReplaceOptions::EFlags::UpdateExistingRows, true};
+		MirrorTable->UpdateFromFindReplaceExpressions(Options);
+		
+		MirrorTable->MarkPackageDirty();
+	}
+}
+
+void FMirrorDataTableEditorModule::ExecuteMenuAction_ClearSheet(const FToolMenuContext& Context)
+{
+	if (UMirrorDataTable* MirrorTable = GetMirrorDataTableFromContext(Context))
+	{
+		FScopedTransaction Transaction(LOCTEXT("UndoAction_ClearSheet", "Clear all rows from the mirror data table."));
+		
+		MirrorTable->Modify();
+		
+		// Clear sheet and let others know we've done so.
+		MirrorTable->EmptyTable();
+		MirrorTable->InvalidateCachedSkeletonData();
+		MirrorTable->OnDataTableChanged().Broadcast();
+		
+		// Let editor know we've changed the data and refresh view.
+		FDataTableEditorUtils::BroadcastPostChange(MirrorTable, FDataTableEditorUtils::EDataTableChangeInfo::RowList);
+		
+		MirrorTable->MarkPackageDirty();
+	}
+}
+
+FSlateIcon FMirrorDataTableEditorModule::GetStatusIcon(UMirrorDataTable* InMirrorDataTable)
+{
+	static const FName ReimportStatusBackground("Persona.AssetClass.Skeleton");
+	static const FName ReimportStatusUnknown("Blueprint.CompileStatus.Overlay.Unknown");
+	static const FName ReimportStatusError("Blueprint.CompileStatus.Overlay.Error");
+	static const FName ReimportStatusGood("Blueprint.CompileStatus.Overlay.Good");
+	static const FName ReimportStatusWarning("Blueprint.CompileStatus.Overlay.Warning");
+	
+	if (InMirrorDataTable)
+	{
+		if (!InMirrorDataTable->Skeleton)
+		{
+			return FSlateIcon(FAppStyle::GetAppStyleSetName(), ReimportStatusBackground, NAME_None, ReimportStatusError);
+		}
+	
+		switch (InMirrorDataTable->GetSkeletonSyncStatus())
+		{
+		case UMirrorDataTable::ESyncStatus::UpToDate:
+			return FSlateIcon(FAppStyle::GetAppStyleSetName(), ReimportStatusBackground, NAME_None, ReimportStatusGood);
+		case UMirrorDataTable::ESyncStatus::Stale:
+			return FSlateIcon(FAppStyle::GetAppStyleSetName(), ReimportStatusBackground, NAME_None, ReimportStatusWarning);
+		default:
+			return FSlateIcon(FAppStyle::GetAppStyleSetName(), ReimportStatusBackground, NAME_None, ReimportStatusUnknown);
+		}
+	}
+
+	return FSlateIcon(FAppStyle::GetAppStyleSetName(), ReimportStatusBackground, NAME_None, ReimportStatusUnknown);
+}
+
+#undef LOCTEXT_NAMESPACE
+    
+IMPLEMENT_MODULE(FMirrorDataTableEditorModule, MirrorDataTableEditor)
\ No newline at end of file
diff --git a/Engine/Source/Editor/MirrorDataTableEditor/Private/MirrorDataTableValidator.cpp b/Engine/Source/Editor/MirrorDataTableEditor/Private/MirrorDataTableValidator.cpp
new file mode 100644
--- /dev/null
+++ b/Engine/Source/Editor/MirrorDataTableEditor/Private/MirrorDataTableValidator.cpp
@@ -0,0 +1,350 @@
+﻿// Copyright Epic Games, Inc. All Rights Reserved.
+
+#include "MirrorDataTableValidator.h"
+
+#include "Animation/Skeleton.h"
+#include "Misc/ScopedSlowTask.h"
+
+#define LOCTEXT_NAMESPACE "MDT_Validator"
+
+constexpr int32 NumSkeletonBackedRowTypes = 4;
+constexpr int32 NumRowTypes = 5;
+
+static TStaticArray<FText, NumRowTypes> RowTypeLabels { 
+	LOCTEXT("BoneLabel", "Bone"),
+	LOCTEXT("AnimNotifyLabel", "Animation Notify"),
+	LOCTEXT("CurveLabel", "Curve"),
+	LOCTEXT("SyncMarkerLabel", "Sync Marker"),
+	LOCTEXT("CustomLabel", "Custom"),
+};
+
+FMirrorDataTableValidator::FResult FMirrorDataTableValidator::Validate(const UMirrorDataTable& MirrorDataTable)
+{
+	FResult OutResult;
+
+	// Keep track of progress.
+	constexpr int32 NumProgressSteps = 3;
+	FScopedSlowTask Progress(NumProgressSteps, LOCTEXT("ValidatorProgress", "Validating Mirror Data Table..."));
+	
+	// Show progress dialog.
+    constexpr bool bShowCancelButton = true;
+    Progress.MakeDialog(bShowCancelButton);
+	
+	// Warn if bone scope is set to explicit but no bones are selected.
+	if (MirrorDataTable.BoneScope == EMirrorTableBoneRefreshScope::ExplicitBoneList && MirrorDataTable.BoneScopeNameList.Num() == 0)
+	{
+		OutResult.AddIssue(
+			FIssue::ESeverity::Warning,
+			LOCTEXT("EmptyBoneScopeList", "Bone scope is set to 'Explicit Bone List' but no bones are selected. The full skeleton will be used as a fallback when refreshing.")
+		);
+	}
+
+	ValidateSkeletonBackedEntries(MirrorDataTable, OutResult, Progress);
+	ValidateDuplicateMappings(MirrorDataTable, OutResult, Progress);
+	ValidateReverseMappings(MirrorDataTable, OutResult, Progress);
+	
+	return OutResult;
+}
+
+void FMirrorDataTableValidator::ValidateSkeletonBackedEntries(const UMirrorDataTable& MirrorDataTable, FResult& OutResult, FScopedSlowTask& InOutProgress)
+{
+	InOutProgress.EnterProgressFrame(1.0f, LOCTEXT("ValidateSkeletonEntries", "Checking for missing or stale skeleton entries..."));
+	
+	if (!MirrorDataTable.Skeleton)
+	{
+		OutResult.AddIssue(
+			FIssue::ESeverity::Error,
+			LOCTEXT("MissingSkeleton", "Mirror Data Table has no associated skeleton.")
+		);
+		return;
+	}
+	
+	TSet<FName> SourceEntriesPerCategory[NumSkeletonBackedRowTypes];
+	TSet<FName> DisabledEntriesPerCategory[NumSkeletonBackedRowTypes];
+	
+	TArray<FName> SkeletonCurveNames;
+	MirrorDataTable.Skeleton->GetCurveMetaDataNames(SkeletonCurveNames);
+	
+	auto DoesSkeletonEntryExist = [&SkeletonCurveNames](const UMirrorDataTable& MirrorDataTable, EMirrorRowType::Type RowType, FName EntryName) -> bool
+	{
+		if (EntryName.IsNone() || !MirrorDataTable.Skeleton)
+		{
+			return false;
+		}
+			
+		switch (RowType)
+		{
+		case EMirrorRowType::Bone:
+			return MirrorDataTable.Skeleton->GetReferenceSkeleton().FindBoneIndex(EntryName) != INDEX_NONE;
+			
+		case EMirrorRowType::AnimationNotify:
+			return MirrorDataTable.Skeleton->AnimationNotifies.Contains(EntryName);
+			
+		case EMirrorRowType::Curve:
+			return SkeletonCurveNames.Contains(EntryName);
+			
+		case EMirrorRowType::SyncMarker:
+			return MirrorDataTable.Skeleton->GetExistingMarkerNames().Contains(EntryName);
+			
+		case EMirrorRowType::Custom:
+		default: 
+			return false;
+		}
+	};
+	
+	MirrorDataTable.ForeachRow<FMirrorTableRow>(TEXT("ValidateSkeletonEntries"), [&DoesSkeletonEntryExist, &MirrorDataTable, &OutResult, &SourceEntriesPerCategory, &DisabledEntriesPerCategory](const FName & RowName, const FMirrorTableRow& Row)
+	{
+		// Custom row entry does not rely on skeleton. Skip validation.
+		if (Row.MirrorEntryType == EMirrorRowType::Custom)
+		{
+			return;
+		}
+
+		// Determine severity.
+		FIssue::ESeverity IssueSeverity = Row.bEnabled ? FIssue::ESeverity::Error : FIssue::ESeverity::Warning;
+		FText DisabledSuffix;
+			
+		// Don't fail validation when disabled.
+		if (!Row.bEnabled)
+		{
+			DisabledSuffix = LOCTEXT("DisabledEntry", " This entry is disabled.");
+			
+			DisabledEntriesPerCategory[Row.MirrorEntryType].Add(Row.Name);
+		}
+			
+		// Check for entry being invalid.
+		if (Row.Name.IsNone())
+		{
+			OutResult.AddIssue(
+				FIssue::ESeverity::Error,
+				LOCTEXT("InvalidSourceEntry", "Entry not specified."),
+				RowName
+			);
+		}
+			
+		// Always track the source entry so the second pass doesn't falsely report it as missing.
+		SourceEntriesPerCategory[Row.MirrorEntryType].Add(Row.Name);
+
+		// Check for mirror entry being invalid.
+		if (Row.MirroredName.IsNone())
+		{
+			OutResult.AddIssue(
+				IssueSeverity,
+				FText::Format(LOCTEXT("InvalidMirroredEntry", "Mirrored entry not specified.{0}"), DisabledSuffix),
+				RowName
+			);
+		}
+			
+		// Stale entries (missing in skeleton, found in table)
+			
+		// Check for row entry name being stale
+		if (!Row.Name.IsNone() && !DoesSkeletonEntryExist(MirrorDataTable, Row.MirrorEntryType, Row.Name))
+		{
+			OutResult.AddIssue(
+				IssueSeverity,
+				FText::Format(LOCTEXT("MissingSourceEntry", "{1} '{0}' no longer exists in skeleton data.{2}"), FText::FromName(Row.Name), RowTypeLabels[Row.MirrorEntryType], DisabledSuffix),
+				RowName
+			);
+		}
+			
+		// Check for row mirror pair name being stale
+		if (!Row.MirroredName.IsNone() && !DoesSkeletonEntryExist(MirrorDataTable, Row.MirrorEntryType, Row.MirroredName))
+		{
+			OutResult.AddIssue(
+				IssueSeverity,
+				FText::Format(LOCTEXT("MissingMirroredEntry", "Mirrored {1} '{0}' no longer exists in skeleton data.{2}"), FText::FromName(Row.MirroredName), RowTypeLabels[Row.MirrorEntryType], DisabledSuffix),
+				RowName
+			);
+		}
+	});
+	
+	// Missing entries (missing in table, found in skeleton)
+	{
+		auto ProcessEntry = [&DisabledEntriesPerCategory, &OutResult, &SourceEntriesPerCategory](EMirrorRowType::Type RowType, FName EntryName)
+		{
+			if (!SourceEntriesPerCategory[RowType].Contains(EntryName))
+			{
+				const bool bIsEnabled = !DisabledEntriesPerCategory[RowType].Contains(EntryName);
+				FIssue::ESeverity IssueSeverity = bIsEnabled ? FIssue::ESeverity::Error : FIssue::ESeverity::Warning;
+				FText DisabledSuffix = !bIsEnabled ? LOCTEXT("DisabledEntry", " This entry is disabled.") : FText::GetEmpty();
+
+
+				OutResult.AddIssue(
+					IssueSeverity,
+					FText::Format(LOCTEXT("MissingEntry", "{0} '{1}' is missing from the table.{2}"), RowTypeLabels[RowType], FText::FromName(EntryName), DisabledSuffix)
+				);
+			}
+		};
+		
+		const FReferenceSkeleton& ReferenceSkeleton = MirrorDataTable.Skeleton->GetReferenceSkeleton();
+		if (MirrorDataTable.BoneScope == EMirrorTableBoneRefreshScope::ExplicitBoneList && MirrorDataTable.BoneScopeNameList.Num() > 0)
+		{
+			for (const FName& Name : MirrorDataTable.BoneScopeNameList)
+			{
+				if (ReferenceSkeleton.FindBoneIndex(Name) != INDEX_NONE)
+				{
+					ProcessEntry(EMirrorRowType::Bone, Name);
+				}
+			}
+		}
+		else
+		{
+			for (int32 BoneIndex = 0; BoneIndex < ReferenceSkeleton.GetNum(); ++BoneIndex)
+			{
+				const FName Name = ReferenceSkeleton.GetBoneName(BoneIndex);
+				ProcessEntry(EMirrorRowType::Bone, Name);
+			}
+		}
+		
+		for (const FName& Name : MirrorDataTable.Skeleton->AnimationNotifies)
+		{
+			ProcessEntry(EMirrorRowType::AnimationNotify, Name);
+		}
+		
+		for (const FName& Name : MirrorDataTable.Skeleton->GetExistingMarkerNames())
+		{
+			ProcessEntry(EMirrorRowType::SyncMarker, Name);
+		}
+		
+		for (const FName& Name : SkeletonCurveNames)
+		{
+			ProcessEntry(EMirrorRowType::Curve, Name);
+		}
+	}
+}
+
+void FMirrorDataTableValidator::ValidateDuplicateMappings(const UMirrorDataTable& MirrorDataTable, FResult& OutResult, FScopedSlowTask& InOutProgress)
+{
+	InOutProgress.EnterProgressFrame(1.0f, LOCTEXT("ValidateDuplicateMappings", "Checking for duplicate entries..."));
+	
+	TMap<FName, TArray<FName>> CategoryMappings[NumRowTypes]; // [row type] entry name -> array of row names
+	
+	MirrorDataTable.ForeachRow<FMirrorTableRow>(TEXT("ValidateDuplicateMappings"), [&MirrorDataTable, &OutResult, &CategoryMappings](const FName & RowName, const FMirrorTableRow& Row)
+	{
+		TArray<FName>& RowNames = CategoryMappings[Row.MirrorEntryType].FindOrAdd(Row.Name);
+		RowNames.Add(RowName);
+	});
+	
+	// For each row type
+	for (int CategoryIndex = 0; CategoryIndex < NumRowTypes; CategoryIndex++)
+	{
+		// Check if any entry names are found more than once
+		for (const TTuple<FName, TArray<FName>>& EntryDuplicates : CategoryMappings[CategoryIndex])
+		{
+			const TArray<FName> & RowNames = EntryDuplicates.Value;
+			if (RowNames.Num() <= 1)
+			{
+				continue;
+			}
+			
+			// Log an issue for each duplicate found
+			const FName EntryName = EntryDuplicates.Key;
+			for (const FName& RowName : RowNames)
+			{
+				OutResult.AddIssue(
+					FIssue::ESeverity::Error,
+					FText::Format(LOCTEXT("DuplicateMappings", "Found duplicate mappings for entry '{0}'."), FText::FromName(EntryName)),
+					RowName
+				);
+			}
+		}
+	}
+}
+
+void FMirrorDataTableValidator::ValidateReverseMappings(const UMirrorDataTable& MirrorDataTable, FResult& OutResult, FScopedSlowTask& InOutProgress)
+{
+	InOutProgress.EnterProgressFrame(1.0f, LOCTEXT("ValidateReverseMappings", "Checking that mirroring is invertible..."));
+	
+	TMap<FName, FName> EntryNameToRowName[NumRowTypes];
+	
+	MirrorDataTable.ForeachRow<FMirrorTableRow>(TEXT("ValidateReverseMappings"), [&MirrorDataTable, &OutResult, &EntryNameToRowName](const FName & RowName, const FMirrorTableRow& Row)
+	{
+		EntryNameToRowName[Row.MirrorEntryType].FindOrAdd(Row.Name) = RowName;
+	});
+	
+	const bool bUsingExplicitBoneList = MirrorDataTable.BoneScope == EMirrorTableBoneRefreshScope::ExplicitBoneList && MirrorDataTable.BoneScopeNameList.Num() > 0;
+
+	MirrorDataTable.ForeachRow<FMirrorTableRow>(TEXT("ValidateReverseMappings"), [&MirrorDataTable, &OutResult, &EntryNameToRowName, bUsingExplicitBoneList](const FName & RowName, const FMirrorTableRow& Row)
+	{
+		FName MirroredRowName = EntryNameToRowName[Row.MirrorEntryType].FindRef(Row.MirroredName);
+		const FMirrorTableRow* MirroredRow = MirrorDataTable.FindRow<FMirrorTableRow>(MirroredRowName, TEXT("ValidateReverseMappings"), false);
+
+		FIssue::ESeverity IssueSeverity = !Row.bEnabled ? FIssue::ESeverity::Warning : FIssue::ESeverity::Error;
+		FText DisabledSuffix = !Row.bEnabled ? LOCTEXT("DisabledEntry", " This entry is disabled.") : FText::GetEmpty();
+
+		// When using an explicit bone list, skip reverse mapping validation for bone rows where either side is out of scope.
+		// Bones that mirror to nothing still fall through to the regular checks below.
+		if (bUsingExplicitBoneList && Row.MirrorEntryType == EMirrorRowType::Bone && !Row.MirroredName.IsNone())
+		{
+			const bool bSourceInScope = MirrorDataTable.BoneScopeNameList.Contains(Row.Name);
+			const bool bMirrorInScope = MirrorDataTable.BoneScopeNameList.Contains(Row.MirroredName);
+
+			if (!bSourceInScope || !bMirrorInScope)
+			{
+				if (bSourceInScope && !bMirrorInScope)
+				{
+					OutResult.AddIssue(
+						FIssue::ESeverity::Info,
+						FText::Format(
+							LOCTEXT("MirrorTargetOutOfScope", "Bone '{0}' is in the bone scope list but its mirror target '{1}' is not. The reverse entry may not be updated on the next refresh.{2}"),
+							FText::FromName(Row.Name),
+							FText::FromName(Row.MirroredName),
+							DisabledSuffix
+						),
+						RowName
+					);
+				}
+				return;
+			}
+		}
+
+		if (!MirroredRow && !Row.MirroredName.IsNone())
+		{
+			OutResult.AddIssue(
+				IssueSeverity,
+				FText::Format(
+					LOCTEXT("MissingInvertibleMapping", "{0} '{1}' mirrors to '{2}', but '{2}' does not have a matching reverse entry.{3}"), 
+					RowTypeLabels[Row.MirrorEntryType],
+					FText::FromName(Row.Name),
+					FText::FromName(Row.MirroredName),
+					DisabledSuffix
+				),
+				RowName
+			);
+		}
+			
+		if (MirroredRow)
+		{
+			if (MirroredRow->MirroredName != Row.Name)
+			{
+				OutResult.AddIssue(
+					IssueSeverity,
+					FText::Format(
+						LOCTEXT("NonInvertibleMapping", "{0} '{1}' mirrors to '{2}', but '{2}' mirrors to '{3}' instead."), 
+						RowTypeLabels[Row.MirrorEntryType],
+						FText::FromName(Row.Name),
+						FText::FromName(Row.MirroredName),
+						FText::FromName(MirroredRow->MirroredName)
+					),
+					RowName
+				);
+			}
+			else if (Row.bEnabled && !MirroredRow->bEnabled)
+			{
+				OutResult.AddIssue(
+					FIssue::ESeverity::Error,
+					FText::Format(
+						LOCTEXT("DisabledReverseMapping", "{0} '{1}' mirrors to '{2}', but the matching reverse entry '{3}' is disabled."), 
+						RowTypeLabels[Row.MirrorEntryType],
+						FText::FromName(Row.Name),
+						FText::FromName(Row.MirroredName),
+						FText::FromName(MirroredRowName)
+					),
+					RowName
+				);
+			}
+		}
+	});
+}
+
+#undef LOCTEXT_NAMESPACE
\ No newline at end of file
diff --git a/Engine/Source/Editor/MirrorDataTableEditor/Public/MirrorDataTableEditorModule.h b/Engine/Source/Editor/MirrorDataTableEditor/Public/MirrorDataTableEditorModule.h
new file mode 100644
--- /dev/null
+++ b/Engine/Source/Editor/MirrorDataTableEditor/Public/MirrorDataTableEditorModule.h
@@ -0,0 +1,81 @@
+﻿// Copyright Epic Games, Inc. All Rights Reserved.
+
+#pragma once
+
+#include "CoreMinimal.h"
+#include "Animation/MirrorDataTable.h"
+#include "Modules/ModuleManager.h"
+#include "Textures/SlateIcon.h"
+
+class IAssetEditorInstance;
+class IEditableSkeleton;
+class USkeleton;
+class UMirrorDataTable;
+struct FToolMenuContext;
+
+class FMirrorDataTableEditorModule : public IModuleInterface
+{
+public:
+
+	/** Begin IModuleInterface */
+    virtual void StartupModule() override;
+    virtual void ShutdownModule() override;
+	/** End IModuleInterface */
+
+	/** Register tool menus specific when editing a mirror data tables */
+    void RegisterMenus();
+
+	/** Reimports and syncs the active mirror data table from its skeleton. */
+	static void ExecuteMenuAction_ReimportFromSkeleton(const FToolMenuContext& Context, const UMirrorDataTable::FFindReplaceOptions& Options);
+
+	/** Validates the active mirror data table. */
+	static void ExecuteMenuAction_Validate(const FToolMenuContext& Context);
+
+	/** Auto-generates mirror pairs for the active mirror data table. */
+	static void ExecuteMenuAction_AutoPair(const FToolMenuContext& Context);
+
+	/** Clears all rows from the active mirror data table. */
+	static void ExecuteMenuAction_ClearSheet(const FToolMenuContext& Context);
+
+	/** Returns the status icon for the given mirror data table. */
+    static FSlateIcon GetStatusIcon(UMirrorDataTable* InMirrorDataTable);
+
+private:
+
+	/** Validates the given mirror data table and logs all issues. Skips showing notification on validation passed if requested. */
+	static void ValidateMirrorDataTable(UMirrorDataTable* MirrorTable, bool bNotifyOnCleanPass);
+
+	/** Stores per-mirror-table skeleton callback handles for an actively-edited mirror data table. */
+	struct FSkeletonCallbackHandles
+	{
+		/** Editable skeleton interface used to register/unregister notifies callbacks. */
+		TSharedPtr<IEditableSkeleton> EditableSkeleton;
+
+		/** Handle for the curve metadata changed callback. */
+		FDelegateHandle CurveMetaDataChangedHandle;
+
+		/** The skeleton that was bound. Stored separately so unbinding is safe even if the MirrorTable's Skeleton property has already changed. */
+		TWeakObjectPtr<USkeleton> BoundSkeleton;
+	};
+
+	/** Binds skeleton change callbacks for the given MirrorTable. Unbinds any existing callbacks first if the MirrorTable was previously bound. */
+	void BindSkeletonCallbacks(UMirrorDataTable* MirrorTable);
+
+	/** Unbinds all skeleton change callbacks for the given MirrorTable and removes its entry from ActiveSkeletonCallbacks. */
+	void UnbindSkeletonCallbacks(UMirrorDataTable* MirrorTable);
+
+	/** Called when any asset is opened in the editor. Binds skeleton callbacks if the asset is a mirror data table. */
+	void OnAssetOpenedInEditor(UObject* Asset, IAssetEditorInstance* EditorInstance);
+
+	/** Called when a property changes on any UObject. Rebinds skeleton callbacks when the Skeleton property changes on a mirror data table. */
+	void OnPropertyChanged(UObject* Object, FPropertyChangedEvent& PropertyChangedEvent);
+
+	/** Per-MirrorTable skeleton callback state for all actively-edited mirror data tables. */
+	TMap<TWeakObjectPtr<UMirrorDataTable>, FSkeletonCallbackHandles> ActiveSkeletonCallbacks;
+
+	/** Handle for the UAssetEditorSubsystem::OnAssetOpenedInEditor delegate. */
+	FDelegateHandle OnAssetOpenedHandle;
+
+	/** Handle for the FCoreUObjectDelegates::OnObjectPropertyChanged delegate. */
+	FDelegateHandle OnObjectPropertyChangedHandle;
+};
diff --git a/Engine/Source/Editor/MirrorDataTableEditor/Public/MirrorDataTableValidator.h b/Engine/Source/Editor/MirrorDataTableEditor/Public/MirrorDataTableValidator.h
new file mode 100644
--- /dev/null
+++ b/Engine/Source/Editor/MirrorDataTableEditor/Public/MirrorDataTableValidator.h
@@ -0,0 +1,91 @@
+﻿// Copyright Epic Games, Inc. All Rights Reserved.
+
+#pragma once
+
+#include "CoreMinimal.h"
+#include "Animation/MirrorDataTable.h"
+
+struct FScopedSlowTask;
+
+class FMirrorDataTableValidator
+{
+	public:
+	
+	/** Used to record an issue when validating a Mirror Data Table. */
+	struct FIssue
+	{
+		enum class ESeverity : uint8
+		{
+			Info,
+			Warning,
+			Error,
+		};
+	
+		/** Determines the scale of the issue*/
+		ESeverity Severity = ESeverity::Info;
+		
+		/** Summary of what went wrong. */
+		FText Message;
+		
+		/** Row where the issue is present. */
+		FName RowName = NAME_None;
+	};
+	
+	/** Used to record the results from validating a Mirror Data Table. */
+	struct FResult
+	{
+		/** Record of everything that is not properly setup. */
+		TArray<FIssue> Issues;
+	
+		/** Do we have any issues at all? */
+		bool HasErrors() const
+		{
+			for (const FIssue& Issue : Issues)
+			{
+				if (Issue.Severity == FIssue::ESeverity::Error)
+				{
+					return true;
+				}
+			}
+		
+			return false;
+		}
+	
+		/** Do we have any warnings at all? */
+		bool HasWarnings() const
+		{
+			for (const FIssue& Issue : Issues)
+			{
+				if (Issue.Severity == FIssue::ESeverity::Warning)
+				{
+					return true;
+				}
+			}
+		
+			return false;
+		}
+	
+		/** Records an issue that occurred during validation. */
+		void AddIssue(FIssue::ESeverity Severity, const FText& Message, FName RowName = NAME_None)
+		{
+			FIssue& Issue = Issues.AddDefaulted_GetRef();
+			Issue.Severity = Severity;
+			Issue.Message = Message;
+			Issue.RowName = RowName;
+		}
+	};
+	
+	/** Looks for any issues with the provided data table and record them. */
+	static FResult Validate(const UMirrorDataTable& MirrorDataTable);
+	
+	private:
+	
+	/** Checks for missing or stale skeleton entries. */
+	static void ValidateSkeletonBackedEntries(const UMirrorDataTable& MirrorDataTable, FResult& OutResult, FScopedSlowTask& InOutProgress);
+	
+	/** Catches duplicate logical source entries of the same row type. */
+	static void ValidateDuplicateMappings(const UMirrorDataTable& MirrorDataTable, FResult& OutResult, FScopedSlowTask& InOutProgress);
+	
+	/** Checks for two entries mapping to the same mirror entry causing the operation to no be invertible. Ex. A -> B and C -> B. */
+	static void ValidateReverseMappings(const UMirrorDataTable& MirrorDataTable, FResult& OutResult, FScopedSlowTask& InOutProgress);
+};
diff --git a/Engine/Source/Editor/Persona/Persona.Build.cs b/Engine/Source/Editor/Persona/Persona.Build.cs
--- a/Engine/Source/Editor/Persona/Persona.Build.cs
+++ b/Engine/Source/Editor/Persona/Persona.Build.cs
@@ -77,6 +77,7 @@ public class Persona : ModuleRules
 				"NaniteUtilities",
 				"ToolMenus",
                 "CurveEditor",
+                "MirrorDataTableEditor",
 				"SequencerWidgets",
 				"TimeManagement",
                 "Sequencer",
diff --git a/Engine/Source/Editor/Persona/Private/PersonaModule.cpp b/Engine/Source/Editor/Persona/Private/PersonaModule.cpp
--- a/Engine/Source/Editor/Persona/Private/PersonaModule.cpp
+++ b/Engine/Source/Editor/Persona/Private/PersonaModule.cpp
@@ -122,6 +122,9 @@ void FPersonaModule::StartupModule()
 
 	//Call this to make sure AnimGraph module is setup
 	FModuleManager::Get().LoadModuleChecked(TEXT("AnimGraph"));
+	
+	// Call this to make sure MDTEditor module is setup
+	FModuleManager::Get().LoadModuleChecked(TEXT("MirrorDataTableEditor"));
 
 	// Make sure the advanced preview scene module is loaded 
 	FModuleManager::Get().LoadModuleChecked("AdvancedPreviewScene");
diff --git a/Engine/Source/Editor/UnrealEd/Private/Factories/MirrorDataTableFactory.cpp b/Engine/Source/Editor/UnrealEd/Private/Factories/MirrorDataTableFactory.cpp
--- a/Engine/Source/Editor/UnrealEd/Private/Factories/MirrorDataTableFactory.cpp
+++ b/Engine/Source/Editor/UnrealEd/Private/Factories/MirrorDataTableFactory.cpp
@@ -344,7 +344,8 @@ UObject* UMirrorDataTableFactory::FactoryCreateNew(UClass* Class, UObject* InPar
 			}
 			if (Skeleton)
 			{
-				DataTable->FindReplaceMirroredNames();
+				// Sync fully so the table is considered up to date on first open.
+				DataTable->UpdateFromFindReplaceExpressions(UMirrorDataTable::FFindReplaceOptions::Sync());
 			}
 		}
 	}
diff --git a/Engine/Source/Runtime/Engine/Classes/Animation/AnimationSettings.h b/Engine/Source/Runtime/Engine/Classes/Animation/AnimationSettings.h
--- a/Engine/Source/Runtime/Engine/Classes/Animation/AnimationSettings.h
+++ b/Engine/Source/Runtime/Engine/Classes/Animation/AnimationSettings.h
@@ -99,6 +99,14 @@ class UAnimationSettings : public UDeveloperSettings
 	UPROPERTY(config, EditAnywhere, Category = Mirroring)
 	TArray<FMirrorFindReplaceExpression> MirrorFindReplaceExpressions;
 
+	/** When a mirror pair is not found map to self instead of none.  */
+	UPROPERTY(config, EditAnywhere, Category = Mirroring)
+	bool bOnMirrorPairFindReplaceFailedMapToSelf = true;
+
+	/** When enabled, automatically validates the mirror data table when it is opened in the editor. Only shows a notification if issues are found. */
+	UPROPERTY(config, EditAnywhere, Category = Mirroring)
+	bool bValidateMirrorDataTableOnOpen = true;
+	
 	/** Project specific default frame-rate used when (re)initializing any animation based data */
 	UPROPERTY(config, EditAnywhere, Category = AnimationData)
 	FFrameRate DefaultFrameRate;
diff --git a/Engine/Source/Runtime/Engine/Classes/Animation/MirrorDataTable.h b/Engine/Source/Runtime/Engine/Classes/Animation/MirrorDataTable.h
--- a/Engine/Source/Runtime/Engine/Classes/Animation/MirrorDataTable.h
+++ b/Engine/Source/Runtime/Engine/Classes/Animation/MirrorDataTable.h
@@ -5,6 +5,7 @@
 #include "Engine/DataTable.h"
 #include "BoneContainer.h"
 #include "CustomBoneIndexArray.h"
+#include "Containers/StaticArray.h"
 #include "MirrorDataTable.generated.h"
 
 /** Type referenced by a row in the mirror data table */
@@ -37,6 +38,20 @@ namespace EMirrorFindReplaceMethod
     };
 }
 
+#if WITH_EDITORONLY_DATA
+
+/** Controls which bones are updated when refreshing the mirror data table from the skeleton. */
+UENUM()
+enum class EMirrorTableBoneRefreshScope : uint8
+{
+	/** Refreshes all bones found in the associated skeleton. */
+	FullSkeleton,
+
+	/** Only refresh the bones explicitly listed in the Specific Bone List. Useful for partial mirrors or managing a large skeleton. */
+	ExplicitBoneList,
+};
+
+#endif
 
 /**  Base Mirror Table containing all data required by the animation mirroring system. */
 USTRUCT()
@@ -44,19 +59,27 @@ struct FMirrorTableRow : public FTableRowBase
 {
 	GENERATED_BODY()
 
-	UPROPERTY(EditAnywhere, BlueprintReadOnly, Category = Mirroring)
-	FName Name;
+	UPROPERTY(EditAnywhere, BlueprintReadOnly, Category = Mirroring, meta=(EditCondition="bEnabled"))
+	FName Name = NAME_None;
 
-	UPROPERTY(EditAnywhere, BlueprintReadOnly, Category = Mirroring)
-	FName MirroredName;
+	UPROPERTY(EditAnywhere, BlueprintReadOnly, Category = Mirroring, meta=(EditCondition="bEnabled"))
+	FName MirroredName = NAME_None;
+
+	UPROPERTY(EditAnywhere, BlueprintReadOnly, Category = Mirroring, meta=(EditCondition="bEnabled"))
+	TEnumAsByte<EMirrorRowType::Type> MirrorEntryType = EMirrorRowType::Bone;
 
 	UPROPERTY(EditAnywhere, BlueprintReadOnly, Category = Mirroring)
-	TEnumAsByte<EMirrorRowType::Type> MirrorEntryType;
+	bool bEnabled = true;
+	
+	FMirrorTableRow() = default;
 
-	FMirrorTableRow()
-		: Name(NAME_None)
-		, MirroredName(NAME_None)
-		, MirrorEntryType(EMirrorRowType::Bone) {}
+	FMirrorTableRow(const FName& InName, const FName& InMirroredName, const TEnumAsByte<EMirrorRowType::Type>& InMirrorEntryType, bool bInEnabled)
+		: Name(InName)
+		, MirroredName(InMirroredName)
+		, MirrorEntryType(InMirrorEntryType)
+		, bEnabled(bInEnabled)
+	{
+	}
 
 	ENGINE_API FMirrorTableRow(const FMirrorTableRow& Other);
 	ENGINE_API FMirrorTableRow& operator=(FMirrorTableRow const& Other);
@@ -162,12 +185,64 @@ public:
 		TCustomBoneIndexArray<FCompactPoseBoneIndex, FCompactPoseBoneIndex>& OutCompactPoseMirrorBones,
 		TCustomBoneIndexArray<FQuat, FCompactPoseBoneIndex>& OutComponentSpaceRefRotations) const;
 
-#if WITH_EDITOR  
+#if WITH_EDITOR
+
+	/** Describes the sync relationship between the mirror data table and its associated skeleton. */
+	enum class ESyncStatus : uint8
+	{
+		Unknown,   // Never been synced
+		UpToDate,  // GUIDs match the current skeleton
+		Stale,     // GUIDs no longer match the current skeleton
+	};
+
+	/** Options for updating the table from find/replace expressions. */
+	struct FFindReplaceOptions
+	{
+		/** Controls which mutations are allowed during the find/replace refresh. */
+		enum class EFlags : uint8
+		{
+			None = 0,						// Do not perform row mutations.
+			UpdateExistingRows = 1 << 0,	// Only update rows that already exist in the mirror data table.
+			AddMissingRows = 1 << 1,		// Add rows for skeleton entries that do not yet exist in the mirror data table.
+			DisableStaleRows = 1 << 2,		// Disable existing skeleton-derived rows that were not found in current skeleton scan.
+			DeleteStaleRows = 1 << 3,		// Deletes existing skeleton-derived rows that were not found in current skeleton scan.
+		};
+		
+		/** Combination of operations to perform when applying find/replace expressions. Defaults to legacy behavior. */
+		EFlags Flags = EFlags::AddMissingRows;
+		
+		/** Whether to show a completion notification once the operation is finished. */
+		bool bShowNotification = false;
+		
+		/** Preset: Only add missing rows. */
+		ENGINE_API static FFindReplaceOptions AddMissingOnly();
+
+		/** Preset: Updates existing rows and disables stale rows. */
+		ENGINE_API static FFindReplaceOptions UpdateExisting();
+
+		/** Preset: Update existing rows, adds missing rows, and disables stale rows. */
+		ENGINE_API static FFindReplaceOptions Sync();
+		
+		/** Sets whether to show a completion notification. */
+		ENGINE_API FFindReplaceOptions& WithNotification(bool bInShowNotification = true);
+	};
+	
 	/**
 	 * Populates the table by running the MirrorFindReplaceExpressions on bone names in the Skeleton.  If the mirrored name is also found 
 	 * on the Skeleton it is added to the table.
 	 */
+	UE_DEPRECATED(5.8, "Please use UpdateFromFindReplaceExpressions() instead.")
 	ENGINE_API void FindReplaceMirroredNames();
+	
+	/** Updates the table from the current find/replace expressions and skeleton data. */
+	ENGINE_API void UpdateFromFindReplaceExpressions(const FFindReplaceOptions& Options);
+	
+	/** Returns the sync status of the table relative to its associated skeleton. */
+	ENGINE_API ESyncStatus GetSkeletonSyncStatus() const;
+
+	/** Marks the cached skeleton data as stale. The next call to GetSkeletonSyncStatus() will return Stale until the table is refreshed from the skeleton. */
+	ENGINE_API void InvalidateCachedSkeletonData();
+
 #endif // WITH_EDITOR
 
 	/**
@@ -201,12 +276,15 @@ public:
 
 public:
 
+	/** Used to compute the mirror pair of the data table entries. */
 	UPROPERTY(EditAnywhere, Category = CreateTable)
 	TArray<FMirrorFindReplaceExpression> MirrorFindReplaceExpressions;
 
+	/** Axis to mirror across (X/Y/Z). Use to flip transforms/poses. */
 	UPROPERTY(EditAnywhere, Category = Mirroring)
 	TEnumAsByte<EAxis::Type> MirrorAxis;
 
+	/** Used to determine if the root motion attribute should be mirrored when using this mirror data table. */
 	UPROPERTY(EditAnywhere, Category = Mirroring)
 	bool  bMirrorRootMotion = true;
 	
@@ -215,14 +293,11 @@ public:
 
 	// Index of the mirror bone for a given bone index in the reference skeleton, or INDEX_NONE if the bone is not mirrored
 	TCustomBoneIndexArray<FSkeletonPoseBoneIndex, FSkeletonPoseBoneIndex> BoneToMirrorBoneIndex;
-	
-	UE_DEPRECATED(5.3, "UID-based mirroring has been remove, please use CurveToMirrorCurveMap.")
-	TArray<SmartName::UID_Type> CurveMirrorSourceUIDArray;
-	
-	UE_DEPRECATED(5.3, "UID-based mirroring has been remove, please use CurveToMirrorCurveMap.")
-	TArray<SmartName::UID_Type> CurveMirrorTargetUIDArray;
 
-	// Map from animation curve to mirrored animation curve
+	// Set of all the disabled mirror entries/pairs
+	TSet<FName> DisabledEntries;
+	
+	// Map from animation curve to mirrored animation curve	
 	TMap<FName, FName> CurveToMirrorCurveMap;
 	
 	// Map from animation notify to mirrored animation notify
@@ -231,9 +306,64 @@ public:
 	// Map from sync marker to mirrored sync marker 
 	TMap<FName, FName> SyncToMirrorSyncMap;
 
+#if WITH_EDITORONLY_DATA
+	
+	/** Determines which bones are updated when refreshing the table from the associated skeleton. Set to 'Explicit Bone List' to target only specific bones. */
+	UPROPERTY(EditAnywhere, Category = Import)
+	EMirrorTableBoneRefreshScope BoneScope = EMirrorTableBoneRefreshScope::FullSkeleton;
+
+	/** The set of bones to refresh when Bone Scope is set to 'Explicit Bone List'. Managed via the bone picker in the Import panel. */
+	UPROPERTY()
+	TSet<FName> BoneScopeNameList;
+	
+#endif
+	
 protected: 
 
 	// Fill BoneToMirrorBoneIndex, CurveMirrorSourceUIDArray, CurveMirrorTargetUIDArray and NotifyToMirrorNotifyIndex based on the Skeleton and Table Contents
 	ENGINE_API void FillMirrorArrays();
+
+private:
+	
+	/** Contains entry data for a given row type (bone/notify/sync marker/curve/custom). */
+	struct FCategoryState
+	{
+		/** Existing that already exist in the table for this category. */
+		TSet<FName> ExistingEntryNames;
+		
+		/** Entries processed from the current skeleton scan for this category. */
+		TSet<FName> ProcessedEntryNames;
+		
+		/**
+		 * Maps logical entry names to row names in the table.
+		 * Row names must be unique across all categories, for example:
+		 *   left_1 -> left_1:Bone
+		 *   left_1 -> left_1:SyncMarker
+		 */
+		TMap<FName, FName> EntryNameToRowName;
+		
+		/** Entries with duplicate rows in the table for this category. */
+		TSet<FName> DuplicateEntryNames;
+	};
+	
+	/** Collects all the existing mirror entries and categorizes them by their mirror row type. */
+	void InitCategoryEntryStates(TStaticArray<FCategoryState, 5> & OutStates) const;
+	
+	/** Generates a unique row name for a mirror data table entry. */
+	FName MakeUniqueMirrorRowName(const FName EntryName, const EMirrorRowType::Type RowType) const;
+	
+#if WITH_EDITORONLY_DATA
+
+	UPROPERTY()
+	FGuid SkeletonHierarchyGuid;
+
+	UPROPERTY()
+	FGuid SkeletonVirtualBonesHierarchyGuid;
+
+#endif
+	
 };
 
+#if WITH_EDITOR
+ENUM_CLASS_FLAGS(UMirrorDataTable::FFindReplaceOptions::EFlags);
+#endif
\ No newline at end of file
diff --git a/Engine/Source/Runtime/Engine/Private/Animation/MirrorDataTable.cpp b/Engine/Source/Runtime/Engine/Private/Animation/MirrorDataTable.cpp
--- a/Engine/Source/Runtime/Engine/Private/Animation/MirrorDataTable.cpp
+++ b/Engine/Source/Runtime/Engine/Private/Animation/MirrorDataTable.cpp
@@ -9,6 +9,14 @@
 #include "Internationalization/Regex.h"
 #include "UObject/LinkerLoad.h"
 
+#if WITH_EDITOR
+#include "Misc/ScopedSlowTask.h"
+#include "DataTableEditorUtils.h"
+#include "Framework/Notifications/NotificationManager.h"
+#include "Widgets/Notifications/SNotificationList.h"
+#include "Modules/ModuleManager.h"
+#endif
+
 #include UE_INLINE_GENERATED_CPP_BY_NAME(MirrorDataTable)
 
 #define LOCTEXT_NAMESPACE "MirrorDataTables"
@@ -29,17 +37,18 @@ FMirrorTableRow& FMirrorTableRow::operator=(FMirrorTableRow const& Other)
 	Name = Other.Name;
 	MirroredName = Other.MirroredName;
 	MirrorEntryType = Other.MirrorEntryType;
+	bEnabled = Other.bEnabled;
 	return *this;
 }
 
 bool FMirrorTableRow::operator==(FMirrorTableRow const& Other) const
 {
-	return (Name == Other.Name && MirroredName == Other.MirroredName && MirrorEntryType == Other.MirrorEntryType);
+	return (Name == Other.Name && MirroredName == Other.MirroredName && MirrorEntryType == Other.MirrorEntryType && bEnabled == Other.bEnabled);
 }
 
 bool FMirrorTableRow::operator!=(FMirrorTableRow const& Other) const
 {
-	return (Name != Other.Name || MirroredName != Other.MirroredName || MirrorEntryType != Other.MirrorEntryType);
+	return (Name != Other.Name || MirroredName != Other.MirroredName || MirrorEntryType != Other.MirrorEntryType || bEnabled != Other.bEnabled);
 }
 
 bool FMirrorTableRow::operator<(FMirrorTableRow const& Other) const
@@ -61,24 +70,46 @@ bool FMirrorTableRow::operator<(FMirrorTableRow const& Other) const
 UMirrorDataTable::UMirrorDataTable(const FObjectInitializer& ObjectInitializer)
 	: Super(ObjectInitializer), MirrorAxis(EAxis::X)
 {
-#if WITH_EDITORONLY_DATA
+	
+#if WITH_EDITOR
 	OnDataTableChanged().AddUObject(this, &UMirrorDataTable::FillMirrorArrays);
 #endif 
+	
 }
 
 void UMirrorDataTable::PostLoad()
 {
 	Super::PostLoad();
 
-	FillMirrorArrays(); 
+	FillMirrorArrays();
+
+#if WITH_EDITOR
+	// Only warn for assets that have been synced before but are now stale.
+	if (GetSkeletonSyncStatus() == ESyncStatus::Stale)
+	{
+		UE_LOGF(LogAnimation, Warning, "MirrorDataTable '%ls' is out of date with its skeleton. Consider reimporting from skeleton.", *GetName());
+	}
+#endif
 }
 
 #if WITH_EDITOR
 
 void UMirrorDataTable::PostEditChangeProperty(FPropertyChangedEvent& PropertyChangedEvent)
 {
+	const FName MemberName = PropertyChangedEvent.GetMemberPropertyName();
+
+	if (MemberName == GET_MEMBER_NAME_CHECKED(UMirrorDataTable, Skeleton))
+	{
+		InvalidateCachedSkeletonData();
+		FillMirrorArrays();
+	}
+	else if (MemberName != GET_MEMBER_NAME_CHECKED(UMirrorDataTable, bMirrorRootMotion))
+	{
+		// No need to update FillMirrorArrays() when bMirrorRootMotion changes as this is a variable used only at runtime when mirroring attributes.
+		FillMirrorArrays();
+	}
+
 	Super::PostEditChangeProperty(PropertyChangedEvent);
-	FillMirrorArrays(); 
 }
 
 #endif // WITH_EDITOR
@@ -96,7 +127,7 @@ FName UMirrorDataTable::GetSettingsMirrorName(FName InName)
 
 FName UMirrorDataTable::GetMirrorName(FName InName, const TArray<FMirrorFindReplaceExpression>& MirrorFindReplaceExpressions)
 {
-	FName ReplacedName;
+	FName ReplacedName = NAME_None;
 	FString InNameString = InName.ToString();
 	for (const FMirrorFindReplaceExpression& regExStr : MirrorFindReplaceExpressions)
 	{
@@ -238,116 +269,404 @@ FName UMirrorDataTable::FindBestMirroredBone(
 
 #if WITH_EDITOR  
 
+UMirrorDataTable::FFindReplaceOptions UMirrorDataTable::FFindReplaceOptions::AddMissingOnly()
+{
+	return { EFlags::AddMissingRows, false };
+}
+
+UMirrorDataTable::FFindReplaceOptions UMirrorDataTable::FFindReplaceOptions::UpdateExisting()
+{
+	return { EFlags::UpdateExistingRows | EFlags::DisableStaleRows, false };
+}
+
+UMirrorDataTable::FFindReplaceOptions UMirrorDataTable::FFindReplaceOptions::Sync()
+{
+	return { EFlags::UpdateExistingRows | EFlags::AddMissingRows | EFlags::DisableStaleRows, false };
+}
+
+UMirrorDataTable::FFindReplaceOptions& UMirrorDataTable::FFindReplaceOptions::WithNotification(bool bInShowNotification)
+{
+	bShowNotification = bInShowNotification;
+	return *this;
+}
+
 void UMirrorDataTable::FindReplaceMirroredNames()
+{
+	UpdateFromFindReplaceExpressions(FFindReplaceOptions::AddMissingOnly());
+}
+
+void UMirrorDataTable::UpdateFromFindReplaceExpressions(const FFindReplaceOptions& Options)
 {
 	if (!Skeleton)
 	{
 		return; 
 	}
 
-	const FReferenceSkeleton& RefSkeleton = Skeleton->GetReferenceSkeleton();
-
-	TSet<FName> NamesByCategory[4];
-	ForeachRow<FMirrorTableRow>(TEXT("UMirrorDataTable::FindReplaceMirroredNames"), [&NamesByCategory] (const FName& Key, const FMirrorTableRow& Value) mutable
+	// @todo: LogAnimation default filter is Warning so any Log messages will not show in the UE unless `log Animation Log` command is ran. Not useful for TDiAs.
+	UE_LOGF(LogAnimation, Log, "Starting operation to update MirrorDataTable (%ls) via Find Replace expressions", *GetName());
+	
+	// Determine what are we doing.
+	const bool bShouldUpdateExisting = EnumHasAnyFlags(Options.Flags, FFindReplaceOptions::EFlags::UpdateExistingRows);
+	const bool bShouldAddMissing = EnumHasAnyFlags(Options.Flags, FFindReplaceOptions::EFlags::AddMissingRows);
+	const bool bShouldDisableStale = EnumHasAnyFlags(Options.Flags, FFindReplaceOptions::EFlags::DisableStaleRows);
+	
+	// Keep track of progress.
+	constexpr int32 NumProgressSteps = 7; // Gather entries: 1, Update/Add each entry type: 2-5, Disable missing entries: 6, Update runtime maps: 7. 
+	FScopedSlowTask Progress(NumProgressSteps, LOCTEXT("FindReplaceMirroredNamesOperation", "Finding and/or replacing mirror pairs..."));
+	
+	// Show progress dialog.
+	constexpr bool bShowCancelButton = true;
+	Progress.MakeDialog(bShowCancelButton);
+	
+	// Build map of existing entries (organized by their row type).
+	TStaticArray<FCategoryState, 5> CategoryStates;
+	InitCategoryEntryStates(CategoryStates);
+	
+	bool bChangedTable = false;	// Has any change was done to the table?
+	bool bModified = false;		// Has this table already been modified?
+	
+	// Operation stats.
+	int NumAddedRows = 0;
+	int NumUpdatedRows = 0;
+	int NumStaleRows = 0;
+	int NumDuplicateRows = 0;
+		
+	// Call any time we attempt to modify the inner data table struct.
+	auto TryToModify = [this, &bModified, &bChangedTable]()
+	{
+		// Ensure we only modify once to avoid storing extra data in transaction buffer.
+		if (!bModified)
 		{
-			NamesByCategory[Value.MirrorEntryType].Add(Value.Name);
+			Modify();
+			bModified = true;
 		}
-	);
-
-	bool bChangedTable = false;
-
-	auto AddMirrorRow = [this, &bChangedTable, &NamesByCategory](const FName& Name, const FName& MirroredName, EMirrorRowType::Type RowType)
+			
+		// Any time we try to modify it means we changed the internal data of the table.
+		bChangedTable = true;
+	};
+	
+	// Finds or adds a new mirror row (ensuring it has a unique name).
+	auto FindOrAddMirrorRow = [this, &TryToModify, &CategoryStates, &NumAddedRows](const FName& EntryName, const EMirrorRowType::Type RowType, const bool bAllowAdd) -> FMirrorTableRow*
 	{
-		// directly add rows to avoid using FDataTableEditorUtils, which is not appropriate at this point
-		// equivalent to FDataTableEditorUtils::AddRow(DataTable, BoneName);
-		static const FString CategoryName[] = { TEXT(":Bone"), TEXT(":Notify"), TEXT(":Curve"), TEXT(":SyncMarker") };
-		if (!NamesByCategory[RowType].Contains(Name))
+		FCategoryState& State = CategoryStates[RowType];
+			
+		// Get the actual row name in the data table.
+		if (const FName* ExistingRowName = State.EntryNameToRowName.Find(EntryName))
+		{
+			return FindRow<FMirrorTableRow>(*ExistingRowName, TEXT("UMirrorDataTable::FindOrAddMirrorRow"), false);
+		}
+		
+		// We couldn't find the row, but we aren't able to add a new one, exit.
+		if (!bAllowAdd || !RowStruct)
+		{
+			return nullptr;
+		}
+			
+		// Get unique row name across data table entries.
+		const FName NewRowName = MakeUniqueMirrorRowName(EntryName, RowType);
+			
+		// Add new row to data table.
 		{
-			FName RowName = Name;
-			FMirrorTableRow* ExistingRow = FindRow<FMirrorTableRow>(RowName, TEXT("UMirrorDataTable::FindReplaceMirroredNames"), false);
-			uint32 RenameAttempts = 0; 
-			while (ExistingRow)
+			TryToModify();
+			
+			// Allocate data to store information, using UScriptStruct to know its size
+			uint8* RowData = (uint8*)FMemory::Malloc(RowStruct->GetStructureSize());
+			RowStruct->InitializeStruct(RowData);
+			
+			// Directly add rows to avoid using FDataTableEditorUtils, which is not appropriate at this point
+			// equivalent to FDataTableEditorUtils::AddRow(DataTable, BoneName). This is so we can batch changes and only fire notifications at the end. 
+			AddRowInternal(NewRowName, RowData);
+			
+			++NumAddedRows;
+			
+			UE_LOGF(LogAnimation, Log, "MirrorDataTable '%ls': Added new row '%ls'", *GetName(), *NewRowName.ToString());
+		}
+			
+		// Keep the cached state accurate for the rest of this run.
+		State.ExistingEntryNames.Add(EntryName);
+		State.EntryNameToRowName.Add(EntryName, NewRowName);
+			
+		return FindRow<FMirrorTableRow>(NewRowName, TEXT("UMirrorDataTable::FindOrAddMirrorRow"), false);
+	};
+	
+	// Updates or add a mirror row (if entry is missing from the data table). Only calls Modify() at most once.
+	auto UpdateOrAddMirrorRow = [this,  &TryToModify, &CategoryStates, &FindOrAddMirrorRow, &NumUpdatedRows, bShouldUpdateExisting, bShouldAddMissing]
+	(const FName& EntryName, const FName& MirroredName, EMirrorRowType::Type RowType, bool bEnabled)
+	{
+		FCategoryState& State = CategoryStates[RowType];
+		
+		// Ambiguous table state. Skip auto-update one of multiple rows.
+		if (State.DuplicateEntryNames.Contains(EntryName))
+		{
+			UE_LOGF(LogAnimation, Warning, "MirrorDataTable '%ls': Ambiguous table state. Skipping auto-update of entry '%ls' due to it having duplicates.", *GetName(), *EntryName.ToString());
+			return;
+		}
+			
+		// Mark this entry as present in the current skeleton scan.
+		State.ProcessedEntryNames.Add(EntryName);
+			
+		const bool bHasExistingEntry = State.ExistingEntryNames.Contains(EntryName);
+			
+		FMirrorTableRow* MirrorRow = nullptr;
+			
+		if (bHasExistingEntry)
+		{
+			// Reliable lookup using actual row name
+			FName ExistingRowName = State.EntryNameToRowName.FindRef(EntryName);
+			MirrorRow = FindRow<FMirrorTableRow>(ExistingRowName, TEXT("UMirrorDataTable::UpdateOrAddMirrorRow"), false);
+			
+			
+			// Update values.
+			if (MirrorRow && bShouldUpdateExisting)
 			{
-				// Row names must be unique - in this case append a category name
-				FString RowString = Name.ToString() + CategoryName[RowType];
-				if (RenameAttempts > 0)
+				TryToModify();
+				
+				const FMirrorTableRow UpdatedRow = { EntryName, MirroredName, RowType, bEnabled };
+				const bool bHasAnythingChanged = UpdatedRow != *MirrorRow;
+				
+				*MirrorRow = UpdatedRow;
+
+				if (bHasAnythingChanged)
 				{
-					RowString.Appendf(TEXT("%d"), RenameAttempts);
+					++NumUpdatedRows;
 				}
-				RowName = *RowString;
-				ExistingRow = FindRow<FMirrorTableRow>(RowName, TEXT("UMirrorDataTable::FindReplaceMirroredNames"), false);
-				RenameAttempts++;
+				
+				UE_LOGF(LogAnimation, Log, "MirrorDataTable '%ls': Updated data for row '%ls'.", *GetName(),  *ExistingRowName.ToString());
 			}
-			
-			if (RowStruct)
-			{
-				Modify();
-				// Allocate data to store information, using UScriptStruct to know its size
-				uint8* RowData = (uint8*)FMemory::Malloc(RowStruct->GetStructureSize());
-				RowStruct->InitializeStruct(RowData);
-				// Add to row map
-				AddRowInternal(RowName, RowData);
-				bChangedTable = true;
-			}
-			FMirrorTableRow* MirrorRow = FindRow<FMirrorTableRow>(RowName, TEXT("UMirrorDataTable::FindReplaceMirroredNames"), false);
+		}
+		else
+		{
+			MirrorRow = FindOrAddMirrorRow(EntryName, RowType, bShouldAddMissing);
 			if (MirrorRow)
 			{
+				// TryToModify() is already called by FindOrAddMirrorRow().
+				
+				MirrorRow->Name = EntryName;
 				MirrorRow->MirroredName = MirroredName;
-				MirrorRow->Name = Name;
 				MirrorRow->MirrorEntryType = RowType;
+				MirrorRow->bEnabled = bEnabled;
 			}
 		}
 	};
-
-	for (int32 BoneIndex = 0; BoneIndex < RefSkeleton.GetNum(); ++BoneIndex)
+	
+	const FReferenceSkeleton& RefSkeleton = Skeleton->GetReferenceSkeleton();
+	
+	// Helper to get mirrored entry name.
+	UAnimationSettings* AnimationSettings = UAnimationSettings::Get();
+	auto GetMirroredName = [&AnimationSettings, this](const FName& InName) -> FName
 	{
-		FName BoneName = RefSkeleton.GetBoneName(BoneIndex);
-		FName MirroredName = FindReplace(BoneName);
-		if (!MirroredName.IsNone() && RefSkeleton.FindBoneIndex(MirroredName) != INDEX_NONE)
+		FName MirroredName = FindReplace(InName);
+			
+		if (MirroredName.IsNone())
 		{
-			AddMirrorRow(BoneName, MirroredName, EMirrorRowType::Bone);
+			if (AnimationSettings && AnimationSettings->bOnMirrorPairFindReplaceFailedMapToSelf)
+			{
+				MirroredName = InName;
+				UE_LOGF(LogAnimation, Warning, "MirrorDataTable '%ls': No mirror entry found for '%ls', falling back to self.", *GetName(), *InName.ToString());
+			}
+			else
+			{
+				UE_LOGF(LogAnimation, Warning, "MirrorDataTable '%ls': No mirror entry found for '%ls'.", *GetName(), *InName.ToString());
+			}
+		}
+		return MirroredName;
+	};
+	
+	// Add mirrored bone pairs.
+	Progress.EnterProgressFrame(1.0f, LOCTEXT("FindReplaceMirroredNamesOperation_AddingBones", "Adding bone mirror pair entries..."));
+	if (BoneScope == EMirrorTableBoneRefreshScope::ExplicitBoneList && BoneScopeNameList.Num() > 0)
+	{
+		for (const FName& BoneName : BoneScopeNameList)
+		{
+			if (RefSkeleton.FindBoneIndex(BoneName) == INDEX_NONE)
+			{
+				continue;
+			}
+			const FName MirroredName = GetMirroredName(BoneName);
+			const bool bEnabled = !MirroredName.IsNone() && RefSkeleton.FindBoneIndex(MirroredName) != INDEX_NONE;
+			UpdateOrAddMirrorRow(BoneName, MirroredName, EMirrorRowType::Bone, bEnabled);
 		}
 	}
-
-	for (const FName& Notify : Skeleton->AnimationNotifies)
+	else
 	{
-		FName MirroredName = FindReplace(Notify);
-		if (!MirroredName.IsNone() && Skeleton->AnimationNotifies.Contains(MirroredName))
+		for (int32 BoneIndex = 0; BoneIndex < RefSkeleton.GetNum(); ++BoneIndex)
 		{
-			AddMirrorRow(Notify, MirroredName, EMirrorRowType::AnimationNotify);
+			const FName BoneName = RefSkeleton.GetBoneName(BoneIndex);
+			const FName MirroredName = GetMirroredName(BoneName);
+			const bool bEnabled = !MirroredName.IsNone() && RefSkeleton.FindBoneIndex(MirroredName) != INDEX_NONE;
+			UpdateOrAddMirrorRow(BoneName, MirroredName, EMirrorRowType::Bone, bEnabled);
 		}
 	}
 	
+	// Add mirrored animation notify pairs.
+	Progress.EnterProgressFrame(1.0f, LOCTEXT("FindReplaceMirroredNamesOperation_AddingNotifies", "Adding animation notify mirror pair entries..."));
+	for (const FName& Notify : Skeleton->AnimationNotifies)
+	{
+		const FName MirroredName = GetMirroredName(Notify);
+		const bool bEnabled = !MirroredName.IsNone() && Skeleton->AnimationNotifies.Contains(MirroredName);
+		
+		UpdateOrAddMirrorRow(Notify, MirroredName, EMirrorRowType::AnimationNotify, bEnabled);
+	}
+	
+	// Add mirrored sync markers pairs.
+	Progress.EnterProgressFrame(1.0f, LOCTEXT("FindReplaceMirroredNamesOperation_AddingSyncMarkers", "Adding sync marker mirror pair entries..."));
 	for (const FName& SyncMarker : Skeleton->GetExistingMarkerNames())
 	{
-		FName MirroredName = FindReplace(SyncMarker);
-		if (!MirroredName.IsNone() && Skeleton->GetExistingMarkerNames().Contains(MirroredName))
-		{
-			AddMirrorRow(SyncMarker, MirroredName, EMirrorRowType::SyncMarker);
-		}
+		const FName MirroredName = GetMirroredName(SyncMarker);
+		const bool bEnabled = !MirroredName.IsNone() && Skeleton->GetExistingMarkerNames().Contains(MirroredName);
+		
+		UpdateOrAddMirrorRow(SyncMarker, MirroredName, EMirrorRowType::SyncMarker, bEnabled);
 	}
 
-	TArray<FName> CurveNames;
-	Skeleton->GetCurveMetaDataNames(CurveNames);
+	// Add mirrored curves pairs.
+	Progress.EnterProgressFrame(1.0f, LOCTEXT("FindReplaceMirroredNamesOperation_AddingCurves", "Adding curve mirror pair entries..."));
+	{
+		TArray<FName> CurveNames;
+		Skeleton->GetCurveMetaDataNames(CurveNames);
+		
+		for (const FName& CurveName : CurveNames)
+		{
+			const FName MirroredName = GetMirroredName(CurveName);
+			const bool bEnabled = !MirroredName.IsNone() && CurveNames.Contains(MirroredName);
 
-	// For every source curve, try to find the curve with the same name in the target.
-	for (const FName& CurveName : CurveNames)
+			UpdateOrAddMirrorRow(CurveName, MirroredName, EMirrorRowType::Curve, bEnabled);
+		}
+	}
+	
+	// @todo: Disable rows with no entry name or mirror entry name.
+	
+	// Disable stale skeleton entries.
+	Progress.EnterProgressFrame(1.0f, LOCTEXT("FindReplaceMirroredNamesOperation_DisabledEntriesMissingFromSkeleton", "Disabling skeleton based entries missing from the associated skeleton ..."));
+	if (bShouldDisableStale)
 	{
-		FName MirroredName = FindReplace(CurveName);
-		if (!MirroredName.IsNone() && CurveNames.Contains(MirroredName))
+		// Exclude "Custom" category since those entries do not come from the skeleton.
+		constexpr int CategoriesToCheck = (CategoryStates.Num() - 1);
+		
+		for (int CategoryIndex = 0; CategoryIndex < CategoriesToCheck; ++CategoryIndex)
 		{
-			AddMirrorRow(CurveName, MirroredName, EMirrorRowType::Curve);
+			FCategoryState& State = CategoryStates[CategoryIndex];
+			
+			// Entries that existed in the data table but were not encountered from the current skeleton scan.
+			TSet<FName> StaleEntryNames = State.ExistingEntryNames.Difference(State.ProcessedEntryNames);
+			
+			for (const FName& EntryName : StaleEntryNames)
+			{
+				// Ambiguous table state. Skip auto-update one of multiple rows.
+				if (State.DuplicateEntryNames.Contains(EntryName))
+				{
+					continue;
+				}
+				
+				// Reliable lookup using actual row name.
+				const FName* StaleRowName = State.EntryNameToRowName.Find(EntryName);
+				if (!StaleRowName)
+				{
+					continue;
+				}
+				
+				// Get stale row.
+				FMirrorTableRow* StaleMirrorRow = FindRow<FMirrorTableRow>(*StaleRowName, TEXT("UMirrorDataTable::ApplyFindReplaceExpressions"), false);
+				if (!StaleMirrorRow)
+				{
+					continue;
+				}
+				
+				// Mark as stale, if not already.
+				if (StaleMirrorRow->bEnabled)
+				{
+					TryToModify();
+					
+					StaleMirrorRow->bEnabled = false;
+					
+					++NumStaleRows;
+					
+					UE_LOGF(LogAnimation, Log, "MirrorDataTable '%ls': Disabled stale row '%ls'.", *GetName(), *StaleRowName->ToString());
+				}
+			}
 		}
 	}
-
+	
+	// Check if we have any duplicate entries at all.
+	for (int CategoryIndex = 0; CategoryIndex < CategoryStates.Num(); ++CategoryIndex)
+	{
+		NumDuplicateRows += CategoryStates[CategoryIndex].DuplicateEntryNames.Num();
+	}
+	
+	// Keep tracks that were up-to-date now.
+	const bool bIsFullyRefreshed = bShouldUpdateExisting && bShouldAddMissing && bShouldDisableStale;
+	const bool bAreAnyWarningsPresent = NumDuplicateRows > 0;
+	if (bIsFullyRefreshed && !bAreAnyWarningsPresent)
+	{
+		SkeletonHierarchyGuid = Skeleton->GetGuid();
+		SkeletonVirtualBonesHierarchyGuid = Skeleton->GetVirtualBoneGuid();
+	}
+	
+	// Let others react to data table changes.
+	Progress.EnterProgressFrame(1.0f, LOCTEXT("FindReplaceMirroredNamesOperation_UpdatingRuntimeData", "Updating runtime mirror maps..."));
 	if (bChangedTable)
 	{
 		OnDataTableChanged().Broadcast();
+		
+		// Note: There is no need to run FillMirrorArrays(), as this will be handled via broadcasting the OnDataTableChanged delegate.
+		
+		// Let editor know we've changed the data and refresh view.
+		FDataTableEditorUtils::BroadcastPostChange(this, FDataTableEditorUtils::EDataTableChangeInfo::RowList);
 	}
-	FillMirrorArrays(); 
+	
+	// Show completion notification from operation.
+	if (Options.bShowNotification)
+	{
+		const bool bFailure = bAreAnyWarningsPresent;
+		const FText FormatFailure = LOCTEXT("MirrorReimportSummary_Failure", "Update complete with warnings. See log for Details. Skipped auto-update for {3} duplicate entries. Added {0} entries, updated {1} entries, disabled {2} stale entries."); 
+		const FText FormatSuccess = LOCTEXT("MirrorReimportSummary_Success", "Update complete. See log for details. Added {0} entries, updated {1} entries, disabled {2} stale entries.");
+		const FText Message = FText::Format(
+			bFailure ? FormatFailure : FormatSuccess,
+			NumAddedRows,
+			NumUpdatedRows,
+			NumStaleRows,
+			NumDuplicateRows
+		);
+		
+		FNotificationInfo Info(Message);
+		Info.bFireAndForget = true;
+		Info.ExpireDuration = 7.0f;
+		Info.bUseSuccessFailIcons = true;
+		
+		if (TSharedPtr<SNotificationItem> Notification = FSlateNotificationManager::Get().AddNotification(Info))
+		{
+			Notification->SetCompletionState(bFailure ? SNotificationItem::CS_Fail: SNotificationItem::CS_Success);
+		}
+	}
+}
+
+UMirrorDataTable::ESyncStatus UMirrorDataTable::GetSkeletonSyncStatus() const
+{
+	if (!Skeleton)
+	{
+		return ESyncStatus::Unknown;
+	}
+
+	if (!SkeletonHierarchyGuid.IsValid())
+	{
+		return ESyncStatus::Unknown;
+	}
+
+	const bool bIsUpToDate = Skeleton->GetGuid() == SkeletonHierarchyGuid && Skeleton->GetVirtualBoneGuid() == SkeletonVirtualBonesHierarchyGuid;
+
+	return bIsUpToDate ? ESyncStatus::UpToDate : ESyncStatus::Stale;
+}
+
+void UMirrorDataTable::InvalidateCachedSkeletonData()
+{
+	// Non-zero sentinel: Will be valid but won't match any skeleton, so GetSkeletonSyncStatus() returns Stale (not Unknown).
+	static constexpr FGuid StaleGuid(MAX_uint32, MAX_uint32, MAX_uint32, MAX_uint32);
+	
+	SkeletonHierarchyGuid = StaleGuid;
+	SkeletonVirtualBonesHierarchyGuid = StaleGuid;
 }
 
 #endif // WITH_EDITOR
+
 void UMirrorDataTable::FillCompactPoseMirrorBones(const FBoneContainer& BoneContainer, const TCustomBoneIndexArray<FSkeletonPoseBoneIndex, FSkeletonPoseBoneIndex>& MirrorBoneIndexes, TCustomBoneIndexArray<FCompactPoseBoneIndex, FCompactPoseBoneIndex>& OutCompactPoseMirrorBones)
 {
 	const int32 NumReqBones = BoneContainer.GetCompactPoseNumBones();
@@ -385,7 +704,7 @@ void UMirrorDataTable::FillMirrorBoneIndexes(const USkeleton* InSkeleton, TCusto
 	TMap<FName, FName> NameToMirrorNameBoneMap; 
 	ForeachRow<FMirrorTableRow>(TEXT("UMirrorDataTable::FillMirrorBoneIndexes"), [&NameToMirrorNameBoneMap](const FName& Key, const FMirrorTableRow& Value) mutable
 		{
-			if (Value.MirrorEntryType == EMirrorRowType::Bone)
+			if (Value.MirrorEntryType == EMirrorRowType::Bone && Value.bEnabled)
 			{
 				NameToMirrorNameBoneMap.Add(Value.Name, Value.MirroredName);
 			}
@@ -442,12 +761,13 @@ void UMirrorDataTable::FillCompactPoseAndComponentRefRotations(
 	}
 }
 
-
 void UMirrorDataTable::FillMirrorArrays()
 {
+	DisabledEntries.Empty();
 	SyncToMirrorSyncMap.Empty();
 	AnimNotifyToMirrorAnimNotifyMap.Empty();
 	CurveToMirrorCurveMap.Empty();
+	
 	if (!Skeleton)
 	{
 		BoneToMirrorBoneIndex.Empty();
@@ -458,6 +778,12 @@ void UMirrorDataTable::FillMirrorArrays()
 	
 	ForeachRow<FMirrorTableRow>(TEXT("UMirrorDataTable::FillMirrorArrays"), [this](const FName& Key, const FMirrorTableRow& Value) mutable
 		{
+			if (!Value.bEnabled)
+			{
+				DisabledEntries.Add(Value.MirroredName);
+				return;
+			}
+			
 			switch (Value.MirrorEntryType)
 			{
 			case  EMirrorRowType::Curve:
@@ -487,4 +813,52 @@ void UMirrorDataTable::FillMirrorArrays()
 	);
 }
 
+void UMirrorDataTable::InitCategoryEntryStates(TStaticArray<FCategoryState, 5>& OutStates) const
+{
+	static FString CategoryNames[] = { "Bone", "AnimationNotify", "Curve", "SyncMarker", "Custom" };
+	
+	ForeachRow<FMirrorTableRow>(TEXT("UMirrorDataTable::BuildCategoryEntryState"), [&OutStates](const FName& RowName, const FMirrorTableRow& Value)
+		{
+			const int32 CategoryIndex = Value.MirrorEntryType;
+			check(CategoryIndex >= 0 && CategoryIndex < OutStates.Num())
+			
+			FCategoryState & State = OutStates[CategoryIndex];
+
+			// We can assume that bones, notifies, sync markers, and curves are uniquely named.
+			// However, if any matching row name (of the same type) was manually added this can lead to duplicates, in this we just skip updating them and inform user.
+			if (const FName* ExistingRowName = State.EntryNameToRowName.Find(Value.Name))
+			{
+				UE_LOGF(LogAnimation, Warning, "Duplicate mirror table entry `%ls` found in category `%ls`. Existing row `%ls`, duplicate row `%ls`.", *Value.Name.ToString(), *CategoryNames[CategoryIndex], *ExistingRowName->ToString(), *RowName.ToString())
+				State.DuplicateEntryNames.Add(Value.Name);
+				return;
+			}
+			
+			State.ExistingEntryNames.Add(Value.Name);
+			State.EntryNameToRowName.Add(Value.Name, RowName);
+		}
+	);
+}
+
+FName UMirrorDataTable::MakeUniqueMirrorRowName(const FName EntryName, const EMirrorRowType::Type RowType) const
+{
+	static const FString CategorySuffix[] = { TEXT(":Bone"), TEXT(":Notify"), TEXT(":Curve"), TEXT(":SyncMarker"), TEXT(":Custom") };
+	
+	FName CandidateRowName = EntryName;
+	uint32 RenamedAttempts = 0;
+	
+	while (FindRow<FMirrorTableRow>(CandidateRowName, TEXT("UMirrorDataTable::MakeUniqueMirrorRowName"), false))
+	{
+		FString RowString = EntryName.ToString() + CategorySuffix[RowType];
+		if (RenamedAttempts > 0)
+		{
+			RowString.Appendf(TEXT("%u"), RenamedAttempts);
+		}
+		
+		CandidateRowName = FName(*RowString);
+		++RenamedAttempts;
+	}
+	
+	return CandidateRowName;
+}
+
 #undef LOCTEXT_NAMESPACE
