Add clear comments indicating that most of FTickFunction is not threadsafe, and fix up other tick comments.
Deprecate GetLastTickGameTime and rename as it does not do what it says.
#Synced-CL 52401938

diff --git a/Engine/Source/Runtime/Engine/Classes/Engine/EngineBaseTypes.h b/Engine/Source/Runtime/Engine/Classes/Engine/EngineBaseTypes.h
--- a/Engine/Source/Runtime/Engine/Classes/Engine/EngineBaseTypes.h
+++ b/Engine/Source/Runtime/Engine/Classes/Engine/EngineBaseTypes.h
@@ -4,7 +4,7 @@
 
 /**
  *	This file is for shared structs and enums that need to be declared before the rest of Engine.
- *  The typical use case is for structs used in the renderer and also in script code.
+ *  This includes structs used by the renderer and task execution framework.
  */
 
 #include "Async/TaskGraphFwd.h"
@@ -23,9 +23,7 @@ struct FSlateBrush;
 struct FSlateIcon;
 struct FTickContext;
 
-//
-//	EInputEvent
-//
+/** Describes the type of input event received for a specific key or button */
 UENUM( BlueprintType, meta=(ScriptName="InputEventType"))
 enum EInputEvent : int
 {
@@ -37,6 +35,7 @@ enum EInputEvent : int
 	IE_MAX                  =5,
 };
 
+/** Describes how the application consumes mouse input events */
 UENUM()
 enum class EMouseCaptureMode : uint8
 {
@@ -52,6 +51,7 @@ enum class EMouseCaptureMode : uint8
 	CaptureDuringRightMouseDown,
 };
 
+/** Describes the behavior of a mouse cursor when inside the viewport window */
 UENUM()
 enum class EMouseLockMode : uint8
 {
@@ -97,10 +97,10 @@ enum ETickingGroup : int
 	/** Any item that needs rigid body and cloth simulation to be complete before being executed. */
 	TG_PostPhysics UMETA(DisplayName="Post Physics"),
 
-	/** Any item that needs the update work to be done before being ticked. */
+	/** Any item that needs to be ticked after all normal gameplay tasks. */
 	TG_PostUpdateWork UMETA(DisplayName="Post Update Work"),
 
-	/** Catchall for anything demoted to the end. */
+	/** Last group which tick functions can be delayed into because of dependencies. */
 	TG_LastDemotable UMETA(Hidden, DisplayName = "Last Demotable"),
 
 	/** Special tick group that is not actually a tick group. After every tick group this is repeatedly re-run until there are no more newly spawned items to run. */
@@ -109,43 +109,50 @@ enum ETickingGroup : int
 	TG_MAX,
 };
 
-/**
- * This is small structure to hold prerequisite tick functions
- */
+/** This is a small structure to define how a tick function depends on other tick functions. */
 USTRUCT()
 struct FTickPrerequisite
 {
-	GENERATED_USTRUCT_BODY()
+	GENERATED_BODY()
 
-	/** Tick functions live inside of UObjects, so we need a separate weak pointer to the UObject solely for the purpose of determining if PrerequisiteTickFunction is still valid. */
+public:
+	/**
+	 * Tick functions are owned by UObjects, so we need a separate weak pointer to the UObject
+	 * solely for the purpose of determining if PrerequisiteTickFunction is still valid.
+	 */
 	TWeakObjectPtr<class UObject> PrerequisiteObject;
 
-	/** Pointer to the actual tick function and must be completed prior to our tick running. */
+	/** 
+	 * Pointer to the actual tick function and must be completed prior to our tick running.
+	 * @warning This should not be accessed directly, use Get() to validate the PrerequisiteObject first.
+	 */
 	struct FTickFunction*		PrerequisiteTickFunction;
 
-	/** Noop constructor. */
 	FTickPrerequisite()
-	: PrerequisiteTickFunction(nullptr)
+		: PrerequisiteTickFunction(nullptr)
 	{
 	}
+
 	/** 
-		* Constructor
-		* @param TargetObject - UObject containing this tick function. Only used to verify that the other pointer is still usable
-		* @param TargetTickFunction - Actual tick function to use as a prerequisite
-	**/
+	 * Constructor
+	 * @param TargetObject - UObject containing this tick function. Only used to verify that the other pointer is still usable
+	 * @param TargetTickFunction - Actual tick function to use as a prerequisite
+	 */
 	FTickPrerequisite(UObject* TargetObject, struct FTickFunction& TargetTickFunction)
-	: PrerequisiteObject(TargetObject)
-	, PrerequisiteTickFunction(&TargetTickFunction)
+		: PrerequisiteObject(TargetObject)
+		, PrerequisiteTickFunction(&TargetTickFunction)
 	{
 		check(PrerequisiteTickFunction);
 	}
+
 	/** Equality operator, used to prevent duplicates and allow removal by value. */
 	bool operator==(const FTickPrerequisite& Other) const
 	{
 		return PrerequisiteObject == Other.PrerequisiteObject &&
 			PrerequisiteTickFunction == Other.PrerequisiteTickFunction;
 	}
-	/** Return the tick function, if it is still valid. Can be null if the tick function was null or the containing UObject has been garbage collected. */
+
+	/** Return the tick function, if it is still valid. Can be null if the tick function was null or the containing UObject is invalid. */
 	struct FTickFunction* Get()
 	{
 		if (PrerequisiteObject.IsValid(true))
@@ -165,37 +172,39 @@ struct FTickPrerequisite
 	}
 };
 
-/** 
-* Abstract Base class for all tick functions.
-**/
+/**
+ * Abstract base class for all tick functions.
+ * Registered tick functions are executed by the functions in FTickTaskManagerInterface during the engine tick frame.
+ * 
+ * @warning Most of these functions are not threadsafe. The tick manager stores raw pointers to these structs 
+ * and accesses them on the game thread, so registered functions cannot be safely destroyed during a tick frame.
+ */
 USTRUCT()
 struct FTickFunction
 {
-	GENERATED_USTRUCT_BODY()
+	GENERATED_BODY()
 
 public:
 	// The following UPROPERTYs are for configuration and inherited from the CDO/archetype/blueprint etc
 
 	/**
-	 * Defines the minimum tick group for this tick function. These groups determine the relative order of when objects tick during a frame update.
-	 * Given prerequisites, the tick may be delayed.
-	 *
-	 * @see ETickingGroup 
-	 * @see FTickFunction::AddPrerequisite()
+	 * Defines the first tick group where this function can be executed.
+	 * These groups determine the relative order of when objects tick during a frame update.
+	 * If this function has prerequisites, execution may be delayed to a later group.
 	 */
 	UPROPERTY(EditDefaultsOnly, Category="Tick", AdvancedDisplay)
 	TEnumAsByte<enum ETickingGroup> TickGroup;
 
 	/**
-	 * Defines the tick group that this tick function must finish in. These groups determine the relative order of when objects tick during a frame update.
-	 *
-	 * @see ETickingGroup 
+	 * Defines the tick group that this tick function must finish within. 
+	 * Normal synchronous ticks do not need to set this manually.
+	 * The game thread will stall at the end of this group if the function (or a spawned task) is still in progress.
 	 */
 	UPROPERTY(EditDefaultsOnly, Category="Tick", AdvancedDisplay)
 	TEnumAsByte<enum ETickingGroup> EndTickGroup;
 
 public:
-	/** Bool indicating that this function should execute even if the game is paused. Pause ticks are very limited in capabilities. **/
+	/** Bool indicating that this function should execute even if the gameplay is paused. Pause ticks are very limited in capabilities. */
 	UPROPERTY(EditDefaultsOnly, Category="Tick", AdvancedDisplay)
 	uint8 bTickEvenWhenPaused:1;
 
@@ -217,7 +226,7 @@ public:
 	/** Run this tick first within the tick group, presumably to start async tasks that must be completed with this tick group, hiding the latency. */
 	uint8 bHighPriority:1;
 
-	/** If false, this tick will run on the game thread, otherwise it will run on any thread in parallel with the game thread and in parallel with other "async ticks" **/
+	/** If false, this tick will only run on the game thread. If true it can run on any thread in parallel with the game thread and with other async ticks and tasks */
 	uint8 bRunOnAnyThread:1;
 
 	/** (experimental) if true, the tick function will be run transactionally */
@@ -232,26 +241,24 @@ private:
 
 	enum class ETickState : uint8
 	{
+		// This tick will not execute in the next/current frame.
 		Disabled,
+		// This tick will execute normally.
 		Enabled,
+		// This tick has a tick interval that is counting down before the next execution.
 		CoolingDown
 	};
 
-	/** 
-	 * If Disabled, this tick will not fire
-	 * If CoolingDown, this tick has an interval frequency that is being adhered to currently
-	 * CAUTION: Do not set this directly
-	 **/
+	/** Internal tick state, set by tick manager every frame. */
 	ETickState TickState : 2;
 
 public:
-	/** The frequency in seconds at which this tick function will be executed.  If less than or equal to 0 then it will tick every frame */
+	/** The time in seconds between executions of this tick function. If <= 0 then it will tick every frame. */
 	UPROPERTY(EditDefaultsOnly, Category="Tick", meta=(DisplayName="Tick Interval (secs)"))
 	float TickInterval;
 
 private:
-
-	/** Prerequisites for this tick function **/
+	/** List of prerequisites for this tick function, modified by calling Add/RemovePrerequisite. */
 	TArray<struct FTickPrerequisite> Prerequisites;
 
 	/** Defines the internal state of a tick function, set in TickTaskManager */
@@ -267,7 +274,7 @@ private:
 		HasCompletionEvent,
 	};
 
-	/** Internal Data structure that contains members only required for a registered tick function **/
+	/** Internal data structure that contains members only required for a registered tick function. */
 	struct FInternalData
 	{
 		FInternalData();
@@ -275,7 +282,7 @@ private:
 		/** Whether the tick function is registered. */
 		bool bRegistered : 1;
 
-		/** Cache whether this function was rescheduled as an interval function during StartParallel */
+		/** Cache whether this function was rescheduled as an interval function by the TickManager */
 		bool bWasInterval:1;
 
 		/** Internal state, this determines what TaskPointer points to */
@@ -296,59 +303,88 @@ private:
 		/** Pointer to a type determined by TaskState, do not access directly */
 		void* TaskPointer;
 
-		/** The next function in the cooling down list for ticks with an interval*/
+		/** The next function in the cooling down list for ticks with an interval */
 		FTickFunction* Next;
 
 		/** 
-		 * If TickFrequency is greater than 0 and tick state is CoolingDown, this is the time, 
+		 * If TickInterval is greater than 0 and the TickState is CoolingDown, this is the time, 
 		 * relative to the element ahead of it in the cooling down list, remaining until the next time this function will tick 
-		 **/
+		 */
 		float RelativeTickCooldown;
 
 		/** 
-		 * The last world game time at which we were ticked. Game time used is dependent on bTickEvenWhenPaused
-		 * Valid only if we've been ticked at least once since having a tick interval; otherwise set to -1.f
-		 **/
-		double LastTickGameTimeSeconds;
+		 * The last world game time at which we were ticked, uses GetUnpausedTimeSeconds for bTickEvenWhenPaused.
+		 * Valid only if we've been ticked at least once since having a tick interval, otherwise set to -1.
+		 */
+		double LastIntervalTickSeconds;
 
-		/** Back pointer to the FTickTaskLevel containing this tick function if it is registered **/
-		class FTickTaskLevel*						TickTaskLevel;
+		/** Back pointer to the FTickTaskLevel managing this tick function, if it is currently registered */
+		class FTickTaskLevel* TickTaskLevel;
 	};
 
-	/** Lazily allocated struct that contains the necessary data for a tick function that is registered. **/
+	/** Lazily allocated struct that contains the necessary data for a tick function that is registered */
 	TUniquePtr<FInternalData> InternalData;
 
 public:
-	/** Default constructor, intitalizes to reasonable defaults **/
+	/** Default constructor, initializes to disabled. */
 	ENGINE_API FTickFunction();
-	/** Destructor, unregisters the tick function **/
+
+	/** Destructor, unregisters the tick function if it is still registered. */
 	ENGINE_API virtual ~FTickFunction();
 
 	/** 
-	 * Adds the tick function to the primary list of tick functions. 
-	 * @param Level - level to place this tick function in
-	 **/
+	 * Adds the tick function to the primary list of tick functions for a specific level.
+	 * If called during a tick frame, this will tick in the next tick group.
+	 * 
+	 * @warning This function is not threadsafe and must be called from the game thread.
+	 */
 	ENGINE_API void RegisterTickFunction(class ULevel* Level);
-	/** Removes the tick function from the primary list of tick functions. **/
+
+	/**
+	 * Removes the tick function from the list of tick functions to automatically execute.
+	 * This does not cancel ticks that have already been scheduled in the current tick frame.
+	 * 
+	 * @warning This function is not threadsafe and must be called from the game thread.
+	 */
 	ENGINE_API void UnRegisterTickFunction();
-	/** See if the tick function is currently registered */
+
+	/** 
+	 * See if the tick function is currently registered to a level in the tick manager.
+	 * If this is true and the tick is also enabled, the tick manager will automatically schedule it for execution.
+	 */
 	bool IsTickFunctionRegistered() const { return (InternalData && InternalData->bRegistered); }
 
-	/** Enables or disables this tick function. **/
+	/** 
+	 * Enables or disables this tick function. Disabled tick functions will never execute, even if registered.
+	 * Enabling during a tick frame will tick it in the next tick group if it has not already ticked this frame.
+	 * Disabling during a tick frame will stop it from executing unless it is already in progress.
+	 * 
+	 * @warning This function is not threadsafe to call on a registered tick function and must be called from the game thread. 
+	 */
 	ENGINE_API void SetTickFunctionEnable(bool bInEnabled);
-	/** Returns whether the tick function is currently enabled */
+
+	/** Returns whether the tick function is currently enabled and is possible to execute. */
 	bool IsTickFunctionEnabled() const { return TickState != ETickState::Disabled; }
 
-	/** Update tick interval in the system and overwrite the current cooldown if any. */
+	/** 
+	 * Update tick interval in the system and overwrite the current cooldown if any.
+	 * This may not affect the current tick frame.
+	 *
+	 * @warning This function is not threadsafe to call on a registered tick function and must be called from the game thread. 
+	 */
 	ENGINE_API void UpdateTickIntervalAndCoolDown(float NewTickInterval);
 
-	/** Returns true if it is valid to access this tick function's completion handle */
+	/** 
+	 * Returns true if GetCompletionHandle will return a valid completion handle.
+	 * This is only true for functions that are registered and have not finished ticking in the current frame.
+	 */
 	ENGINE_API bool IsCompletionHandleValid() const;
 
 	/**
 	 * Gets the current completion handle of this tick function, so it can be delayed until a later point when some additional
 	 * tasks have been completed. Only valid after StartFrame has been called and then only until the TickFunction finishes
 	 * execution. This returns a reference so IsCompletionHandleValid must be called first.
+	 * The returned handle can be safely accessed on any thread to coordinate with other async tasks.
 	 */
 	ENGINE_API FGraphEventRef GetCompletionHandle() const;
 
@@ -363,18 +399,18 @@ public:
 	 */
 	ENGINE_API virtual bool DispatchManually();
 
-	/** 
-	 * Gets the action tick group that this function will be elligible to start in.
-	 * Only valid after StartFrame has been called through the end of the frame.
+	/**
+	 * Gets the actual tick group that this function will start executing in.
+	 * This is only guaranteed to be correct after StartFrame has been called through the end of the frame.
 	 */
 	TEnumAsByte<enum ETickingGroup> GetActualTickGroup() const
 	{
 		return (InternalData ? InternalData->ActualStartTickGroup : TickGroup);
 	}
 
-	/** 
-	 * Gets the action tick group that this function will be required to end in.
- 	 * Only valid after StartFrame has been called through the end of the frame.
+	/**
+	 * Returns the last tick group this function could complete within.
+ 	 * This is only guaranteed to be correct after StartFrame has been called through the end of the frame.
 	 */
 	TEnumAsByte<enum ETickingGroup> GetActualEndTickGroup() const
 	{
@@ -382,27 +418,31 @@ public:
 	}
 
 	/** 
-	 * Adds a tick function to the list of prerequisites...in other words, adds the requirement that TargetTickFunction is called before this tick function is 
-	 * @param TargetObject - UObject containing this tick function. Only used to verify that the other pointer is still usable
-	 * @param TargetTickFunction - Actual tick function to use as a prerequisite
-	 **/
+	 * Adds a tick function to the list of prerequisites. This adds the requirement that TargetTickFunction is called before this tick function.
+	 * This prerequisite will only take effect if the TargetTickFunction is registered and enabled, and will not affect the current tick frame.
+	 * @param TargetObject - UObject containing TargetTickFunction, only used to verify that the other pointer is still usable.
+	 * @param TargetTickFunction - Actual tick function to use as a prerequisite.
+	 */
 	ENGINE_API void AddPrerequisite(UObject* TargetObject, struct FTickFunction& TargetTickFunction);
 
 	/** 
-	 * Removes a prerequisite that was previously added.
-	 * @param TargetObject - UObject containing this tick function. Only used to verify that the other pointer is still usable
-	 * @param TargetTickFunction - Actual tick function to use as a prerequisite
-	 **/
+	 * Removes a prerequisite that was previously added. This will not affect the current tick frame.
+	 * @param TargetObject - UObject containing this tick function, only used to verify that the other pointer is still usable.
+	 * @param TargetTickFunction - Actual tick function to use as a prerequisite.
+	 */
 	ENGINE_API void RemovePrerequisite(UObject* TargetObject, struct FTickFunction& TargetTickFunction);
 
-	/** 
-	 * Sets this function to hipri and all prerequisites recursively
-	 * @param bInHighPriority - priority to set
-	 **/
+	/**
+	 * Sets this function to normal or high priority, and also modifies all prerequisites recursively.
+	 *
+	 * @warning This only affects current prerequisites, prerequisites that are added later will not inherit the priority.
+	 */
 	ENGINE_API void SetPriorityIncludingPrerequisites(bool bInHighPriority);
 
 	/**
-	 * @return a reference to prerequisites for this tick function.
+	 * Return a reference to the prerequisites list for this tick function.
+	 *
+	 * @warning Modifying this can create invalid data, use Add/RemovePrerequisite for normal use cases.
 	 */
 	TArray<struct FTickPrerequisite>& GetPrerequisites()
 	{
@@ -410,7 +450,7 @@ public:
 	}
 
 	/**
-	 * @return a reference to prerequisites for this tick function (const).
+	 * Return a reference to the prerequisites list for this tick function (const).
 	 */
 	const TArray<struct FTickPrerequisite>& GetPrerequisites() const
 	{
@@ -419,16 +459,19 @@ public:
 
 	/**
 	 * For actively registered tick functions with an interval, return the last time it ticked.
-	 * For unregistered or non-interval functions it will return -1;
+	 * For unregistered or non-interval functions it will return -1.
 	 */
-	double GetLastTickGameTime() const { return (InternalData ? InternalData->LastTickGameTimeSeconds : -1.0); }
+	double GetLastIntervalTickGameTime() const { return (InternalData ? InternalData->LastIntervalTickSeconds : -1.0); }
+
+	UE_DEPRECATED(5.8, "Renamed because function only returns the correct time for functions with a TickInterval")
+	double GetLastTickGameTime() const { return (InternalData ? InternalData->LastIntervalTickSeconds : -1.0); }
 
 	/** 
-	 * Correctly call ExecuteTick on a nested function, and will set the completion event properly for the duration of the tick function
+	 * Correctly call ExecuteTick on a nested function. This will set the completion event properly for the duration of the tick function.
 	 * @param DeltaTime - frame time to advance, in seconds
 	 * @param TickType - kind of tick for this frame
 	 * @param CurrentThread - thread we are executing on, useful to pass along as new tasks are created
-	 * @param MyCompletionGraphEvent - completion event for this task. Useful for holding the completetion of this task until certain child tasks are complete.
+	 * @param MyCompletionGraphEvent - completion event for this task. Useful for delaying the completion of this task until child tasks are complete.
 	 */
 	ENGINE_API void ExecuteNestedTick(float DeltaTime, ELevelTick TickType, ENamedThreads::Type CurrentThread, const FGraphEventRef& MyCompletionGraphEvent);
 
@@ -458,26 +501,26 @@ private:
 public:
 	// Functions to be used by executing tasks, don't call these directly from user code
 
-	/** Returns the delta time to use when ticking this function given the delta time and world. This also updates internal tracking data */
+	/** Returns the delta time to use for ExecuteTick using the global delta time and world. This also updates internal tracking data. */
 	ENGINE_API float CalculateDeltaTime(float DeltaTime, const class UWorld* TickingWorld);
 
-	/** Logs function info based on flags */
+	/** Logs the DiagnosticMessage and other function info. */
 	ENGINE_API void LogTickFunction(ENamedThreads::Type CurrentThread, bool bLogPrerequisites, int32 Indent = 0);
 
-	/** Logs the prerequisites */
+	/** Logs the prerequisite info specifically. */
 	ENGINE_API void ShowPrerequistes(int32 Indent = 1);
 
-	/** Clear any current task information such as the completion handle, called after execution */
+	/** Clear any current task information such as the completion handle. This must be called after execution. */
 	ENGINE_API void ClearTaskInformation();
 
-
 	/** 
-	 * Abstract function actually execute the tick. Batched tick managers should use ExecuteNestedTick
+	 * Override this abstract function to actually execute the tick. Batched tick managers should use ExecuteNestedTick.
+	 * If this function sets bRunOnAnyThread, this function could execute on a background thread.
 	 * @param DeltaTime - frame time to advance, in seconds
 	 * @param TickType - kind of tick for this frame
 	 * @param CurrentThread - thread we are executing on, useful to pass along as new tasks are created
-	 * @param MyCompletionGraphEvent - completion event for this task. Useful for holding the completetion of this task until certain child tasks are complete.
-	 **/
+	 * @param MyCompletionGraphEvent - completion event for this task. Useful for delaying the completion of this task until child tasks are complete.
+	 */
 	virtual void ExecuteTick(float DeltaTime, ELevelTick TickType, ENamedThreads::Type CurrentThread, const FGraphEventRef& MyCompletionGraphEvent) PURE_VIRTUAL(,);
 	
 #if UE_WITH_REMOTE_OBJECT_HANDLE
@@ -485,10 +528,10 @@ public:
 	FString CachedDiagnosticMessage;
 #endif
 
-	/** Abstract function to describe this tick. Used to print messages about illegal cycles in the dependency graph **/
+	/** Override this abstract function to return a human readable string used for debug commands. */
 	virtual FString DiagnosticMessage() PURE_VIRTUAL(, return TEXT("DiagnosticMessage() not implemented"););
 	
-	/** Function to give a 'context' for this tick, used for grouped active tick reporting */
+	/** Override this function to return an internal name used to group related tick functions. */
 	virtual FName DiagnosticContext(bool bDetailed)
 	{
 		return NAME_None;
@@ -518,15 +561,13 @@ struct TStructOpsTypeTraits<FTickFunction> : public TStructOpsTypeTraitsBase2<FT
 	};
 };
 
-/** 
-* Tick function that calls AActor::TickActor
-**/
+/** Tick function that calls AActor::TickActor. */
 USTRUCT()
 struct FActorTickFunction : public FTickFunction
 {
 	GENERATED_USTRUCT_BODY()
 
-	/**  AActor  that is the target of this tick **/
+	/** Actor that is the target of this tick **/
 #if UE_WITH_REMOTE_OBJECT_HANDLE
 	TObjectPtr<AActor> Target;
 #else
@@ -534,15 +575,18 @@ struct FActorTickFunction : public FTickFunction
 #endif
 
 	/** 
-		* Abstract function actually execute the tick. 
-		* @param DeltaTime - frame time to advance, in seconds
-		* @param TickType - kind of tick for this frame
-		* @param CurrentThread - thread we are executing on, useful to pass along as new tasks are created
-		* @param MyCompletionGraphEvent - completion event for this task. Useful for holding the completetion of this task until certain child tasks are complete.
-	**/
+	 * Override this function to execute an actor tick.
+	 * @param DeltaTime - frame time to advance, in seconds
+	 * @param TickType - kind of tick for this frame
+	 * @param CurrentThread - thread we are executing on, useful to pass along as new tasks are created
+	 * @param MyCompletionGraphEvent - completion event for this task. Useful for delaying the completion of this task until child tasks are complete.
+	 */
 	ENGINE_API virtual void ExecuteTick(float DeltaTime, ELevelTick TickType, ENamedThreads::Type CurrentThread, const FGraphEventRef& MyCompletionGraphEvent) override;
-	/** Abstract function to describe this tick. Used to print messages about illegal cycles in the dependency graph **/
+
+	/** Returns the name of the owning actor, override for other types of ticks. */
 	ENGINE_API virtual FString DiagnosticMessage() override;
+
+	/** Returns the name of the owning actor class. */
 	ENGINE_API virtual FName DiagnosticContext(bool bDetailed) override;
 
 #if UE_WITH_REMOTE_OBJECT_HANDLE
@@ -563,9 +607,7 @@ struct TStructOpsTypeTraits<FActorTickFunction> : public TStructOpsTypeTraitsBas
 	};
 };
 
-/** 
-* Tick function that calls UActorComponent::ConditionalTick
-**/
+/** Tick function that calls UActorComponent::ConditionalTick. */
 USTRUCT()
 struct FActorComponentTickFunction : public FTickFunction
 {
@@ -574,33 +616,34 @@ struct FActorComponentTickFunction : public FTickFunction
 #if UE_WITH_REMOTE_OBJECT_HANDLE
 	TObjectPtr<class UActorComponent> Target;
 #else
-	/**  AActor  component that is the target of this tick **/
+	/** Actor component that is the target of this tick */
 	class UActorComponent*	Target;
 #endif
 
 	/** 
-		* Abstract function actually execute the tick. 
-		* @param DeltaTime - frame time to advance, in seconds
-		* @param TickType - kind of tick for this frame
-		* @param CurrentThread - thread we are executing on, useful to pass along as new tasks are created
-		* @param MyCompletionGraphEvent - completion event for this task. Useful for holding the completetion of this task until certain child tasks are complete.
-	**/
+	 * Override this function to execute an actor component tick
+	 * @param DeltaTime - frame time to advance, in seconds
+	 * @param TickType - kind of tick for this frame
+	 * @param CurrentThread - thread we are executing on, useful to pass along as new tasks are created
+	 * @param MyCompletionGraphEvent - completion event for this task. Useful for delaying the completion of this task until child tasks are complete.
+	 */
 	ENGINE_API virtual void ExecuteTick(float DeltaTime, ELevelTick TickType, ENamedThreads::Type CurrentThread, const FGraphEventRef& MyCompletionGraphEvent) override;
-	/** Abstract function to describe this tick. Used to print messages about illegal cycles in the dependency graph **/
+
+	/** Returns the name of the owning componet, override for other types of ticks. */
 	ENGINE_API virtual FString DiagnosticMessage() override;
+
+	/** Returns the name of the owning component class. */
 	ENGINE_API virtual FName DiagnosticContext(bool bDetailed) override;
 
 	/**
-	 * Conditionally calls ExecuteTickFunc if registered and a bunch of other criteria are met
+	 * Helper function for executing component tick functions using the same conditions as FActorTickFunction.
+	 * This adds a cycle counter scope which is used for profiling.
 	 * @param Target - the actor component we are ticking
 	 * @param bTickInEditor - whether the target wants to tick in the editor
 	 * @param DeltaTime - The time since the last tick.
 	 * @param TickType - Type of tick that we are running
 	 * @param ExecuteTickFunc - the lambda that ultimately calls tick on the actor component
 	 */
-
-	//NOTE: This already creates a UObject stat so don't double count in your own functions
-
 	template <typename ExecuteTickLambda>
 	static void ExecuteTickHelper(UActorComponent* Target, bool bTickInEditor, float DeltaTime, ELevelTick TickType, const ExecuteTickLambda& ExecuteTickFunc);	
 
@@ -610,7 +653,6 @@ struct FActorComponentTickFunction : public FTickFunction
 #endif
 };
 
-
 template<>
 struct TStructOpsTypeTraits<FActorComponentTickFunction> : public TStructOpsTypeTraitsBase2<FActorComponentTickFunction>
 {
diff --git a/Engine/Source/Runtime/Engine/Private/TickTaskManager.cpp b/Engine/Source/Runtime/Engine/Private/TickTaskManager.cpp
--- a/Engine/Source/Runtime/Engine/Private/TickTaskManager.cpp
+++ b/Engine/Source/Runtime/Engine/Private/TickTaskManager.cpp
@@ -1791,7 +1791,7 @@ public:
 		while (TickFunction)
 		{
 			// Note: Timestamp check assumes TickFunction->TickGroup has been evaluated this frame
-			if (bFilterCoolingDown && TickFunction->GetLastTickGameTime() != (TickFunction->bTickEvenWhenPaused ? CurrentUnpausedTime : CurrentTime))
+			if (bFilterCoolingDown && TickFunction->GetLastIntervalTickGameTime() != (TickFunction->bTickEvenWhenPaused ? CurrentUnpausedTime : CurrentTime))
 			{				
 				TickFunction = TickFunction->InternalData->Next;
 				continue;
@@ -2387,7 +2387,7 @@ FTickFunction::FInternalData::FInternalData()
 	, TaskPointer(nullptr)
 	, Next(nullptr)
 	, RelativeTickCooldown(0.f)
-	, LastTickGameTimeSeconds(-1.0)
+	, LastIntervalTickSeconds(-1.0)
 	, TickTaskLevel(nullptr)
 {
 }
@@ -2451,7 +2451,7 @@ void FTickFunction::SetTickFunctionEnable(bool bInEnabled)
 
 		if (TickState == ETickState::Disabled)
 		{
-			InternalData->LastTickGameTimeSeconds = -1.0;
+			InternalData->LastIntervalTickSeconds = -1.0;
 		}
 	}
 	else
@@ -2828,18 +2828,18 @@ float FTickFunction::CalculateDeltaTime(float DeltaTime, const UWorld* TickingWo
 	{
 		// No tick interval. Return the world delta seconds, and make sure to mark that
 		// we're not tracking last-tick-time for this object.
-		InternalData->LastTickGameTimeSeconds = -1.0;
+		InternalData->LastIntervalTickSeconds = -1.0;
 	}
 	else
 	{
 		// We've got a tick interval. Mark last-tick-time. If we already had last-tick-time, return
 		// the time since then; otherwise, return the world delta seconds.
 		const double CurrentWorldTime = (bTickEvenWhenPaused ? TickingWorld->GetUnpausedTimeSeconds() : TickingWorld->GetTimeSeconds());
-		if (InternalData->LastTickGameTimeSeconds >= 0.0)
+		if (InternalData->LastIntervalTickSeconds >= 0.0)
 		{
-			DeltaTime = CurrentWorldTime - InternalData->LastTickGameTimeSeconds;
+			DeltaTime = CurrentWorldTime - InternalData->LastIntervalTickSeconds;
 		}
-		InternalData->LastTickGameTimeSeconds = CurrentWorldTime;
+		InternalData->LastIntervalTickSeconds = CurrentWorldTime;
 	}
 
 	return DeltaTime;
diff --git a/Engine/Source/Runtime/Engine/Public/TickTaskManagerInterface.h b/Engine/Source/Runtime/Engine/Public/TickTaskManagerInterface.h
--- a/Engine/Source/Runtime/Engine/Public/TickTaskManagerInterface.h
+++ b/Engine/Source/Runtime/Engine/Public/TickTaskManagerInterface.h
@@ -13,22 +13,24 @@ class ULevel;
 DECLARE_STATS_GROUP(TEXT("TickGroups"), STATGROUP_TickGroups, STATCAT_Advanced);
 
 /** 
- * Interface for the tick task manager
- **/
+ * Interface for the tick task manager, which is used to queue and execute FTickFunctions.
+ * These functions are called for a world to handle the "tick frame" that runs normal gameplay operations.
+ * The frame is split into multiple tick groups to allow broad coordination between different engine systems.
+ */
 class FTickTaskManagerInterface
 {
 public:
 	virtual ~FTickTaskManagerInterface() = default;
 
-	/** Allocate a new ticking structure for a ULevel **/
+	/** Allocate a new ticking structure to track registered ticks for a level. */
 	virtual FTickTaskLevel* AllocateTickTaskLevel() = 0;
 
-	/** Free a ticking structure for a ULevel **/
+	/** Free a ticking structure used for tracking registered ticks. */
 	virtual void FreeTickTaskLevel(FTickTaskLevel* TickTaskLevel) = 0;
 
 	/**
-	 * Queue all of the ticks for a frame
-	 *
+	 * Queue all of the ticks for one tick frame.
+	 * This initializes execution of all tick functions for the frame and handles prerequisite scheduling.
 	 * @param World	- World currently ticking
 	 * @param DeltaSeconds - time in seconds since last tick
 	 * @param TickType - type of tick (viewports only, time only, etc)
@@ -45,13 +47,13 @@ public:
 	virtual void RunPauseFrame(UWorld* InWorld, float DeltaSeconds, ELevelTick TickType, const TArray<ULevel*>& LevelsToTick) = 0;
 
 	/**
-		* Run a tick group, ticking all actors and components
-		* @param Group - Ticking group to run
-		* @param bBlockTillComplete - if true, do not return until all ticks are complete
-	*/
+	 * Run a tick group, ticking all actors and components registered to execute in that tick group.
+	 * @param Group - Ticking group to run
+	 * @param bBlockTillComplete - if true, do not return until all ticks with the corresponding EndTickGroup have completed
+	 */
 	virtual void RunTickGroup(ETickingGroup Group, bool bBlockTillComplete ) = 0;
 
-	/** Finish a frame of ticks **/
+	/** Finish a frame of ticks for all worlds that called StartFrame. */
 	virtual void EndFrame() = 0;
 
 	/** Dumps all registered tick functions to output device. */
@@ -60,13 +62,6 @@ public:
 	/** Returns a map of enabled ticks, grouped by 'diagnostic context' string, along with count of enabled ticks */
 	virtual void GetEnabledTickFunctionCounts(UWorld* InWorld, TSortedMap<FName, int32, FDefaultAllocator, FNameFastLess>& TickContextToCountMap, int32& EnabledCount, bool bDetailed, bool bFilterCoolingDown=false) = 0;
 
-	/**
-	 * Singleton to retrieve the GLOBAL tick task manager
-	 *
-	 * @return Reference to the global cache tick task manager
-	 */
+	/** Accessor for the singleton global tick task manager. */
 	static ENGINE_API FTickTaskManagerInterface& Get();
-
 };
-
-
