Fix for mismatching static switch input in which one is Material Attribute, This causes mismatch between Shader Parameters used to generate DDC Key and Shader Layout.
Previously, this was fixed silently in the translator by compiling using default properties when meeting such a static switch.
This was breaking some existing materials.
So now, we modify the Shader Parameters manually to handle those existing materials and avoid the mismatch between the Shader Parameters used to generate DDC Key and the ones used to generate the Shader Layout.

[RN] minor | Rendering Architecture
#Synced-CL 52420828

diff --git a/Engine/Source/Runtime/Engine/Private/Materials/MaterialCachedData.cpp b/Engine/Source/Runtime/Engine/Private/Materials/MaterialCachedData.cpp
--- a/Engine/Source/Runtime/Engine/Private/Materials/MaterialCachedData.cpp
+++ b/Engine/Source/Runtime/Engine/Private/Materials/MaterialCachedData.cpp
@@ -417,19 +417,13 @@ void FMaterialCachedExpressionData::AddParameter(const FMaterialParameterInfo& P
 		InsertParameterValue(ParameterMeta, Index, OutReferencedTexture, OutReferencedTextureCollection);
 	}
 
-	// ParameterExpressionsAndValues is only used in case PedanticErrorChecksEnabled=1 currently. Avoid doing operations if disabled.
-	static const IConsoleVariable* CVar = IConsoleManager::Get().FindConsoleVariable(TEXT("r.Material.PedanticErrorChecksEnabled"));
-	if (CVar && !CVar->GetBool())
-	{
-		return;
-	}
-
-	const bool bEqualValues = CheckDuplicateParameterValuesEqual(Index, ParameterMeta);
-	const FString CurrentParamValueText = GetParameterValueAsText(ParameterMeta);
+	FParsedExpressionData ExpressionData;
+	ExpressionData.ParameterIdx = Index;
+	ExpressionData.ParameterMeta = ParameterMeta;
+	ExpressionData.Expression = InExpression;
+	ExpressionData.ParameterName = ParameterInfo.Name;
 
-	FParamExpressionsValues& ParamExpressionsValues = ParameterExpressionsAndValues.FindOrAdd(ParameterInfo.Name);
-	ParamExpressionsValues.ExpressionsValues.AddUnique({ InExpression,  CurrentParamValueText });
-	ParamExpressionsValues.bHasDuplicatesWithDifferentValues |= !bEqualValues;
+	ParsedExpressions.Add(ExpressionData);
 }
 
 void FMaterialCachedExpressionData::UpdateForFunction(const FMaterialCachedExpressionContext& Context, UMaterialFunctionInterface* Function, EMaterialParameterAssociation Association, int32 ParameterIndex)
@@ -539,6 +533,12 @@ void FMaterialCachedExpressionData::UpdateForExpressions(const FMaterialCachedEx
 			// Try add the parameter. If this fails, the parameter is being added twice with different values. Report it as error.
 			AddParameter(ParameterInfo, ParameterMeta, AssetPath, ReferencedTexture, ReferencedTextureCollection, Expression);
 		}
+		else
+		{
+			FParsedExpressionData ExpressionData;
+			ExpressionData.Expression = Expression;
+			ParsedExpressions.Add(ExpressionData);
+		}
 
 
 		if (ReferencedTexture)
@@ -828,13 +828,85 @@ void FMaterialCachedExpressionData::UpdatePropertyConnections(UMaterial& Materia
 	}
 }
 
-void FMaterialCachedExpressionData::AnalyzeMaterial(UMaterial& Material)
+#if WITH_EDITOR
+void FMaterialCachedExpressionData::DoPostAnalyzeChecks(const UMaterial& BaseMat)
+{
+	static const IConsoleVariable* CVar = IConsoleManager::Get().FindConsoleVariable(TEXT("r.Material.PedanticErrorChecksEnabled"));
+	const bool bPedanticErrorChecksEnabled = CVar && CVar->GetBool();
+
+	for (const FParsedExpressionData& ParameterData : ParsedExpressions)
+	{
+		// Legacy materials may have a static switch with MaterialAttributes on one branch
+		// and a non-MaterialAttributes value (e.g. scalar constant) on the other.
+		// During translation, the non-MaterialAttributes branch gets broadcast to ALL material properties, causing
+		// IsMaterialPropertyUsed() to report them as used, which leads to wrong shader parameters.
+		// To fix this, detect the mismatch here and conservatively mark all properties as connected for those legacy materials.
+		{
+			auto MarkAllPropertiesConnectedForMismatchedMaterialAttributesSwitch = [&](const FExpressionInput& InputA, const FExpressionInput& InputB)
+				{
+					const FExpressionInput TracedA = InputA.GetTracedInput();
+					const FExpressionInput TracedB = InputB.GetTracedInput();
+
+					if (!TracedA.IsConnected() || !TracedB.IsConnected())
+					{
+						return;
+					}
+
+					const bool bAIsMaterialAttributes = TracedA.Expression->IsResultMaterialAttributes(TracedA.OutputIndex);
+					const bool bBIsMaterialAttributes = TracedB.Expression->IsResultMaterialAttributes(TracedB.OutputIndex);
+
+					if (bAIsMaterialAttributes != bBIsMaterialAttributes)
+					{
+						for (int32 PropertyIndex = 0; PropertyIndex < MP_MaterialAttributes; ++PropertyIndex)
+						{
+							SetPropertyConnected((EMaterialProperty)PropertyIndex);
+						}
+					}
+				};
+
+			if (UMaterialExpressionStaticSwitch* StaticSwitch = Cast<UMaterialExpressionStaticSwitch>(ParameterData.Expression))
+			{
+				MarkAllPropertiesConnectedForMismatchedMaterialAttributesSwitch(StaticSwitch->A, StaticSwitch->B);
+			}
+			else if (UMaterialExpressionStaticSwitchParameter* StaticSwitchParam = Cast<UMaterialExpressionStaticSwitchParameter>(ParameterData.Expression))
+			{
+				MarkAllPropertiesConnectedForMismatchedMaterialAttributesSwitch(StaticSwitchParam->A, StaticSwitchParam->B);
+			}
+		}
+
+
+		if (bPedanticErrorChecksEnabled)
+		{
+			if (ParameterData.ParameterIdx == INDEX_NONE)
+			{
+				continue;
+			}
+
+			const bool bEqualValues = CheckDuplicateParameterValuesEqual(ParameterData.ParameterIdx, ParameterData.ParameterMeta);
+			const FString CurrentParamValueText = GetParameterValueAsText(ParameterData.ParameterMeta);
+
+			FParamExpressionsValues& ParamExpressionsValues = ParameterExpressionsAndValues.FindOrAdd(ParameterData.ParameterName);
+			ParamExpressionsValues.ExpressionsValues.AddUnique({ ParameterData.Expression,  CurrentParamValueText });
+			ParamExpressionsValues.bHasDuplicatesWithDifferentValues |= !bEqualValues;
+		}
+	}
+}
+#endif
+
+void FMaterialCachedExpressionData::AnalyzeMaterial(UMaterial& Material, const FMaterialLayersFunctions* LayersOverrides)
 {
 
 	UpdatePropertyConnections(Material);
 		
 	FMaterialCachedExpressionContext Context;
-	UpdateForExpressions(Context, Material.GetExpressions(), EMaterialParameterAssociation::GlobalParameter, -1);
+	Context.LayerOverrides = LayersOverrides;
+
+	UpdateForExpressions(Context, Material.GetExpressions(), EMaterialParameterAssociation::GlobalParameter, INDEX_NONE);
+
+#if WITH_EDITOR
+	DoPostAnalyzeChecks(Material);
+	ParsedExpressions.Empty();
+#endif
 }
 
 void FMaterialCachedExpressionData::Validate(const UMaterialInterface& Material)
diff --git a/Engine/Source/Runtime/Engine/Private/Materials/MaterialExpressions.cpp b/Engine/Source/Runtime/Engine/Private/Materials/MaterialExpressions.cpp
--- a/Engine/Source/Runtime/Engine/Private/Materials/MaterialExpressions.cpp
+++ b/Engine/Source/Runtime/Engine/Private/Materials/MaterialExpressions.cpp
@@ -8629,34 +8629,11 @@ EMaterialValueType UMaterialExpressionStaticSwitchParameter::GetInputValueType(i
 	return MCT_Unknown;
 }
 
-
 EMaterialValueType UMaterialExpressionStaticSwitchParameter::GetOutputValueType(int32 OutputIndex)
 {
 	return GetCompositeOutput(OutputIndex);
 }
 
-// Helper: Legacy materials may have a static switch with MaterialAttributes on one branch and a
-// non-MaterialAttributes value (e.g. scalar constant) on the other. The UI now prevents creating
-// such connections, but they can still exist in serialized data. The non-material attributes is 
-// evaluated to a MA instance with that value set to each attribute blindly, causing some attributes 
-// to mistakenly be considered connected, which leads to invalid shader layouts.
-// Instead, when we detect this mismatch, compile the default expression for the current material attribute.
-static int32 CompileStaticSwitchBranch(FExpressionInput& ActiveInput, UMaterialExpression* SwitchExpression, int32 OutputIndex, FMaterialCompiler* Compiler)
-{
-	const FExpressionInput TracedActive = ActiveInput.GetTracedInput();
-	const bool bIsResultMaterialAttributes = SwitchExpression->IsResultMaterialAttributes(OutputIndex);
-	if (bIsResultMaterialAttributes && TracedActive.IsConnected())
-	{
-		const bool bIsTracedNotMA = !TracedActive.Expression->IsResultMaterialAttributes(TracedActive.OutputIndex);
-		const FGuid MaterialAttribute = Compiler->GetMaterialAttribute();
-		if (bIsTracedNotMA && MaterialAttribute.IsValid())
-		{
-			return FMaterialAttributeDefinitionMap::CompileDefaultExpression(Compiler, MaterialAttribute);
-		}
-	}
-	return ActiveInput.Compile(Compiler);
-}
-
 FExpressionInput* UMaterialExpressionStaticSwitchParameter::GetEffectiveInput(class FMaterialCompiler* Compiler)
 {
 	bool bSucceeded;
@@ -8668,6 +8645,7 @@ FExpressionInput* UMaterialExpressionStaticSwitchParameter::GetEffectiveInput(cl
 		Compiler->Errorf(TEXT("Missing A input"));
 		bSucceeded = false;
 	}
+
 	if (!B.GetTracedInput().IsConnected())
 	{
 		Compiler->Errorf(TEXT("Missing B input"));
@@ -8692,13 +8670,13 @@ int32 UMaterialExpressionStaticSwitchParameter::Compile(class FMaterialCompiler*
 		}
 		else
 		{
-			if(DefaultValue)
+			if (DefaultValue)
 			{
-				return CompileStaticSwitchBranch(A, this, OutputIndex, Compiler);
+				return A.Compile(Compiler);
 			}
 			else
 			{
-				return CompileStaticSwitchBranch(B, this, OutputIndex, Compiler);
+				return B.Compile(Compiler);
 			}
 		}
 	}
@@ -8706,7 +8684,7 @@ int32 UMaterialExpressionStaticSwitchParameter::Compile(class FMaterialCompiler*
 	FExpressionInput* EffectiveInput = GetEffectiveInput(Compiler);
 	if (EffectiveInput)
 	{
-		return CompileStaticSwitchBranch(*EffectiveInput, this, OutputIndex, Compiler);
+		return EffectiveInput->Compile(Compiler);			
 	}
 	return INDEX_NONE;
 }
@@ -8897,6 +8875,7 @@ FExpressionInput* UMaterialExpressionStaticSwitch::GetEffectiveInput(class FMate
 			return nullptr;
 		}
 	}
+
 	return bValue ? &A : &B;
 }
 
@@ -8924,8 +8903,7 @@ int32 UMaterialExpressionStaticSwitch::Compile(class FMaterialCompiler* Compiler
 		}
 	}
 
-	FExpressionInput& ActiveInput = bValue ? A : B;
-	return CompileStaticSwitchBranch(ActiveInput, this, OutputIndex, Compiler);
+	return bValue ? A.Compile(Compiler) : B.Compile(Compiler);
 }
 
 void UMaterialExpressionStaticSwitch::GetCaption(TArray<FString>& OutCaptions) const
diff --git a/Engine/Source/Runtime/Engine/Private/Materials/MaterialInstanceConstant.cpp b/Engine/Source/Runtime/Engine/Private/Materials/MaterialInstanceConstant.cpp
--- a/Engine/Source/Runtime/Engine/Private/Materials/MaterialInstanceConstant.cpp
+++ b/Engine/Source/Runtime/Engine/Private/Materials/MaterialInstanceConstant.cpp
@@ -241,11 +241,8 @@ void UMaterialInstanceConstant::UpdateCachedData()
 			FMaterialLayersFunctions MaterialLayers;
 			LocalStaticParameters.GetMaterialLayers(MaterialLayers);
 
-			FMaterialCachedExpressionContext Context;
-			Context.LayerOverrides = &MaterialLayers;
 			LocalCachedExpressionData.Reset(new FMaterialCachedExpressionData());
-			LocalCachedExpressionData->UpdatePropertyConnections(*BaseMaterial);
-			LocalCachedExpressionData->UpdateForExpressions(Context, BaseMaterial->GetExpressions(), GlobalParameter, INDEX_NONE);
+			LocalCachedExpressionData->AnalyzeMaterial(*BaseMaterial, &MaterialLayers);
 
 			ReferencedDefaultTextures = LocalCachedExpressionData->EditorOnlyData->EditorReferencedDefaultTextures;
 		}
diff --git a/Engine/Source/Runtime/Engine/Public/MaterialCachedData.h b/Engine/Source/Runtime/Engine/Public/MaterialCachedData.h
--- a/Engine/Source/Runtime/Engine/Public/MaterialCachedData.h
+++ b/Engine/Source/Runtime/Engine/Public/MaterialCachedData.h
@@ -219,7 +219,7 @@ struct FMaterialCachedExpressionData
 	void UpdatePropertyConnections(UMaterial& Material);
 	void UpdateForFunction(const FMaterialCachedExpressionContext& Context, UMaterialFunctionInterface* Function, EMaterialParameterAssociation Association, int32 ParameterIndex);
 	void UpdateForLayerFunctions(const FMaterialCachedExpressionContext& Context, const FMaterialLayersFunctions& LayerFunctions);
-	void AnalyzeMaterial(UMaterial& Material);
+	void AnalyzeMaterial(UMaterial& Material, const FMaterialLayersFunctions* LayersOverrides = nullptr);
 
 	void Validate(const UMaterialInterface& Material);
 
@@ -303,6 +303,9 @@ struct FMaterialCachedExpressionData
 	void PostSerialize(const FArchive& Ar);
 
 private:
+#if WITH_EDITOR
+	void DoPostAnalyzeChecks(const UMaterial& BaseMat);
+#endif
 	bool CheckDuplicateParameterValuesEqual(const int32 InExistingParamIdx, const FMaterialParameterMetadata& DuplicateParameterMeta);
 	FString GetParameterValueAsText(const FMaterialParameterMetadata& ParameterMeta);
 	void InsertParameterValue(const FMaterialParameterMetadata& ParameterMeta, const int32 Index, UObject*& OutReferencedTexture, UTextureCollection*& OutReferencedTextureCollection);
@@ -460,6 +463,23 @@ public:
 	};
 	/** Stores Parameters to expresions and string values. Can be used for error logging duplicates.*/
 	TMap<FName, FParamExpressionsValues> ParameterExpressionsAndValues;
+
+	struct FParsedExpressionData
+	{
+		UMaterialExpression* Expression = nullptr;
+
+		// These are optional in case the expression is a Parameter
+		int32 ParameterIdx = INDEX_NONE;
+		FName ParameterName;
+		
+		// This could theoretically be gathered from Expression 
+		// but being in a Material Function can override values
+		FMaterialParameterMetadata ParameterMeta;
+	};
+
+	// Used for any post analyze checks
+	// Should be cleared after AnalyzeMaterial
+	TArray<FParsedExpressionData> ParsedExpressions;
 #endif
 };
 
