Fix crash in FSkinnedMeshComponentRecreateRenderStateContext from cl 52426729
#Synced-CL 52380299

diff --git a/Engine/Source/Runtime/Engine/Private/SkinnedAsset.cpp b/Engine/Source/Runtime/Engine/Private/SkinnedAsset.cpp
--- a/Engine/Source/Runtime/Engine/Private/SkinnedAsset.cpp
+++ b/Engine/Source/Runtime/Engine/Private/SkinnedAsset.cpp
@@ -591,7 +591,11 @@ FSkinnedMeshComponentRecreateRenderStateContext::FSkinnedMeshComponentRecreateRe
 	TSet<FSceneInterface*> Scenes;
 
 	FObjectCacheContextScope ObjectCacheScope;
-	for (ISkinnedMeshComponent* Component : (InSkinnedAsset ? ObjectCacheScope.GetContext().GetSkinnedMeshComponents(InSkinnedAsset) : ObjectCacheScope.GetContext().GetSkinnedMeshComponents()))
+	TObjectCacheIterator<ISkinnedMeshComponent> SkinnedMeshComponentsIter = InSkinnedAsset
+		? ObjectCacheScope.GetContext().GetSkinnedMeshComponents(InSkinnedAsset)
+		: ObjectCacheScope.GetContext().GetSkinnedMeshComponents();
+
+	for (ISkinnedMeshComponent* Component : SkinnedMeshComponentsIter)
 	{
 		IPrimitiveComponent* PrimitiveComponent = Component->GetPrimitiveComponentInterface();
 
