Match mangled ICH name against SCS DefaultSceneRoot name

#rnx
#Synced-CL 52432244

diff --git a/Engine/Source/Runtime/CoreUObject/Private/UObject/UObjectArchetype.cpp b/Engine/Source/Runtime/CoreUObject/Private/UObject/UObjectArchetype.cpp
--- a/Engine/Source/Runtime/CoreUObject/Private/UObject/UObjectArchetype.cpp
+++ b/Engine/Source/Runtime/CoreUObject/Private/UObject/UObjectArchetype.cpp
@@ -149,6 +149,17 @@ FORCENOINLINE UObject* FindArchetypeFromRequiredInfoImpl(const UClass* Class, co
 			{
 				Result = ArchetypeToSearch->GetClass()->FindArchetype(Class, Name);
 			}
+
+			if(!Result && (ObjectFlags & RF_ArchetypeObject))
+			{
+				if(const UClass* OuterClass = Cast<UClass>(Outer))
+				{
+					if(const UClass* SuperClass = OuterClass->GetSuperClass())
+					{
+						Result = SuperClass->FindArchetype(Class, Name);
+					}
+				}
+			}
 		}
 
 		if (!Result)
diff --git a/Engine/Source/Runtime/Engine/Private/BlueprintGeneratedClass.cpp b/Engine/Source/Runtime/Engine/Private/BlueprintGeneratedClass.cpp
--- a/Engine/Source/Runtime/Engine/Private/BlueprintGeneratedClass.cpp
+++ b/Engine/Source/Runtime/Engine/Private/BlueprintGeneratedClass.cpp
@@ -1572,6 +1572,37 @@ UObject* UBlueprintGeneratedClass::FindArchetype(const UClass* ArchetypeClass, c
 				Class = nullptr;
 			}
 		}
+		
+		if(Archetype == nullptr)
+		{
+			// We failed to find an archetype by name - but the ICH mangles
+			// the DefaultSceneRoot name *and* tries to cull unused templates
+			// so do a relaxed compare to find the real default scene root:
+			static const FName DefaultSceneRootSCS("DefaultSceneRoot_GEN_VARIABLE");
+			static const FName ICHDefaultSceneRoot("ICH-DefaultSceneRoot_GEN_VARIABLE");
+			if(ArchetypeName == DefaultSceneRootSCS || ArchetypeName == ICHDefaultSceneRoot)
+			{
+				Class = const_cast<UBlueprintGeneratedClass*>(this);
+				while (Class && !Archetype)
+				{
+					// look for an SCS template:
+					if(const USimpleConstructionScript* SCS = Class->SimpleConstructionScript)
+					{
+						if(const USCS_Node* Node = SCS->GetDefaultSceneRootNode())
+						{
+							// Assets may have orphaned default scene root nodes - only honor
+							// them if they still have a spot in the root nodes list:
+							if(SCS->GetRootNodes().Contains(Node))
+							{
+								Archetype = Node->ComponentTemplate;
+							}
+						}
+					}
+
+					Class = Cast<UBlueprintGeneratedClass>(Class->GetSuperClass());
+				}
+			}
+		}
 	}
 
 	return Archetype;
