Add more Insights CPU profiling timing scopes for more detail on editor boot sequence

#Synced-CL 52429873

diff --git a/Engine/Plugins/Experimental/EditorDataStorageFeatures/Source/TedsSettings/Private/TedsSettingsManager.cpp b/Engine/Plugins/Experimental/EditorDataStorageFeatures/Source/TedsSettings/Private/TedsSettingsManager.cpp
--- a/Engine/Plugins/Experimental/EditorDataStorageFeatures/Source/TedsSettings/Private/TedsSettingsManager.cpp
+++ b/Engine/Plugins/Experimental/EditorDataStorageFeatures/Source/TedsSettings/Private/TedsSettingsManager.cpp
@@ -71,6 +71,8 @@ FTedsSettingsManager::FTedsSettingsManager()
 
 void FTedsSettingsManager::Initialize()
 {
+	TRACE_CPUPROFILER_EVENT_SCOPE(FTedsSettingsManager::Initialize);
+
 	if (!bIsInitialized)
 	{
 		using namespace UE::Editor::DataStorage;
diff --git a/Engine/Plugins/UbaController/Source/Private/UbaControllerModule.cpp b/Engine/Plugins/UbaController/Source/Private/UbaControllerModule.cpp
--- a/Engine/Plugins/UbaController/Source/Private/UbaControllerModule.cpp
+++ b/Engine/Plugins/UbaController/Source/Private/UbaControllerModule.cpp
@@ -207,6 +207,8 @@ void FUbaControllerModule::ShutdownModule()
 
 void FUbaControllerModule::InitializeController()
 {
+	TRACE_CPUPROFILER_EVENT_SCOPE(FUbaControllerModule::InitializeController);
+
 	// We should never Initialize the controller twice
 	if (ensureAlwaysMsgf(!bControllerInitialized, TEXT("Multiple initialization of UBA controller!")))
 	{
diff --git a/Engine/Source/Developer/DesktopPlatform/Private/TargetReceipt.cpp b/Engine/Source/Developer/DesktopPlatform/Private/TargetReceipt.cpp
--- a/Engine/Source/Developer/DesktopPlatform/Private/TargetReceipt.cpp
+++ b/Engine/Source/Developer/DesktopPlatform/Private/TargetReceipt.cpp
@@ -25,6 +25,8 @@ bool TryParseBuildVersion(const FJsonObject& Object, FBuildVersion& OutVersion)
 
 bool FTargetReceipt::Read(const FString& FileName, bool bExpandVariables)
 {
+	TRACE_CPUPROFILER_EVENT_SCOPE(FTargetReceipt::Read);
+
 	// Read the file from disk
 	FString FileContents;
 	if (!FFileHelper::LoadFileToString(FileContents, *FileName))
diff --git a/Engine/Source/Developer/ShaderCompilerCommon/Private/DXCWrapper.cpp b/Engine/Source/Developer/ShaderCompilerCommon/Private/DXCWrapper.cpp
--- a/Engine/Source/Developer/ShaderCompilerCommon/Private/DXCWrapper.cpp
+++ b/Engine/Source/Developer/ShaderCompilerCommon/Private/DXCWrapper.cpp
@@ -93,6 +93,8 @@ FDllHandle::~FDllHandle()
 
 FDxcModuleWrapper::FDxcModuleWrapper()
 {
+	TRACE_CPUPROFILER_EVENT_SCOPE(FDxcModuleWrapper::FDxcModuleWrapper);
+
 	if (!GDxcHandle.IsValid())
 	{
 		GDxilHandle = new FDllHandle(TEXT("dxil.dll"));
diff --git a/Engine/Source/Developer/ToolMenus/Private/ToolMenus.cpp b/Engine/Source/Developer/ToolMenus/Private/ToolMenus.cpp
--- a/Engine/Source/Developer/ToolMenus/Private/ToolMenus.cpp
+++ b/Engine/Source/Developer/ToolMenus/Private/ToolMenus.cpp
@@ -2173,6 +2173,8 @@ void UToolMenus::PopulateToolBarBuilderWithEntry(
 
 void UToolMenus::PopulateToolBarBuilder(FToolBarBuilder& ToolBarBuilder, UToolMenu* MenuData)
 {
+	TRACE_CPUPROFILER_EVENT_SCOPE(UToolMenus::PopulateToolBarBuilder);
+
 	if (GetEditMenusMode() && !MenuData->IsEditing() && EditMenuDelegate.IsBound())
 	{
 		TWeakObjectPtr<UToolMenu> WeakMenuPtr = MenuData;
diff --git a/Engine/Source/Editor/GameProjectGeneration/Private/GameProjectGenerationModule.cpp b/Engine/Source/Editor/GameProjectGeneration/Private/GameProjectGenerationModule.cpp
--- a/Engine/Source/Editor/GameProjectGeneration/Private/GameProjectGenerationModule.cpp
+++ b/Engine/Source/Editor/GameProjectGeneration/Private/GameProjectGenerationModule.cpp
@@ -177,6 +177,8 @@ bool FGameProjectGenerationModule::OpenProject(const FString& InProjectFile, FTe
 
 void FGameProjectGenerationModule::LoadTemplateCategories()
 {
+	TRACE_CPUPROFILER_EVENT_SCOPE(FGameProjectGenerationModule::LoadTemplateCategories);
+
 	// Now discover and all data driven templates
 	TArray<FString> TemplateRootFolders;
 
diff --git a/Engine/Source/Editor/GameProjectGeneration/Private/SProjectBrowser.cpp b/Engine/Source/Editor/GameProjectGeneration/Private/SProjectBrowser.cpp
--- a/Engine/Source/Editor/GameProjectGeneration/Private/SProjectBrowser.cpp
+++ b/Engine/Source/Editor/GameProjectGeneration/Private/SProjectBrowser.cpp
@@ -366,6 +366,8 @@ SProjectBrowser::SProjectBrowser()
 
 void SProjectBrowser::Construct( const FArguments& InArgs )
 {
+	TRACE_CPUPROFILER_EVENT_SCOPE(SProjectBrowser::Construct);
+
 	// Prepare the projects box
 	ProjectsBox = SNew(SVerticalBox);
 
@@ -806,6 +808,8 @@ static TSharedPtr<FProjectItem> CreateProjectItem(const FString& ProjectFilename
 
 FReply SProjectBrowser::FindProjects()
 {
+	TRACE_CPUPROFILER_EVENT_SCOPE(SProjectBrowser::FindProjects);
+
 	const FString LastSelectedProjectFile = CurrentSelectedProjectPath;
 
 	ProjectItemsSource.Empty();
diff --git a/Engine/Source/Editor/LevelEditor/Private/SLevelViewport.cpp b/Engine/Source/Editor/LevelEditor/Private/SLevelViewport.cpp
--- a/Engine/Source/Editor/LevelEditor/Private/SLevelViewport.cpp
+++ b/Engine/Source/Editor/LevelEditor/Private/SLevelViewport.cpp
@@ -299,6 +299,8 @@ bool SLevelViewport::IsInForegroundTab() const
 
 void SLevelViewport::Construct(const FArguments& InArgs, const FAssetEditorViewportConstructionArgs& InConstructionArguments)
 {
+	TRACE_CPUPROFILER_EVENT_SCOPE(SLevelViewport::Construct);
+
 	GetMutableDefault<ULevelEditorViewportSettings>()->OnSettingChanged().AddRaw(this, &SLevelViewport::HandleViewportSettingChanged);
 
 	ParentLayout = StaticCastSharedPtr<FLevelViewportLayout>(InConstructionArguments.ParentLayout);
diff --git a/Engine/Source/Editor/TurnkeySupport/Private/TurnkeySupportModule.cpp b/Engine/Source/Editor/TurnkeySupport/Private/TurnkeySupportModule.cpp
--- a/Engine/Source/Editor/TurnkeySupport/Private/TurnkeySupportModule.cpp
+++ b/Engine/Source/Editor/TurnkeySupport/Private/TurnkeySupportModule.cpp
@@ -3489,6 +3489,8 @@ static constexpr bool bDeleteTurnkeyProcessOnCompletion = (PLATFORM_WINDOWS);
 
 void FTurnkeySupportModule::UpdateSdkInfo()
 {
+	TRACE_CPUPROFILER_EVENT_SCOPE(FTurnkeySupportModule::UpdateSdkInfo);
+
 	// make sure all known platforms are in the map
 	if (PerPlatformSdkInfo.Num() == 0)
 	{
diff --git a/Engine/Source/Editor/UnrealEd/Private/AssetEditorViewportLayout.cpp b/Engine/Source/Editor/UnrealEd/Private/AssetEditorViewportLayout.cpp
--- a/Engine/Source/Editor/UnrealEd/Private/AssetEditorViewportLayout.cpp
+++ b/Engine/Source/Editor/UnrealEd/Private/AssetEditorViewportLayout.cpp
@@ -187,6 +187,8 @@ const FName FAssetEditorViewportLayout::GetActivePaneConfigurationTypeName() con
 
 TSharedRef<SWidget> FAssetEditorViewportLayout::BuildViewportLayout(TSharedPtr<SDockTab> InParentDockTab, TSharedPtr<FEditorViewportTabContent> InParentTab, const FString& LayoutString)
 {
+	TRACE_CPUPROFILER_EVENT_SCOPE(FAssetEditorViewportLayout::BuildViewportLayout);
+
 	// We don't support reconfiguring an existing layout object, as this makes handling of transitions
 	// particularly difficult.  Instead just destroy the old layout and create a new layout object.
 	check(!ParentTab.IsValid());
diff --git a/Engine/Source/Editor/UnrealEd/Private/EditorEngine.cpp b/Engine/Source/Editor/UnrealEd/Private/EditorEngine.cpp
--- a/Engine/Source/Editor/UnrealEd/Private/EditorEngine.cpp
+++ b/Engine/Source/Editor/UnrealEd/Private/EditorEngine.cpp
@@ -2754,6 +2754,8 @@ void UEditorEngine::PostEditChangeProperty(FPropertyChangedEvent& PropertyChange
 
 void UEditorEngine::Cleanse( bool ClearSelection, bool Redraw, const FText& TransReset, bool bTransReset )
 {
+	TRACE_CPUPROFILER_EVENT_SCOPE(UEditorEngine::Cleanse);
+
 	check( !TransReset.IsEmpty() );
 
 	if (GIsRunning || IsRunningCommandlet())
diff --git a/Engine/Source/Editor/UnrealEd/Private/EditorServer.cpp b/Engine/Source/Editor/UnrealEd/Private/EditorServer.cpp
--- a/Engine/Source/Editor/UnrealEd/Private/EditorServer.cpp
+++ b/Engine/Source/Editor/UnrealEd/Private/EditorServer.cpp
@@ -1955,6 +1955,8 @@ void UEditorEngine::CheckForWorldGCLeaks( UWorld* NewWorld, UPackage* WorldPacka
 
 void UEditorEngine::EditorDestroyWorld( FWorldContext & Context, const FText& CleanseText, UWorld* NewWorld )
 {
+	TRACE_CPUPROFILER_EVENT_SCOPE(UEditorEngine::EditorDestroyWorld);
+
 	if( FModuleManager::Get().IsModuleLoaded("LevelEditor") )
 	{
 		FLevelEditorModule& LevelEditor = FModuleManager::GetModuleChecked<FLevelEditorModule>("LevelEditor");
diff --git a/Engine/Source/Runtime/ApplicationCore/Private/Windows/WindowsApplication.cpp b/Engine/Source/Runtime/ApplicationCore/Private/Windows/WindowsApplication.cpp
--- a/Engine/Source/Runtime/ApplicationCore/Private/Windows/WindowsApplication.cpp
+++ b/Engine/Source/Runtime/ApplicationCore/Private/Windows/WindowsApplication.cpp
@@ -1221,6 +1221,8 @@ FWindowsApplication::FWindowsApplication( const HINSTANCE HInstance, const HICON
 	, ClipCursorRect()
 	, RepeatTime{}
 {
+	TRACE_CPUPROFILER_EVENT_SCOPE(FWindowsApplication::FWindowsApplication);
+
 	FMemory::Memzero(ModifierKeyState, EModifierKey::Count);
 
 	// Disable the process from being showing "ghosted" not responding messages during slow tasks
diff --git a/Engine/Source/Runtime/ApplicationCore/Private/Windows/WindowsWindow.cpp b/Engine/Source/Runtime/ApplicationCore/Private/Windows/WindowsWindow.cpp
--- a/Engine/Source/Runtime/ApplicationCore/Private/Windows/WindowsWindow.cpp
+++ b/Engine/Source/Runtime/ApplicationCore/Private/Windows/WindowsWindow.cpp
@@ -52,6 +52,8 @@ TSharedRef< FWindowsWindow > FWindowsWindow::Make()
 
 void FWindowsWindow::Initialize( FWindowsApplication* const Application, const TSharedRef< FGenericWindowDefinition >& InDefinition, HINSTANCE InHInstance, const TSharedPtr< FWindowsWindow >& InParent, const bool bShowImmediately )
 {
+	TRACE_CPUPROFILER_EVENT_SCOPE(FWindowsWindow::Initialize);
+
 	Definition = InDefinition;
 	OwningApplication = Application;
 
@@ -495,6 +497,8 @@ void FWindowsWindow::AdjustWindowRegion( int32 Width, int32 Height )
 
 void FWindowsWindow::ReshapeWindow( int32 NewX, int32 NewY, int32 NewWidth, int32 NewHeight )
 {
+	TRACE_CPUPROFILER_EVENT_SCOPE(FWindowsWindow::ReshapeWindow);
+
 	WINDOWINFO WindowInfo;
 	FMemory::Memzero( WindowInfo );
 	WindowInfo.cbSize = sizeof( WindowInfo );
diff --git a/Engine/Source/Runtime/Core/Private/Internationalization/Culture.cpp b/Engine/Source/Runtime/Core/Private/Internationalization/Culture.cpp
--- a/Engine/Source/Runtime/Core/Private/Internationalization/Culture.cpp
+++ b/Engine/Source/Runtime/Core/Private/Internationalization/Culture.cpp
@@ -104,6 +104,8 @@ FCulture::FCulture(TUniquePtr<FCultureImplementation>&& InImplementation)
 	, CachedVariant(Implementation->GetVariant())
 	, CachedIsRightToLeft(Implementation->IsRightToLeft())
 {
+	TRACE_CPUPROFILER_EVENT_SCOPE(FCulture::FCulture);
+
 	RefreshCultureDisplayNames(TArray<FString>()); // The display name for the current language will be updated by a post-construct call to RefreshCultureDisplayNames with the correct language data
 }
 
diff --git a/Engine/Source/Runtime/Core/Private/Internationalization/ICUInternationalization.cpp b/Engine/Source/Runtime/Core/Private/Internationalization/ICUInternationalization.cpp
--- a/Engine/Source/Runtime/Core/Private/Internationalization/ICUInternationalization.cpp
+++ b/Engine/Source/Runtime/Core/Private/Internationalization/ICUInternationalization.cpp
@@ -659,6 +659,8 @@ FCulturePtr FICUInternationalization::FindOrMakeCulture(const FString& Name, con
 
 FCulturePtr FICUInternationalization::FindOrMakeCanonizedCulture(const FString& Name, const EAllowDefaultCultureFallback AllowDefaultFallback)
 {
+	TRACE_CPUPROFILER_EVENT_SCOPE(FICUInternationalization::FindOrMakeCanonizedCulture);
+
 	// Find the cached culture.
 	{
 		FScopeLock Lock(&CachedCulturesCS);
diff --git a/Engine/Source/Runtime/Core/Private/Misc/ConfigCacheIni.cpp b/Engine/Source/Runtime/Core/Private/Misc/ConfigCacheIni.cpp
--- a/Engine/Source/Runtime/Core/Private/Misc/ConfigCacheIni.cpp
+++ b/Engine/Source/Runtime/Core/Private/Misc/ConfigCacheIni.cpp
@@ -7436,6 +7436,8 @@ void FConfigCacheIni::PreInitializePlatformPlugins()
 
 FConfigCacheIni* FConfigCacheIni::ForPlatform(FName PlatformName)
 {
+	TRACE_CPUPROFILER_EVENT_SCOPE(FConfigCacheIni::ForPlatform);
+
 #if ALLOW_OTHER_PLATFORM_CONFIG
 	check(GConfig != nullptr && GConfig->bIsReadyForUse);
 
diff --git a/Engine/Source/Runtime/Core/Private/Misc/DataDrivenPlatformInfoRegistry.cpp b/Engine/Source/Runtime/Core/Private/Misc/DataDrivenPlatformInfoRegistry.cpp
--- a/Engine/Source/Runtime/Core/Private/Misc/DataDrivenPlatformInfoRegistry.cpp
+++ b/Engine/Source/Runtime/Core/Private/Misc/DataDrivenPlatformInfoRegistry.cpp
@@ -545,6 +545,8 @@ static void LoadDDPIIniSettings(const FConfigFile& IniFile, FDataDrivenPlatformI
 */
 const TMap<FName, FDataDrivenPlatformInfo>& FDataDrivenPlatformInfoRegistry::GetAllPlatformInfos()
 {
+	TRACE_CPUPROFILER_EVENT_SCOPE(FDataDrivenPlatformInfoRegistry::GetAllPlatformInfos);
+
 	static bool bHasSearchedForPlatforms = false;
 
 	// look on disk for special files
diff --git a/Engine/Source/Runtime/Core/Private/Windows/WindowsPlatformMisc.cpp b/Engine/Source/Runtime/Core/Private/Windows/WindowsPlatformMisc.cpp
--- a/Engine/Source/Runtime/Core/Private/Windows/WindowsPlatformMisc.cpp
+++ b/Engine/Source/Runtime/Core/Private/Windows/WindowsPlatformMisc.cpp
@@ -4808,6 +4808,8 @@ ENetworkConnectionType FWindowsPlatformMisc::GetNetworkConnectionType()
 
 void FWindowsPlatformMisc::StartNetworkConnectionMonitoring()
 {
+	TRACE_CPUPROFILER_EVENT_SCOPE(FWindowsPlatformMisc::StartNetworkConnectionMonitoring);
+
 	check(GNetworkListManager == nullptr && GNLMEventListener == nullptr);
 
 	// COM must be initialized on this thread before we can create COM objects.
diff --git a/Engine/Source/Runtime/CoreUObject/Private/UObject/Obj.cpp b/Engine/Source/Runtime/CoreUObject/Private/UObject/Obj.cpp
--- a/Engine/Source/Runtime/CoreUObject/Private/UObject/Obj.cpp
+++ b/Engine/Source/Runtime/CoreUObject/Private/UObject/Obj.cpp
@@ -5850,6 +5850,7 @@ AUTORTFM_DISABLE void StaticExit();
 
 void InitUObject()
 {
+	TRACE_CPUPROFILER_EVENT_SCOPE(InitUObject);
 	LLM_SCOPE(ELLMTag::InitUObject);
 
 	FGCCSyncObject::Create();
diff --git a/Engine/Source/Runtime/Engine/Private/Materials/Material.cpp b/Engine/Source/Runtime/Engine/Private/Materials/Material.cpp
--- a/Engine/Source/Runtime/Engine/Private/Materials/Material.cpp
+++ b/Engine/Source/Runtime/Engine/Private/Materials/Material.cpp
@@ -357,6 +357,8 @@ void FMaterialResource::EndAllowCachingStaticParameterValues()
 
 void FMaterialResource::BuildShaderMapIdOverride(const FBuildShaderMapIdArgs& Args) const
 {
+	TRACE_CPUPROFILER_EVENT_SCOPE(FMaterialResource::BuildShaderMapIdOverride);
+
 	FMaterial::BuildShaderMapIdOverride(Args);
 #if WITH_EDITOR
 	FMaterialShaderMapId& OutId = *Args.OutId;
diff --git a/Engine/Source/Runtime/Engine/Private/Materials/MaterialShader.cpp b/Engine/Source/Runtime/Engine/Private/Materials/MaterialShader.cpp
--- a/Engine/Source/Runtime/Engine/Private/Materials/MaterialShader.cpp
+++ b/Engine/Source/Runtime/Engine/Private/Materials/MaterialShader.cpp
@@ -3037,6 +3037,8 @@ void FMaterialShaderMap::AppendCompileStateDebugInfo(FStringBuilderBase& OutDebu
 
 void FMaterialShaderMap::ProcessCompilationResults(const TArray<FShaderCommonCompileJobPtr>& InCompilationResults, int32& InOutJobIndex, float& TimeBudget)
 {
+	TRACE_CPUPROFILER_EVENT_SCOPE(FMaterialShaderMap::ProcessCompilationResults);
+
 	check(!bCompilationFinalized);
 
 	double StartTime = FPlatformTime::Seconds();
@@ -3151,6 +3153,8 @@ public:
 private:
 	FMaterialShaderMapLayout* FindOrCreateLayout(const EShaderPlatform Platform, const EShaderPermutationFlags Flags, const uint64 ParameterHash, const FMaterialShaderParameters& MaterialParameters)
 	{
+		TRACE_CPUPROFILER_EVENT_SCOPE(FMaterialShaderMapLayoutCache::FindOrCreateLayout);
+
 		FMaterialShaderMapLayout* Layout = nullptr;
 		{
 			FReadScopeLock Locker(LayoutLock);
@@ -3803,6 +3807,8 @@ FMaterialShaderMap::~FMaterialShaderMap()
 #if WITH_EDITOR
 FMaterialShaderMap* FMaterialShaderMap::AcquireFinalizedClone()
 {
+	TRACE_CPUPROFILER_EVENT_SCOPE(FMaterialShaderMap::AcquireFinalizedClone);
+
 	checkSlow(IsInGameThread());
 
 	const FMaterialShaderMapContent* LocalContent = GetContent();
diff --git a/Engine/Source/Runtime/Engine/Private/ShaderCompiler/ShaderCompiler.cpp b/Engine/Source/Runtime/Engine/Private/ShaderCompiler/ShaderCompiler.cpp
--- a/Engine/Source/Runtime/Engine/Private/ShaderCompiler/ShaderCompiler.cpp
+++ b/Engine/Source/Runtime/Engine/Private/ShaderCompiler/ShaderCompiler.cpp
@@ -752,6 +752,8 @@ FShaderCompilingManager::FShaderCompilingManager() :
 	bAllowForIncompleteShaderMaps(false),
 	Notification(MakeUnique<FAsyncCompilationNotification>(GetAssetNameFormat()))
 {
+	TRACE_CPUPROFILER_EVENT_SCOPE(FShaderCompilingManager::FShaderCompilingManager);
+
 	// don't perform any initialization if compiling is not allowed
 	if (!AllowShaderCompiling())
 	{
diff --git a/Engine/Source/Runtime/Engine/Private/ShaderCompiler/ShaderCompilerEditor.cpp b/Engine/Source/Runtime/Engine/Private/ShaderCompiler/ShaderCompilerEditor.cpp
--- a/Engine/Source/Runtime/Engine/Private/ShaderCompiler/ShaderCompilerEditor.cpp
+++ b/Engine/Source/Runtime/Engine/Private/ShaderCompiler/ShaderCompilerEditor.cpp
@@ -209,6 +209,8 @@ FString GetGlobalShaderMapKeyString(const FGlobalShaderMapId& ShaderMapId, EShad
 /** Creates a string key for the derived data cache entry for the global shader map. */
 UE::DerivedData::FCacheKey GetGlobalShaderMapKey(const FGlobalShaderMapId& ShaderMapId, EShaderPlatform Platform, const ITargetPlatform* TargetPlatform, TArray<FShaderTypeDependency> const& Dependencies)
 {
+	TRACE_CPUPROFILER_EVENT_SCOPE(GetGlobalShaderMapKey);
+
 	const FString DataKey = GetGlobalShaderMapKeyString(ShaderMapId, Platform, Dependencies);
 	static const UE::DerivedData::FCacheBucket Bucket(ANSITEXTVIEW("GlobalShaderMap"), TEXTVIEW("GlobalShader"));
 	return {Bucket, FIoHash::HashBuffer(MakeMemoryView(FTCHARToUTF8(DataKey)))};
diff --git a/Engine/Source/Runtime/Engine/Private/StaticMesh.cpp b/Engine/Source/Runtime/Engine/Private/StaticMesh.cpp
--- a/Engine/Source/Runtime/Engine/Private/StaticMesh.cpp
+++ b/Engine/Source/Runtime/Engine/Private/StaticMesh.cpp
@@ -3192,6 +3192,8 @@ void FStaticMeshRenderData::SyncUVChannelData(const TArray<FStaticMaterial>& Obj
 ------------------------------------------------------------------------------*/
 void FStaticMeshLODSettings::Initialize(const ITargetPlatformSettings* TargetPlatformSettings)
 {
+	TRACE_CPUPROFILER_EVENT_SCOPE(FStaticMeshLODSettings::Initialize);
+
 	check(!Groups.Num());
 	// Ensure there is a default LOD group.
 	Groups.FindOrAdd(NAME_None);
diff --git a/Engine/Source/Runtime/Landscape/Private/LandscapeEditResourcesSubsystem.cpp b/Engine/Source/Runtime/Landscape/Private/LandscapeEditResourcesSubsystem.cpp
--- a/Engine/Source/Runtime/Landscape/Private/LandscapeEditResourcesSubsystem.cpp
+++ b/Engine/Source/Runtime/Landscape/Private/LandscapeEditResourcesSubsystem.cpp
@@ -404,6 +404,8 @@ ULandscapeEditResourcesSubsystem::ULandscapeEditResourcesSubsystem()
 
 void ULandscapeEditResourcesSubsystem::Initialize(FSubsystemCollectionBase& Collection)
 {
+	TRACE_CPUPROFILER_EVENT_SCOPE(ULandscapeEditResourcesSubsystem::Initialize);
+
 	Super::Initialize(Collection);
 	
 #if WITH_EDITOR  
diff --git a/Engine/Source/Runtime/Landscape/Private/LandscapeUtils.cpp b/Engine/Source/Runtime/Landscape/Private/LandscapeUtils.cpp
--- a/Engine/Source/Runtime/Landscape/Private/LandscapeUtils.cpp
+++ b/Engine/Source/Runtime/Landscape/Private/LandscapeUtils.cpp
@@ -495,6 +495,8 @@ bool FBatchTextureCopy::ProcessTextureCopies()
 
 UMaterialInstance* CreateToolLandscapeMaterialInstanceConstant(UMaterialInterface* BaseMaterial)
 {
+	TRACE_CPUPROFILER_EVENT_SCOPE(CreateToolLandscapeMaterialInstanceConstant);
+
 	UObject* Outer = GetTransientPackage();
 	// Use the base material's name as the base of our MIC to help debug: 
 	FString MICName(FString::Format(TEXT("LandscapeMaterialInstanceConstant_{0}"), { *BaseMaterial->GetName() }));
diff --git a/Engine/Source/Runtime/Projects/Private/ModuleDescriptor.cpp b/Engine/Source/Runtime/Projects/Private/ModuleDescriptor.cpp
--- a/Engine/Source/Runtime/Projects/Private/ModuleDescriptor.cpp
+++ b/Engine/Source/Runtime/Projects/Private/ModuleDescriptor.cpp
@@ -842,6 +842,8 @@ void FModuleDescriptor::UnloadModulesForPhase(ELoadingPhase::Type LoadingPhase,
 #if !IS_MONOLITHIC
 bool FModuleDescriptor::CheckModuleCompatibility(const TArray<FModuleDescriptor>& Modules, TArray<FString>& OutIncompatibleFiles)
 {
+	TRACE_CPUPROFILER_EVENT_SCOPE(FModuleDescriptor::CheckModuleCompatibility);
+
 	FModuleManager& ModuleManager = FModuleManager::Get();
 
 	bool bResult = true;
diff --git a/Engine/Source/Runtime/Projects/Private/PluginManager.cpp b/Engine/Source/Runtime/Projects/Private/PluginManager.cpp
--- a/Engine/Source/Runtime/Projects/Private/PluginManager.cpp
+++ b/Engine/Source/Runtime/Projects/Private/PluginManager.cpp
@@ -2832,6 +2832,8 @@ bool FPluginManager::TryUnloadModulesForPlugin(const FPlugin& Plugin, const ELoa
 
 bool FPluginManager::LoadModulesForEnabledPlugins( const ELoadingPhase::Type LoadingPhase )
 {
+	TRACE_CPUPROFILER_EVENT_SCOPE(FPluginManager::LoadModulesForEnabledPlugins);
+
 	// Figure out which plugins are enabled
 	bool bSuccess = true;
 	if (!ConfigureEnabledPlugins())
diff --git a/Engine/Source/Runtime/RenderCore/Private/GlobalShader.cpp b/Engine/Source/Runtime/RenderCore/Private/GlobalShader.cpp
--- a/Engine/Source/Runtime/RenderCore/Private/GlobalShader.cpp
+++ b/Engine/Source/Runtime/RenderCore/Private/GlobalShader.cpp
@@ -588,6 +588,8 @@ FGlobalShaderMap* GetGlobalShaderMap(EShaderPlatform Platform)
 
 void PrimeGlobalShaderSourceHashCache(EShaderPlatform Platform, const ITargetPlatform* TargetPlatform)
 {
+	TRACE_CPUPROFILER_EVENT_SCOPE(PrimeGlobalShaderSourceHashCache);
+
 #if WITH_EDITOR
 	FPlatformTypeLayoutParameters LayoutParams;
 	LayoutParams.InitializeForPlatform(TargetPlatform);
@@ -638,6 +640,8 @@ FGlobalShaderMapSection* FGlobalShaderMapSection::CreateFromArchive(FArchive& Ar
 #if WITH_EDITOR
 FGlobalShaderMapSection* FGlobalShaderMapSection::CreateFromCache(FShaderCacheLoadContext& Ctx)
 {
+	TRACE_CPUPROFILER_EVENT_SCOPE(FGlobalShaderMapSection::CreateFromCache);
+
 	FGlobalShaderMapSection* Section = new FGlobalShaderMapSection();
 	if (Section->Serialize(Ctx))
 	{
diff --git a/Engine/Source/Runtime/RenderCore/Private/Shader.cpp b/Engine/Source/Runtime/RenderCore/Private/Shader.cpp
--- a/Engine/Source/Runtime/RenderCore/Private/Shader.cpp
+++ b/Engine/Source/Runtime/RenderCore/Private/Shader.cpp
@@ -606,6 +606,8 @@ const FSHAHash& FShaderType::GetSourceHash(EShaderPlatform ShaderPlatform) const
 
 void FShaderType::Initialize(const TMap<FString, TArray<const TCHAR*> >& ShaderFileToUniformBufferVariables)
 {
+	TRACE_CPUPROFILER_EVENT_SCOPE(FShaderType::Initialize);
+
 #if WITH_EDITOR
 	//#todo-rco: Need to call this only when Initializing from a Pipeline once it's removed from the global linked list
 	if (!FPlatformProperties::RequiresCookedData())
diff --git a/Engine/Source/Runtime/RenderCore/Private/ShaderCore.cpp b/Engine/Source/Runtime/RenderCore/Private/ShaderCore.cpp
--- a/Engine/Source/Runtime/RenderCore/Private/ShaderCore.cpp
+++ b/Engine/Source/Runtime/RenderCore/Private/ShaderCore.cpp
@@ -2254,6 +2254,8 @@ static void InternalGetShaderIncludes(const TCHAR* EntryPointVirtualFilePath, co
 */
 void AddShaderSourceFileEntry(TArray<FString>& OutVirtualFilePaths, FString VirtualFilePath, EShaderPlatform ShaderPlatform, const FName* ShaderPlatformName)
 {
+	TRACE_CPUPROFILER_EVENT_SCOPE(AddShaderSourceFileEntry);
+
 	check(CheckVirtualShaderFilePath(VirtualFilePath));
 	if (!OutVirtualFilePaths.Contains(VirtualFilePath))
 	{
@@ -2310,6 +2312,8 @@ void GetAllVirtualShaderSourcePaths(TArray<FString>& OutVirtualFilePaths, EShade
 */
 void VerifyShaderSourceFiles(EShaderPlatform ShaderPlatform)
 {
+	TRACE_CPUPROFILER_EVENT_SCOPE(VerifyShaderSourceFiles);
+
 #if WITH_EDITORONLY_DATA
 	if (!FPlatformProperties::RequiresCookedData() && AllowShaderCompiling())
 	{
@@ -3648,6 +3652,7 @@ void BuildShaderFileToUniformBufferMap(TMap<FString, TArray<const TCHAR*> >& Sha
 
 void InitializeShaderHashCache()
 {
+	TRACE_CPUPROFILER_EVENT_SCOPE(InitializeShaderHashCache);
 	FRWScopeLock ShaderHashAccessLock(GShaderHashAccessRWLock, SLT_Write);
 	GShaderHashCache.Initialize();
 }
diff --git a/Engine/Source/Runtime/RenderCore/Public/ShaderCompilerJobTypes.h b/Engine/Source/Runtime/RenderCore/Public/ShaderCompilerJobTypes.h
--- a/Engine/Source/Runtime/RenderCore/Public/ShaderCompilerJobTypes.h
+++ b/Engine/Source/Runtime/RenderCore/Public/ShaderCompilerJobTypes.h
@@ -10,6 +10,7 @@
 #include "ShaderPreprocessTypes.h"
 #include "VertexFactory.h"
 #include "Templates/Function.h"
+#include "ProfilingDebugging/CpuProfilerTrace.h"
 
 class FShaderCommonCompileJob;
 class FShaderCompileJob;
@@ -471,6 +472,8 @@ inline bool FShaderCommonCompileJob::Equals(const FShaderCommonCompileJob& Rhs)
 
 inline void FShaderCommonCompileJob::Destroy() const
 {
+	TRACE_CPUPROFILER_EVENT_SCOPE(FShaderCommonCompileJob::Destroy);
+
 	switch (Type)
 	{
 	case EShaderCompileJobType::Single: delete static_cast<const FShaderCompileJob*>(this); break;
diff --git a/Engine/Source/Runtime/Slate/Private/Framework/Docking/TabManager.cpp b/Engine/Source/Runtime/Slate/Private/Framework/Docking/TabManager.cpp
--- a/Engine/Source/Runtime/Slate/Private/Framework/Docking/TabManager.cpp
+++ b/Engine/Source/Runtime/Slate/Private/Framework/Docking/TabManager.cpp
@@ -1159,6 +1159,8 @@ void FTabManager::UnregisterAllTabSpawners()
 TSharedPtr<SWidget> FTabManager::RestoreFrom(const TSharedRef<FLayout>& Layout, const TSharedPtr<SWindow>& ParentWindow, const bool bEmbedTitleAreaContent,
 	const EOutputCanBeNullptr RestoreAreaOutputCanBeNullptr)
 {
+	TRACE_CPUPROFILER_EVENT_SCOPE(FTabManager::RestoreFrom);
+
 	ActiveLayoutName = Layout->LayoutName;
 
 	TSharedPtr<SDockingArea> PrimaryDockArea;
@@ -2049,6 +2051,8 @@ TSharedPtr<SDockingArea> FTabManager::RestoreArea(const TSharedRef<FArea>& AreaT
 TSharedPtr<SDockingNode> FTabManager::RestoreArea_Helper(const TSharedRef<FLayoutNode>& LayoutNode, const TSharedPtr<SWindow>& ParentWindow, const bool bEmbedTitleAreaContent,
 	FSidebarTabLists& OutSidebarTabs, const EOutputCanBeNullptr OutputCanBeNullptr, bool bForceOpenWindowIfNeeded)
 {
+	TRACE_CPUPROFILER_EVENT_SCOPE(FTabManager::RestoreArea_Helper);
+
 #if WITH_EDITOR
 	FSlateApplication::FScopedPreventDebuggingMode Scope(LOCTEXT("RestoringTabsDebugScope", "Disabling debug due to being in tab restore, breakpoints in constructors can infinitely stall during restore."));
 #endif
diff --git a/Engine/Source/Runtime/SlateCore/Private/Widgets/SWindow.cpp b/Engine/Source/Runtime/SlateCore/Private/Widgets/SWindow.cpp
--- a/Engine/Source/Runtime/SlateCore/Private/Widgets/SWindow.cpp
+++ b/Engine/Source/Runtime/SlateCore/Private/Widgets/SWindow.cpp
@@ -1426,6 +1426,8 @@ PRAGMA_ENABLE_DEPRECATION_WARNINGS
 /** Make the window visible */
 void SWindow::ShowWindow()
 {
+	TRACE_CPUPROFILER_EVENT_SCOPE(SWindow::ShowWindow);
+
 	// Make sure the viewport is setup for this window
 	if( !bHasEverBeenShown )
 	{
diff --git a/Engine/Source/Runtime/StudioTelemetry/Private/StudioTelemetry.cpp b/Engine/Source/Runtime/StudioTelemetry/Private/StudioTelemetry.cpp
--- a/Engine/Source/Runtime/StudioTelemetry/Private/StudioTelemetry.cpp
+++ b/Engine/Source/Runtime/StudioTelemetry/Private/StudioTelemetry.cpp
@@ -52,6 +52,8 @@ void FStudioTelemetry::SetRecordEventCallback(OnRecordEventCallback Callback )
 
 bool FStudioTelemetry::StartSession()
 {
+	TRACE_CPUPROFILER_EVENT_SCOPE(FStudioTelemetry::StartSession);
+
 	if (IsSessionRunning())
 	{
 		return true;
