Add UWorldPartitionRuntimeHashSet::ContainsRuntimeHash() to allow filtering the display of grids using wp.Runtime.DebugFilterByRuntimeHashGridName

#Synced-CL 0

diff --git a/Engine/Source/Runtime/Engine/Private/WorldPartition/RuntimeHashSet/WorldPartitionRuntimeHashSetDebugDraw.cpp b/Engine/Source/Runtime/Engine/Private/WorldPartition/RuntimeHashSet/WorldPartitionRuntimeHashSetDebugDraw.cpp
--- a/Engine/Source/Runtime/Engine/Private/WorldPartition/RuntimeHashSet/WorldPartitionRuntimeHashSetDebugDraw.cpp
+++ b/Engine/Source/Runtime/Engine/Private/WorldPartition/RuntimeHashSet/WorldPartitionRuntimeHashSetDebugDraw.cpp
@@ -494,4 +494,14 @@ void UWorldPartitionRuntimeHashSet::Draw3D(const TArray<FWorldPartitionStreaming
 		}
 		return true;
 	});
-}
\ No newline at end of file
+}
+
+bool UWorldPartitionRuntimeHashSet::ContainsRuntimeHash(const FString& Name) const
+{
+	const FRuntimePartitionDesc* FoundRuntimePartitionDesc = RuntimePartitions.FindByPredicate([&Name](const FRuntimePartitionDesc& RuntimePartitionDesc)
+		{
+			return RuntimePartitionDesc.Name.ToString().Equals(Name, ESearchCase::IgnoreCase);
+		});
+
+	return FoundRuntimePartitionDesc != nullptr;
+}
diff --git a/Engine/Source/Runtime/Engine/Public/WorldPartition/RuntimeHashSet/WorldPartitionRuntimeHashSet.h b/Engine/Source/Runtime/Engine/Public/WorldPartition/RuntimeHashSet/WorldPartitionRuntimeHashSet.h
--- a/Engine/Source/Runtime/Engine/Public/WorldPartition/RuntimeHashSet/WorldPartitionRuntimeHashSet.h
+++ b/Engine/Source/Runtime/Engine/Public/WorldPartition/RuntimeHashSet/WorldPartitionRuntimeHashSet.h
@@ -175,6 +175,7 @@ class UWorldPartitionRuntimeHashSet final : public UWorldPartitionRuntimeHash
 public:
 	ENGINE_API virtual bool Draw2D(FWorldPartitionDraw2DContext& DrawContext) const override;
 	ENGINE_API virtual void Draw3D(const TArray<FWorldPartitionStreamingSource>& Sources) const override;
+	ENGINE_API virtual bool ContainsRuntimeHash(const FString& Name) const override;
 	ENGINE_API virtual const FGuid* GetStandaloneHLODActorSourceCellOverride(const FGuid& InActorGuid) const override;
 	ENGINE_API virtual const FGuid* GetCustomHLODActorSourceCellOverride(const FGuid& InActorGuid) const override;
 
