DSTM: Character Movement should be speculative as an optimization to reduce abort/retry loops on the physics queries.

#Synced-CL 52448335

diff --git a/Engine/Source/Runtime/Engine/Private/Components/CharacterMovementComponent.cpp b/Engine/Source/Runtime/Engine/Private/Components/CharacterMovementComponent.cpp
--- a/Engine/Source/Runtime/Engine/Private/Components/CharacterMovementComponent.cpp
+++ b/Engine/Source/Runtime/Engine/Private/Components/CharacterMovementComponent.cpp
@@ -37,6 +37,10 @@
 #include "PhysicsEngine/PhysicsObjectExternalInterface.h"
 #include "Runtime/Experimental/Chaos/Private/Chaos/PhysicsObjectInternal.h"
 
+#if UE_WITH_REMOTE_OBJECT_HANDLE
+#include "UObject/RemoteExecutor.h"
+#endif
+
 #include UE_INLINE_GENERATED_CPP_BY_NAME(CharacterMovementComponent)
 
 CSV_DEFINE_CATEGORY(CharacterMovement, true);
@@ -2768,6 +2772,12 @@ void UCharacterMovementComponent::PerformMovement(float DeltaSeconds)
 {
 	SCOPE_CYCLE_COUNTER(STAT_CharacterMovementPerformMovement);
 
+#if UE_WITH_REMOTE_OBJECT_HANDLE && UE_AUTORTFM
+	// We will speculate during the movement and assume we only hit static geometry
+	UE::RemoteExecutor::FSpeculationExecutionScope AssumeNoCrossServerHits;
+#endif
+
+
 	const UWorld* MyWorld = GetWorld();
 	if (!HasValidData() || MyWorld == nullptr)
 	{
