Add in place pose blending function to AnimationRuntime (BlendTwoPosesTogetherInPlace). Use 2 poses instead of 3 to blend when possible in BlendListBase, TwoWayBlend and Mirror

#Synced-CL 52419093

diff --git a/Engine/Source/Runtime/AnimGraphRuntime/Private/AnimNodes/AnimNode_BlendListBase.cpp b/Engine/Source/Runtime/AnimGraphRuntime/Private/AnimNodes/AnimNode_BlendListBase.cpp
--- a/Engine/Source/Runtime/AnimGraphRuntime/Private/AnimNodes/AnimNode_BlendListBase.cpp
+++ b/Engine/Source/Runtime/AnimGraphRuntime/Private/AnimNodes/AnimNode_BlendListBase.cpp
@@ -316,6 +316,22 @@ void FAnimNode_BlendListBase::Evaluate_AnyThread(FPoseContext& Output)
 			// Single full weight pose - pass-through fast common case
 			BlendPose[PosesToEvaluate[0]].Evaluate(Output);
 		}
+		else if (NumPoses == 2 && CurrentBlendProfile == nullptr && FMath::IsNearlyEqual(BlendWeights[PosesToEvaluate[0]] + BlendWeights[PosesToEvaluate[1]], 1.f))
+		{
+			// Fast path for two active poses without blend profile: evaluate first child directly into Output to avoid allocating a third pose.
+			// Requires weights summing to 1.0 since BlendTwoPosesTogetherInPlace derives the second weight as (1 - first weight).
+			const int32 FirstPoseIndex = PosesToEvaluate[0];
+			const int32 SecondPoseIndex = PosesToEvaluate[1];
+
+			BlendPose[FirstPoseIndex].Evaluate(Output);
+
+			FPoseContext Pose2(Output);
+			BlendPose[SecondPoseIndex].Evaluate(Pose2);
+
+			FAnimationPoseData InOutPoseData(Output);
+			const FAnimationPoseData PoseTwoData(Pose2);
+			FAnimationRuntime::BlendTwoPosesTogetherInPlace(InOutPoseData, PoseTwoData, BlendWeights[FirstPoseIndex]);
+		}
 		else
 		{
 			// Scratch arrays for evaluation, stack allocated
diff --git a/Engine/Source/Runtime/AnimGraphRuntime/Private/AnimNodes/AnimNode_Mirror.cpp b/Engine/Source/Runtime/AnimGraphRuntime/Private/AnimNodes/AnimNode_Mirror.cpp
--- a/Engine/Source/Runtime/AnimGraphRuntime/Private/AnimNodes/AnimNode_Mirror.cpp
+++ b/Engine/Source/Runtime/AnimGraphRuntime/Private/AnimNodes/AnimNode_Mirror.cpp
@@ -267,19 +267,18 @@ void FAnimNode_MirrorBase::Evaluate_AnyThread(FPoseContext& Output)
 	}
 	else if (FAnimWeight::IsRelevant(BlendAlpha))
 	{
-		FPoseContext UnmirroredPose(Output);
-		Source.Evaluate(UnmirroredPose);
+		// Evaluate directly into Output (unmirrored), copy into MirroredPose and mirror it in place
+		Source.Evaluate(Output);
 
 		FPoseContext MirroredPose(Output);
-		MirroredPose = UnmirroredPose;
-		// @TODO: Make BlendXPosesTogether that can blend in-place to save 3rd pose + unneeded pose copy.
-		// This is a common pattern i.e. FAnimNode_TwoWayBlend.
+		MirroredPose.Pose.CopyBonesFrom(Output.Pose);
+		MirroredPose.Curve.CopyFrom(Output.Curve);
+		MirroredPose.CustomAttributes.CopyFrom(Output.CustomAttributes);
 		MirrorPoseInPlace(MirroredPose);
-		
-		FAnimationPoseData OutAnimationPoseData(Output);
-		const FAnimationPoseData UnmirroredPoseData(UnmirroredPose);
+
+		FAnimationPoseData InOutPoseData(Output);
 		const FAnimationPoseData MirroredPoseData(MirroredPose);
-		FAnimationRuntime::BlendTwoPosesTogether(UnmirroredPoseData, MirroredPoseData, (1.0f - BlendAlpha), OutAnimationPoseData);
+		FAnimationRuntime::BlendTwoPosesTogetherInPlace(InOutPoseData, MirroredPoseData, (1.0f - BlendAlpha));
 	}
 	else
 	{
diff --git a/Engine/Source/Runtime/AnimGraphRuntime/Private/AnimNodes/AnimNode_TwoWayBlend.cpp b/Engine/Source/Runtime/AnimGraphRuntime/Private/AnimNodes/AnimNode_TwoWayBlend.cpp
--- a/Engine/Source/Runtime/AnimGraphRuntime/Private/AnimNodes/AnimNode_TwoWayBlend.cpp
+++ b/Engine/Source/Runtime/AnimGraphRuntime/Private/AnimNodes/AnimNode_TwoWayBlend.cpp
@@ -128,16 +128,15 @@ void FAnimNode_TwoWayBlend::Evaluate_AnyThread(FPoseContext& Output)
 	{
 		if (bAIsRelevant)
 		{
-			FPoseContext Pose1(Output);
-			FPoseContext Pose2(Output);
+			// Evaluate A directly into Output to avoid allocating a third pose
+			A.Evaluate(Output);
 
-			A.Evaluate(Pose1);
+			FPoseContext Pose2(Output);
 			B.Evaluate(Pose2);
 
-			FAnimationPoseData BlendedAnimationPoseData(Output);
-			const FAnimationPoseData AnimationPoseOneData(Pose1);
-			const FAnimationPoseData AnimationPoseTwoData(Pose2);
-			FAnimationRuntime::BlendTwoPosesTogether(AnimationPoseOneData, AnimationPoseTwoData, (1.0f - InternalBlendAlpha), BlendedAnimationPoseData);
+			FAnimationPoseData InOutPoseData(Output);
+			const FAnimationPoseData PoseTwoData(Pose2);
+			FAnimationRuntime::BlendTwoPosesTogetherInPlace(InOutPoseData, PoseTwoData, (1.0f - InternalBlendAlpha));
 		}
 		else
 		{
diff --git a/Engine/Source/Runtime/Engine/Private/Animation/AnimationRuntime.cpp b/Engine/Source/Runtime/Engine/Private/Animation/AnimationRuntime.cpp
--- a/Engine/Source/Runtime/Engine/Private/Animation/AnimationRuntime.cpp
+++ b/Engine/Source/Runtime/Engine/Private/Animation/AnimationRuntime.cpp
@@ -524,6 +524,32 @@ void FAnimationRuntime::BlendTwoPosesTogether(const FAnimationPoseData& SourcePo
 	UE::Anim::Attributes::BlendAttributes({ SourcePoseOneData.GetAttributes(), SourcePoseTwoData.GetAttributes() }, { WeightOfPoseOne, WeightOfPoseTwo }, { 0, 1 }, OutAttributes);
 }
 
+void FAnimationRuntime::BlendTwoPosesTogetherInPlace(FAnimationPoseData& InOutPoseOneData, const FAnimationPoseData& SourcePoseTwoData, const float WeightOfPoseOne)
+{
+	FCompactPose& InOutPose = InOutPoseOneData.GetPose();
+	FBlendedCurve& InOutCurve = InOutPoseOneData.GetCurve();
+	UE::Anim::FStackAttributeContainer& InOutAttributes = InOutPoseOneData.GetAttributes();
+
+	const float WeightOfPoseTwo = 1.f - WeightOfPoseOne;
+
+	// Scale first pose in-place, then accumulate second pose
+	BlendPose<ETransformBlendMode::Overwrite>(InOutPose, InOutPose, WeightOfPoseOne);
+	BlendPose<ETransformBlendMode::Accumulate>(SourcePoseTwoData.GetPose(), InOutPose, WeightOfPoseTwo);
+
+	// Ensure that all of the resulting rotations are normalized
+	InOutPose.NormalizeRotations();
+
+	// Use 3 args Lerp over 2 args in-place LerpTo as this is more performant and the memory overhead is negligible
+	FBlendedCurve TempCurve;
+	TempCurve.Lerp(InOutCurve, SourcePoseTwoData.GetCurve(), WeightOfPoseTwo);
+	InOutCurve = MoveTemp(TempCurve);
+
+	// Use a temporary for attributes to avoid aliasing
+	UE::Anim::FStackAttributeContainer TempAttributes;
+	UE::Anim::Attributes::BlendAttributes({ InOutPoseOneData.GetAttributes(), SourcePoseTwoData.GetAttributes() }, { WeightOfPoseOne, WeightOfPoseTwo }, { 0, 1 }, TempAttributes);
+	InOutAttributes = MoveTemp(TempAttributes);
+}
+
 void FAnimationRuntime::BlendTwoPosesTogetherPerBone(
 		const FCompactPose& SourcePose,
 		const FCompactPose& TargetPose,
diff --git a/Engine/Source/Runtime/Engine/Public/AnimationRuntime.h b/Engine/Source/Runtime/Engine/Public/AnimationRuntime.h
--- a/Engine/Source/Runtime/Engine/Public/AnimationRuntime.h
+++ b/Engine/Source/Runtime/Engine/Public/AnimationRuntime.h
@@ -304,6 +304,16 @@ public:
 		const float WeightOfPoseOne,
 		/*out*/ FAnimationPoseData& OutAnimationPoseData);
 
+	/**
+	 * Blends two poses together where the first pose is already in the output buffer (in-place).
+	 * The first pose is scaled by WeightOfPoseOne, then the second pose is accumulated with (1 - WeightOfPoseOne).
+	 * This avoids one pose allocation compared to BlendTwoPosesTogether.
+	 */
+	static ENGINE_API void BlendTwoPosesTogetherInPlace(
+		FAnimationPoseData& InOutPoseOneData,
+		const FAnimationPoseData& SourcePoseTwoData,
+		const float WeightOfPoseOne);
+
 	/**
 	* Blends together a set of poses, each with a given weight.
 	* This function is for BlendSpace per bone blending. BlendSampleDataCache contains the weight information
