Skeletal mesh's UnifiedBoneMap is now constructed from the Mesh Descriptions if available.

This fixes an issue where meshes with Mesh Descriptions would end up with a UnifiedBoneMap constructed from stale data from the previous build.

#Synced-CL 52468466

diff --git a/Engine/Source/Runtime/Engine/Private/SkeletalMeshRenderData.cpp b/Engine/Source/Runtime/Engine/Private/SkeletalMeshRenderData.cpp
--- a/Engine/Source/Runtime/Engine/Private/SkeletalMeshRenderData.cpp
+++ b/Engine/Source/Runtime/Engine/Private/SkeletalMeshRenderData.cpp
@@ -12,6 +12,7 @@
 #include "Rendering/RenderCommandPipes.h"
 #include "Rendering/RayTracingGeometryManager.h"
 #include "Algo/Unique.h"
+#include "SkeletalMeshAttributes.h"
 
 #if WITH_EDITOR
 #include "ProfilingDebugging/CookStats.h"
@@ -229,11 +230,45 @@ inline TArray<FBoneIndexType> BuildUnifiedBoneMap(USkinnedAsset* Owner)
 
 	if (bOptimizeForInstancing)
 	{
-		for (const FSkeletalMeshLODModel& LOD : SkelMeshModel->LODModels)
+		USkeletalMesh* SkelMesh = Cast<USkeletalMesh>(Owner);
+
+		for (int32 LODIndex = 0; LODIndex < SkelMeshModel->LODModels.Num(); LODIndex++)
 		{
-			for (const FSkelMeshSection& Section : LOD.Sections)
+			// For LODs that have a MeshDescription, read bone indices from it rather than
+			// from the ImportedModel's section bone maps, which may be stale from a previous
+			// build.
+			const FMeshDescription* MeshDescription = SkelMesh ? SkelMesh->GetMeshDescription(LODIndex) : nullptr;
+			if (MeshDescription)
 			{
-				UnifiedBoneMap.Append(Section.BoneMap);
+				const FSkeletalMeshConstAttributes MeshAttributes(*MeshDescription);
+				const FSkinWeightsVertexAttributesConstRef VertexSkinWeights = MeshAttributes.GetVertexSkinWeights();
+
+				if (VertexSkinWeights.IsValid())
+				{
+					for (const FVertexID VertexID : MeshDescription->Vertices().GetElementIDs())
+					{
+						const FVertexBoneWeightsConst BoneWeights = VertexSkinWeights.Get(VertexID);
+						for (int32 InfluenceIndex = 0; InfluenceIndex < BoneWeights.Num(); ++InfluenceIndex)
+						{
+							if (BoneWeights[InfluenceIndex].GetWeight() > 0.0f)
+							{
+								// The bone index added here is an index into the ref skeleton, 
+								// which matches the behavior of the non-MeshDescription path below.
+								UnifiedBoneMap.AddUnique(BoneWeights[InfluenceIndex].GetBoneIndex());
+							}
+						}
+					}
+				}
+			}
+			else
+			{
+				// For LODs without a MeshDescription (e.g. generated/simplified LODs), fall back to
+				// reading the existing section bone maps which are stable and won't be rebuilt.
+				const FSkeletalMeshLODModel& LOD = SkelMeshModel->LODModels[LODIndex];
+				for (const FSkelMeshSection& Section : LOD.Sections)
+				{
+					UnifiedBoneMap.Append(Section.BoneMap);
+				}
 			}
 		}
 
