[Input] Remove deprecated 5.6 functions for 5.8

#Synced-CL 52444545

diff --git a/Engine/Source/Runtime/Engine/Public/InputKeyEventArgs.h b/Engine/Source/Runtime/Engine/Public/InputKeyEventArgs.h
--- a/Engine/Source/Runtime/Engine/Public/InputKeyEventArgs.h
+++ b/Engine/Source/Runtime/Engine/Public/InputKeyEventArgs.h
@@ -5,11 +5,6 @@
 #include "Engine/EngineBaseTypes.h"
 #include "InputCoreTypes.h"
 
-#if UE_ENABLE_INCLUDE_ORDER_DEPRECATED_IN_5_6
-#include "GenericPlatform/GenericPlatformInputDeviceMapper.h"
-#include "GenericPlatform/GenericWindow.h"
-#endif
-
 class FViewport;
 
 /**
@@ -50,25 +45,8 @@ public:
 		const bool bInIsTouchEvent,
 		const uint64 InEventTimestamp);
 
-	UE_DEPRECATED(5.6, "Use the newer constructors which takes a timestamp instead.")
-	FInputKeyEventArgs(
-	FViewport* InViewport,
-	FInputDeviceId InInputDevice,
-	FKey InKey,
-	EInputEvent InEvent,
-	float InAmountDepressed,
-	bool bInIsTouchEvent)
-		: FInputKeyEventArgs(InViewport,
-		InInputDevice,
-		InKey,
-		InEvent,
-		InAmountDepressed,
-		bInIsTouchEvent,
-		/*EventTimestamp*/0u)
-	{}
-
 	/**
-	 * Constrcut a FInputKeyEventArgs based on data acquired from the Slate input events.
+	 * Construct a FInputKeyEventArgs based on data acquired from the Slate input events.
 	 * These event args are used to translate from Slate input event arguments into a
 	 * standardized form for the ViewportClient and then the rest of the gameplay framework.
 	 * 
@@ -95,19 +73,6 @@ public:
 	{
 	}
 
-	UE_DEPRECATED(5.6, "Use the newer constructors which takes a timestamp instead.")
-	FInputKeyEventArgs(FViewport* InViewport, FInputDeviceId InInputDevice, FKey InKey, EInputEvent InEvent)
-	: FInputKeyEventArgs(
-	 InViewport,
-	 InInputDevice, 
-	 InKey, 
-	 InEvent,
-	 /*AmountDepressed*/1.0f,
-	 /*bIsTouch*/ false,
-	 0u)
-	{
-	}
-
 	/**
 	 * Construct a FInputKeyEventArgs based on the old params of the ViewportClient InputAxis function. 
 	 * 
@@ -150,32 +115,8 @@ public:
 		FViewport* Viewport = nullptr
 	);
 	
-	UE_DEPRECATED(5.6, "Use the newer constructors which take in a FInputDeviceId and a timestamp instead.")
-	FInputKeyEventArgs(FViewport* InViewport, int32 InControllerId, FKey InKey, EInputEvent InEvent, float InAmountDepressed, bool bInIsTouchEvent)
-		: Viewport(InViewport)
-		, ControllerId(InControllerId)
-		, InputDevice(FInputDeviceId::CreateFromInternalId(InControllerId))
-		, Key(InKey)
-		, Event(InEvent)
-		, AmountDepressed(InAmountDepressed)
-		, bIsTouchEvent(bInIsTouchEvent)
-	{
-	}
-
-	UE_DEPRECATED(5.6, "Use the newer constructors which take in a FInputDeviceId and a timestamp instead.")
-	FInputKeyEventArgs(FViewport* InViewport, int32 InControllerId, FKey InKey, EInputEvent InEvent)
-		: FInputKeyEventArgs(
-			InViewport,
-			FInputDeviceId::CreateFromInternalId(InControllerId),
-			InKey,
-			InEvent,
-			/*EventTimestamp*/ 0u)
-	{
-		ControllerId = InControllerId;
-	}
-	
 	/**
-	 * @return True if this input event is for a Gamepay key, or should be treated as such.
+	 * @return True if this input event is for a gamepad key, or should be treated as such.
 	 * 
 	 * This may be true for non-gamepad FKeys if the input event is simulated input.
 	 */
@@ -227,8 +168,8 @@ public:
 	float DeltaTime = 1.0f / 60.f;
 
 	/**
-	 * For analog key events. The number of samples of analog input
-	 * which are contained in this events InputValue.
+	 * For analog key events. The number of analog input samples
+	 * which are contained in this event's InputValue.
 	 */
 	int32 NumSamples = 1;
 	
@@ -242,7 +183,7 @@ public:
 private:
 
 	/**
-	 * True if this input event is NOT sourced from a physical Human Interface Device (controller, keyboard, mouse, etc)
+	 * True if this input event is NOT sourced from a physical Human Interface Device (controller, keyboard, mouse, etc.)
 	 * and is instead sourced from code, such as faking input events or simulating input for touch screens.
 	 */
 	bool bIsSimulatedInput = false;
@@ -259,7 +200,7 @@ public:
 	 * This timestamp is in terms of CPU Cycles, gathered from the platform's high resolution clock
 	 * via the FPlatformTime::Cycles64() function.
 	 *
-	 * A timestamp value of 0u means that this event was not initalized with a valid timestamp
+	 * A timestamp value of 0u means that this event was not initialized with a valid timestamp
 	 */
 	uint64 EventTimestamp = 0u;
 };
