[shaders] change material shadermap layout and content hashes to fshaderhash (xx64) instead of sha1

#Synced-CL 52473344

diff --git a/Engine/Source/Runtime/Core/Private/UObject/FortniteMainBranchObjectVersion.cpp b/Engine/Source/Runtime/Core/Private/UObject/FortniteMainBranchObjectVersion.cpp
--- a/Engine/Source/Runtime/Core/Private/UObject/FortniteMainBranchObjectVersion.cpp
+++ b/Engine/Source/Runtime/Core/Private/UObject/FortniteMainBranchObjectVersion.cpp
@@ -10,7 +10,7 @@ TMap<FGuid, FGuid> FFortniteMainBranchObjectVersion::GetSystemGuids()
 
 	SystemGuids.Add(DevGuids.GLOBALSHADERMAP_DERIVEDDATA_VER, FGuid("C340EEF623494448979F6D24B496E85A"));
 	SystemGuids.Add(DevGuids.LANDSCAPE_MOBILE_COOK_VERSION, FGuid("32D02EF867C74B71A0D4E0FA41392732"));
-	SystemGuids.Add(DevGuids.MATERIALSHADERMAP_DERIVEDDATA_VER, FGuid("FB0EB87E970C4FB9B217E0E00942BFCB"));
+	SystemGuids.Add(DevGuids.MATERIALSHADERMAP_DERIVEDDATA_VER, FGuid("49B25321CC344253B87962AD4236096E"));
 	SystemGuids.Add(DevGuids.NIAGARASHADERMAP_DERIVEDDATA_VER, FGuid("87F6AE6C269A4592B281752A582E7F13"));
 	SystemGuids.Add(DevGuids.Niagara_LatestScriptCompileVersion, FGuid("266FBE22ADCD4773B98124B68DBE3A4C"));
 	SystemGuids.Add(DevGuids.MaterialTranslationDDCVersion, FGuid("B79683EE9E684C13BAC2FDE8B180477F"));
diff --git a/Engine/Source/Runtime/Engine/Private/Materials/MaterialShader.cpp b/Engine/Source/Runtime/Engine/Private/Materials/MaterialShader.cpp
--- a/Engine/Source/Runtime/Engine/Private/Materials/MaterialShader.cpp
+++ b/Engine/Source/Runtime/Engine/Private/Materials/MaterialShader.cpp
@@ -414,7 +414,7 @@ namespace UE::MaterialInterface::Private
 
 // DDC Key for Material structs saved in incremental cook artifacts, bump when changing related classes which
 // can be loaded/saved to CompactBinary
-#define MATERIAL_CLASS_INCREMENTALCOOK_VER TEXT("B127B459-4F59-4075-905F-0B17A982933E")
+#define MATERIAL_CLASS_INCREMENTALCOOK_VER TEXT("961F0483-145A-4C7A-8DC3-60C382D81DDE")
 
 TMap<EShaderPlatform, FBlake3Hash> GShaderMapPlatformHashes;
 FRWLock GShaderMapPlatformHashesLock;
@@ -2407,7 +2407,7 @@ void FMaterialShaderMap::LoadForRemoteRecompile(FArchive& Ar, EShaderPlatform Sh
 
 void FMaterialShaderMapContent::Finalize(const FShaderMapResourceCode* Code)
 {
-	FSHA1 Hasher;
+	FShaderHashBuilder Hasher;
 	FShaderMapContent::Finalize(Code);
 	UpdateHash(Hasher);
 
@@ -2417,8 +2417,7 @@ void FMaterialShaderMapContent::Finalize(const FShaderMapResourceCode* Code)
 		MeshShaderMap->UpdateHash(Hasher);
 	}
 
-	Hasher.Final();
-	Hasher.GetHash(ShaderContentHash.Hash);
+	ShaderContentHash = Hasher.Finalize();
 }
 
 #if WITH_EDITOR
@@ -3206,13 +3205,13 @@ private:
 		const TArray<FShaderPipelineType*>& SortedMaterialPipelineTypes = FShaderPipelineType::GetSortedTypes(FShaderType::EShaderTypeForDynamicCast::Material);
 		const TArray<FShaderPipelineType*>& SortedMeshMaterialPipelineTypes = FShaderPipelineType::GetSortedTypes(FShaderType::EShaderTypeForDynamicCast::MeshMaterial);
 
-		FSHA1 Hasher;
+		FShaderHashBuilder Hasher;
 #if WITH_EDITOR
 		// GetSourceHash is only available if AllowShaderCompiling (it asserts otherwise).
 		// We only need the hash of sources in the editor, to support incremental cooking's need to know
 		// when to recook the Material.
 		check(AllowShaderCompiling());
-		FSHA1 ShaderCompileHasher;
+		FShaderHashBuilder ShaderCompileHasher;
 #endif
 
 		for (FShaderType* BaseShaderType : SortedMaterialShaderTypes)
@@ -3324,12 +3323,10 @@ private:
 			}
 		}
 
-		Hasher.Final();
-		Hasher.GetHash(Layout.ShaderMapHash.Hash);
+		Layout.ShaderMapHash = Hasher.Finalize();
 #if WITH_EDITOR
 		ShaderCompileHasher.Update((uint8*)&Layout.ShaderMapHash, sizeof(Layout.ShaderMapHash));
-		ShaderCompileHasher.Final();
-		ShaderCompileHasher.GetHash(Layout.ShaderCompileHash.Hash);
+		Layout.ShaderCompileHash = ShaderCompileHasher.Finalize();
 #endif
 	}
 
diff --git a/Engine/Source/Runtime/Engine/Private/Materials/MaterialShaderMapLayout.h b/Engine/Source/Runtime/Engine/Private/Materials/MaterialShaderMapLayout.h
--- a/Engine/Source/Runtime/Engine/Private/Materials/MaterialShaderMapLayout.h
+++ b/Engine/Source/Runtime/Engine/Private/Materials/MaterialShaderMapLayout.h
@@ -3,7 +3,7 @@
 #pragma once
 
 #include "CoreMinimal.h"
-#include "Misc/SecureHash.h"
+#include "Hash/ShaderHash.h"
 #include "RHIDefinitions.h"
 
 class FShaderType;
@@ -43,10 +43,10 @@ class FMaterialShaderMapLayout : public FShaderMapLayout
 {
 public:
 	TArray<FMeshMaterialShaderMapLayout> MeshShaderMaps;
-	FSHAHash ShaderMapHash;
+	FShaderHash ShaderMapHash;
 #if WITH_EDITOR
 	/** Only available if AllowShaderCompiling(). Includes all of ShaderMapHash, plus hash of the shader source files. */
-	FSHAHash ShaderCompileHash;
+	FShaderHash ShaderCompileHash;
 #endif
 };
 
diff --git a/Engine/Source/Runtime/Engine/Public/MaterialShared.h b/Engine/Source/Runtime/Engine/Public/MaterialShared.h
--- a/Engine/Source/Runtime/Engine/Public/MaterialShared.h
+++ b/Engine/Source/Runtime/Engine/Public/MaterialShared.h
@@ -1604,7 +1604,7 @@ private:
 	/** Uniform expressions generated from the material compile. */
 	LAYOUT_FIELD(FMaterialCompilationOutput, MaterialCompilationOutput);
 
-	LAYOUT_FIELD(FSHAHash, ShaderContentHash);
+	LAYOUT_FIELD(FShaderHash, ShaderContentHash);
 
 	LAYOUT_FIELD(FScriptName, UserSceneTextureOutput);
 	LAYOUT_FIELD_INITIALIZED(int32, UserTextureDivisorX, 0);
@@ -1785,7 +1785,7 @@ public:
 	FMeshMaterialShaderMap* AcquireMeshShaderMap(const FVertexFactoryType* VertexFactoryType) { return GetMutableContent()->AcquireMeshShaderMap(VertexFactoryType->GetHashedName()); }
 	FMeshMaterialShaderMap* AcquireMeshShaderMap(const FHashedName& VertexFactoryTypeName) { return GetMutableContent()->AcquireMeshShaderMap(VertexFactoryTypeName); }
 	const FMaterialShaderMapId& GetShaderMapId() const { return ShaderMapId; }
-	const FSHAHash& GetShaderContentHash() const { return GetContent()->ShaderContentHash; }
+	const FShaderHash& GetShaderContentHash() const { return GetContent()->ShaderContentHash; }
 
 #if WITH_EDITOR
 	uint32 AcquireCompilingId(const TRefCountPtr<FSharedShaderCompilerEnvironment>& InMaterialEnvironment);
diff --git a/Engine/Source/Runtime/RenderCore/Private/ShaderMap.cpp b/Engine/Source/Runtime/RenderCore/Private/ShaderMap.cpp
--- a/Engine/Source/Runtime/RenderCore/Private/ShaderMap.cpp
+++ b/Engine/Source/Runtime/RenderCore/Private/ShaderMap.cpp
@@ -1014,23 +1014,6 @@ void FShaderMapContent::Finalize(const FShaderMapResourceCode* Code)
 	ShaderHash = MoveTemp(NewShaderHash);
 }
 
-void FShaderMapContent::UpdateHash(FSHA1& Hasher) const
-{
-	for (int32 ShaderIndex = 0; ShaderIndex < Shaders.Num(); ++ShaderIndex)
-	{
-		const uint64 TypeNameHash = ShaderTypes[ShaderIndex].GetHash();
-		const int32 PermutationId = ShaderPermutations[ShaderIndex];
-		Hasher.Update((uint8*)&TypeNameHash, sizeof(TypeNameHash));
-		Hasher.Update((uint8*)&PermutationId, sizeof(PermutationId));
-	}
-
-	for (const FShaderPipeline* Pipeline : GetShaderPipelines())
-	{
-		const uint64 TypeNameHash = Pipeline->TypeName.GetHash();
-		Hasher.Update((uint8*)&TypeNameHash, sizeof(TypeNameHash));
-	}
-}
-
 void FShaderMapContent::Empty()
 {
 	EmptyShaderPipelines();
diff --git a/Engine/Source/Runtime/RenderCore/Public/Shader.h b/Engine/Source/Runtime/RenderCore/Public/Shader.h
--- a/Engine/Source/Runtime/RenderCore/Public/Shader.h
+++ b/Engine/Source/Runtime/RenderCore/Public/Shader.h
@@ -2501,7 +2501,22 @@ public:
 
 	RENDERCORE_API void Finalize(const FShaderMapResourceCode* Code);
 
-	RENDERCORE_API void UpdateHash(FSHA1& Hasher) const;
+	template <typename THashBuilder>
+	void UpdateHash(THashBuilder& Hasher) const
+	{
+		for (int32 ShaderIndex = 0; ShaderIndex < Shaders.Num(); ++ShaderIndex)
+		{
+			const uint64 TypeNameHash = ShaderTypes[ShaderIndex].GetHash();
+			const int32 PermutationId = ShaderPermutations[ShaderIndex];
+			Hasher.Update((uint8*)&TypeNameHash, sizeof(TypeNameHash));
+			Hasher.Update((uint8*)&PermutationId, sizeof(PermutationId));
+		}
+		for (const FShaderPipeline* Pipeline : GetShaderPipelines())
+		{
+			const uint64 TypeNameHash = Pipeline->TypeName.GetHash();
+			Hasher.Update((uint8*)&TypeNameHash, sizeof(TypeNameHash));
+		}
+	}
 
 protected:
 	RENDERCORE_API void EmptyShaderPipelines();
