UE: Change GetPropertiesSize ensure in RepMovementNetSerializer to a GetStructureSize ensure. Right now these values are the same for native types like FRepMovement since tail padding is included in sizeof. This prepares us to change PropertiesSize to be more accurate

#Synced-CL 52363746

diff --git a/Engine/Source/Runtime/Engine/Private/RepMovementNetSerializer.cpp b/Engine/Source/Runtime/Engine/Private/RepMovementNetSerializer.cpp
--- a/Engine/Source/Runtime/Engine/Private/RepMovementNetSerializer.cpp
+++ b/Engine/Source/Runtime/Engine/Private/RepMovementNetSerializer.cpp
@@ -930,8 +930,8 @@ bool FRepMovementNetSerializer::FNetSerializerRegistryDelegates::IsRepMovementLa
 {
 	const UStruct* RepMovementStruct = FRepMovement::StaticStruct();
 
-	constexpr int ExpectedPropertiesSize = 136;
-	if (!ensureMsgf(RepMovementStruct->GetPropertiesSize() == ExpectedPropertiesSize, TEXT("Unexpected FRepMovement properties size. %d != %d"), RepMovementStruct->GetPropertiesSize(), ExpectedPropertiesSize))
+	constexpr int ExpectedStructureSize = 136;
+	if (!ensureMsgf(RepMovementStruct->GetStructureSize() == ExpectedStructureSize, TEXT("Unexpected FRepMovement structure size. %d != %d"), RepMovementStruct->GetStructureSize(), ExpectedStructureSize))
 	{
 		return false;
 	}
