Fix race condition between FDistanceFieldAsyncQueue::FinalizeAsyncTask() on game thread and FCardRepresentationAsyncQueue::Build() on background worker.
- Issue: FCardRepresentationAsyncQueue::Build(...) reads GenerateSource->GetRenderData()->LODResources[0].DistanceFieldData on a background worker but that member is written by game thread.
- Fix: Pass pointer to new DistanceFieldData when FDistanceFieldAsyncQueue::Build(...) kicks off FAsyncCardRepresentationTask avoiding the need for more synchronization.

#Synced-CL 52471787

diff --git a/Engine/Source/Runtime/Engine/Private/DistanceFieldAtlas.cpp b/Engine/Source/Runtime/Engine/Private/DistanceFieldAtlas.cpp
--- a/Engine/Source/Runtime/Engine/Private/DistanceFieldAtlas.cpp
+++ b/Engine/Source/Runtime/Engine/Private/DistanceFieldAtlas.cpp
@@ -139,7 +139,8 @@ FAsyncCardRepresentationTask* BeginCacheMeshCardRepresentationInternal(
 	FStaticMeshRenderData& RenderData,
 	const FString& DistanceFieldKey,
 	FSourceMeshDataForDerivedDataTask&& OptionalSourceMeshData,
-	TArray<FSignedDistanceFieldBuildPartData>&& OptionalPartData);
+	TArray<FSignedDistanceFieldBuildPartData>&& OptionalPartData,
+	FDistanceFieldVolumeData* DistanceFieldData = nullptr);
 
 // DDC key for distance field data, must be changed when modifying the generation code or data format
 #define DISTANCEFIELD_DERIVEDDATA_VER TEXT("DC2427EE-AD20-4226-ADAD-15CAEB4FC9AB")
@@ -1051,7 +1052,8 @@ void FDistanceFieldAsyncQueue::Build(FAsyncDistanceFieldTask* Task)
 			*Task->StaticMesh->GetRenderData(),
 			Task->DDCKey,
 			MoveTemp(Task->SourceMeshData),
-			MoveTemp(Task->PartData));
+			MoveTemp(Task->PartData),
+			Task->GeneratedVolumeData);
 
 		if (CardTask)
 		{
diff --git a/Engine/Source/Runtime/Engine/Private/MeshCardRepresentation.cpp b/Engine/Source/Runtime/Engine/Private/MeshCardRepresentation.cpp
--- a/Engine/Source/Runtime/Engine/Private/MeshCardRepresentation.cpp
+++ b/Engine/Source/Runtime/Engine/Private/MeshCardRepresentation.cpp
@@ -85,6 +85,7 @@ FAsyncCardRepresentationTask* CreateCardRepresentationTask(
 	const FString& InDDCKey,
 	UStaticMesh* Mesh,
 	UStaticMesh* GenerateSource,
+	FDistanceFieldVolumeData* DistanceFieldData,
 	int32 MaxLumenMeshCards,
 	bool bGenerateDistanceFieldAsIfTwoSided,
 	FSourceMeshDataForDerivedDataTask&& OptionalSourceMeshData,
@@ -233,7 +234,8 @@ FAsyncCardRepresentationTask* BeginCacheMeshCardRepresentationInternal(
 	FStaticMeshRenderData& RenderData,
 	const FString& DistanceFieldKey,
 	FSourceMeshDataForDerivedDataTask&& OptionalSourceMeshData,
-	TArray<FSignedDistanceFieldBuildPartData>&& OptionalPartData)
+	TArray<FSignedDistanceFieldBuildPartData>&& OptionalPartData,
+	FDistanceFieldVolumeData* DistanceFieldData = nullptr)
 {
 	static const auto CVarCards = IConsoleManager::Get().FindTConsoleVariableDataInt(TEXT("r.MeshCardRepresentation"));
 
@@ -260,6 +262,7 @@ FAsyncCardRepresentationTask* BeginCacheMeshCardRepresentationInternal(
 					Key,
 					StaticMeshAsset,
 					MeshToGenerateFrom,
+					DistanceFieldData,
 					BuildSettings.MaxLumenMeshCards,
 					BuildSettings.bGenerateDistanceFieldAsIfTwoSided,
 					MoveTemp(OptionalSourceMeshData),
@@ -286,6 +289,7 @@ FAsyncCardRepresentationTask* CreateCardRepresentationTask(
 	const FString& InDDCKey,
 	UStaticMesh* Mesh,
 	UStaticMesh* GenerateSource,
+	FDistanceFieldVolumeData* DistanceFieldData,
 	int32 MaxLumenMeshCards,
 	bool bGenerateDistanceFieldAsIfTwoSided,
 	FSourceMeshDataForDerivedDataTask&& OptionalSourceMeshData,
@@ -324,6 +328,7 @@ FAsyncCardRepresentationTask* CreateCardRepresentationTask(
 		NewTask->DDCKey = InDDCKey;
 		NewTask->StaticMesh = Mesh;
 		NewTask->GenerateSource = GenerateSource;
+		NewTask->DistanceFieldData = DistanceFieldData;
 		NewTask->GeneratedCardRepresentation = new FCardRepresentationData();
 		NewTask->MaxLumenMeshCards = MaxLumenMeshCards;
 		NewTask->bGenerateDistanceFieldAsIfTwoSided = bGenerateDistanceFieldAsIfTwoSided;
@@ -423,7 +428,7 @@ void FCardRepresentationData::CacheDerivedData(
 	FSourceMeshDataForDerivedDataTask&& OptionalSourceMeshData,
 	TArray<FSignedDistanceFieldBuildPartData>&& OptionalPartData)
 {
-	FAsyncCardRepresentationTask* NewTask = CreateCardRepresentationTask(InDDCKey, Mesh, GenerateSource, MaxLumenMeshCards, bGenerateDistanceFieldAsIfTwoSided, MoveTemp(OptionalSourceMeshData), MoveTemp(OptionalPartData), *this);
+	FAsyncCardRepresentationTask* NewTask = CreateCardRepresentationTask(InDDCKey, Mesh, GenerateSource, nullptr, MaxLumenMeshCards, bGenerateDistanceFieldAsIfTwoSided, MoveTemp(OptionalSourceMeshData), MoveTemp(OptionalPartData), *this);
 
 	if (NewTask)
 	{
@@ -1012,7 +1017,7 @@ void FCardRepresentationAsyncQueue::Build(FAsyncCardRepresentationTask* Task)
 		Task->bSuccess = MeshUtilities->GenerateCardRepresentationData(
 			Task->StaticMesh->GetName(),
 			MeshData,
-			Task->GenerateSource->GetRenderData()->LODResources[0].DistanceFieldData,
+			Task->DistanceFieldData ? Task->DistanceFieldData : Task->GenerateSource->GetRenderData()->LODResources[0].DistanceFieldData,
 			Task->MaxLumenMeshCards,
 			Task->bGenerateDistanceFieldAsIfTwoSided,
 			*Task->GeneratedCardRepresentation);
diff --git a/Engine/Source/Runtime/Engine/Public/MeshCardBuild.h b/Engine/Source/Runtime/Engine/Public/MeshCardBuild.h
--- a/Engine/Source/Runtime/Engine/Public/MeshCardBuild.h
+++ b/Engine/Source/Runtime/Engine/Public/MeshCardBuild.h
@@ -17,6 +17,7 @@
 class FAsyncCompilationNotification;
 class FObjectCacheContextScope;
 class UStaticMesh;
+class FDistanceFieldVolumeData;
 struct FAppendToClassSchemaContext;
 struct FSignedDistanceFieldBuildSectionData;
 
@@ -243,6 +244,7 @@ public:
 	int32 MaxLumenMeshCards = 0;
 	UStaticMesh* StaticMesh = nullptr;
 	UStaticMesh* GenerateSource = nullptr;
+	FDistanceFieldVolumeData* DistanceFieldData = nullptr;
 	FString DDCKey;
 	FCardRepresentationData* GeneratedCardRepresentation;
 	TUniquePtr<FAsyncTask<FAsyncCardRepresentationTaskWorker>> AsyncTask = nullptr;
