Expose EMemoryUnit and FText::AsMemory(EMemoryUnit) overload to UKismetTextLibrary

#rnx
#Synced-CL 52440464

diff --git a/Engine/Source/Runtime/Core/Public/Internationalization/Text.h b/Engine/Source/Runtime/Core/Public/Internationalization/Text.h
--- a/Engine/Source/Runtime/Core/Public/Internationalization/Text.h
+++ b/Engine/Source/Runtime/Core/Public/Internationalization/Text.h
@@ -196,21 +196,24 @@ enum EMemoryUnitStandard : int
 	SI
 };
 
-/** Specific memory unit to use when formatting memory sizes. IEC units are 1024-based (KiB, MiB, ...), SI units are 1000-based (KB, MB, ...). */
-enum class EMemoryUnit : int
+/** 
+ * Specific memory unit to use when formatting memory sizes. IEC units are 1024-based (KiB, MiB, ...), SI units are 1000-based (KB, MB, ...).
+ * @note Redeclared in KismetTextLibrary for meta-data extraction purposes, be sure to update there as well.
+ */
+enum class EMemoryUnit : uint8
 {
 	Unit_Bytes,
 
-	// Auto-select based on magnitude using IEC
+	/* Auto-select based on magnitude using IEC */
 	Auto_IEC,
 
-	// IEC (1024-based)
+	/* IEC (1024-based) */
 	Unit_KiB, Unit_MiB, Unit_GiB, Unit_TiB, Unit_PiB, Unit_EiB,
 
-	// Auto-select based on magnitude using SI
+	/* Auto-select based on magnitude using SI */
 	Auto_SI,
 
-	// SI (1000-based)
+	/* SI (1000-based) */
 	Unit_KB, Unit_MB, Unit_GB, Unit_TB, Unit_PB, Unit_EB
 };
 
diff --git a/Engine/Source/Runtime/Engine/Classes/Kismet/KismetTextLibrary.h b/Engine/Source/Runtime/Engine/Classes/Kismet/KismetTextLibrary.h
--- a/Engine/Source/Runtime/Engine/Classes/Kismet/KismetTextLibrary.h
+++ b/Engine/Source/Runtime/Engine/Classes/Kismet/KismetTextLibrary.h
@@ -43,6 +43,24 @@ enum EMemoryUnitStandard : int
 	SI
 };
 
+UENUM(BlueprintType)
+enum class EMemoryUnit : uint8
+{
+	Unit_Bytes,
+
+	/* Auto-select based on magnitude using IEC (1024-based) */
+	Auto_IEC,
+
+	/* IEC (1024-based) */
+	Unit_KiB, Unit_MiB, Unit_GiB, Unit_TiB, Unit_PiB, Unit_EiB,
+
+	/* Auto-select based on magnitude using SI (1000-based) */
+	Auto_SI,
+
+	/* SI (1000-based) */
+	Unit_KB, Unit_MB, Unit_GB, Unit_TB, Unit_PB, Unit_EB
+};
+
 UENUM(BlueprintType)
 enum class ETextGender : uint8
 {
@@ -303,6 +321,10 @@ class UKismetTextLibrary : public UBlueprintFunctionLibrary
 	UFUNCTION(BlueprintPure, meta=(AdvancedDisplay = "1"), Category="Utilities|Text")
 	static ENGINE_API FText AsMemory(int64 NumBytes, TEnumAsByte<EMemoryUnitStandard> UnitStandard = EMemoryUnitStandard::IEC, bool bUseGrouping = true, int32 MinimumIntegralDigits = 1, int32 MaximumIntegralDigits = 324, int32 MinimumFractionalDigits = 0, int32 MaximumFractionalDigits = 3);
 
+	/** Generate an FText that represents the passed number as a memory size in a specific unit. Use Auto_IEC or Auto_SI to select automatically by magnitude. */
+	UFUNCTION(BlueprintPure, meta=(DisplayName="As Memory (Specific Unit)", AdvancedDisplay = "1"), Category="Utilities|Text")
+	static ENGINE_API FText AsMemoryUnit(int64 NumBytes, EMemoryUnit Unit = EMemoryUnit::Auto_IEC, bool bUseGrouping = true, int32 MinimumIntegralDigits = 1, int32 MaximumIntegralDigits = 324, int32 MinimumFractionalDigits = 0, int32 MaximumFractionalDigits = 3);
+
 	/** Used for formatting text using the FText::Format function and utilized by the UK2Node_FormatText */
 	UFUNCTION(BlueprintPure, meta=(BlueprintInternalUseOnly = "true"))
 	static ENGINE_API FText Format(FText InPattern, TArray<FFormatArgumentData> InArgs);
diff --git a/Engine/Source/Runtime/Engine/Private/KismetTextLibrary.cpp b/Engine/Source/Runtime/Engine/Private/KismetTextLibrary.cpp
--- a/Engine/Source/Runtime/Engine/Private/KismetTextLibrary.cpp
+++ b/Engine/Source/Runtime/Engine/Private/KismetTextLibrary.cpp
@@ -318,6 +318,19 @@ FText UKismetTextLibrary::AsMemory(int64 NumBytes, TEnumAsByte<EMemoryUnitStanda
 	return FText::AsMemory(UnsignedNumBytes, &NumberFormatOptions, nullptr, UnitStandard);
 }
 
+FText UKismetTextLibrary::AsMemoryUnit(int64 NumBytes, EMemoryUnit Unit, bool bUseGrouping, int32 MinimumIntegralDigits, int32 MaximumIntegralDigits, int32 MinimumFractionalDigits, int32 MaximumFractionalDigits)
+{
+	FNumberFormattingOptions NumberFormatOptions;
+	NumberFormatOptions.UseGrouping = bUseGrouping;
+	NumberFormatOptions.MinimumIntegralDigits = MinimumIntegralDigits;
+	NumberFormatOptions.MaximumIntegralDigits = MaximumIntegralDigits;
+	NumberFormatOptions.MinimumFractionalDigits = MinimumFractionalDigits;
+	NumberFormatOptions.MaximumFractionalDigits = MaximumFractionalDigits;
+
+	uint64 UnsignedNumBytes = static_cast<uint64>(FMath::Max(0, NumBytes));
+	return FText::AsMemory(UnsignedNumBytes, Unit, &NumberFormatOptions, nullptr);
+}
+
 FText UKismetTextLibrary::Format(FText InPattern, TArray<FFormatArgumentData> InArgs)
 {
 	return FTextFormatter::Format(MoveTemp(InPattern), MoveTemp(InArgs), false, false);
