Fix error compiling due to redefinition of default argument.

#Synced-CL 52471787

diff --git a/Engine/Source/Runtime/Engine/Private/DistanceFieldAtlas.cpp b/Engine/Source/Runtime/Engine/Private/DistanceFieldAtlas.cpp
--- a/Engine/Source/Runtime/Engine/Private/DistanceFieldAtlas.cpp
+++ b/Engine/Source/Runtime/Engine/Private/DistanceFieldAtlas.cpp
@@ -140,7 +140,7 @@ FAsyncCardRepresentationTask* BeginCacheMeshCardRepresentationInternal(
 	const FString& DistanceFieldKey,
 	FSourceMeshDataForDerivedDataTask&& OptionalSourceMeshData,
 	TArray<FSignedDistanceFieldBuildPartData>&& OptionalPartData,
-	FDistanceFieldVolumeData* DistanceFieldData = nullptr);
+	FDistanceFieldVolumeData* DistanceFieldData);
 
 // DDC key for distance field data, must be changed when modifying the generation code or data format
 #define DISTANCEFIELD_DERIVEDDATA_VER TEXT("DC2427EE-AD20-4226-ADAD-15CAEB4FC9AB")
