[NetPhysics] Improve resimulation RewindData creation flow

#Synced-CL 52498458

diff --git a/Engine/Source/Runtime/Engine/Private/PhysicsEngine/PhysicsReplication.cpp b/Engine/Source/Runtime/Engine/Private/PhysicsEngine/PhysicsReplication.cpp
--- a/Engine/Source/Runtime/Engine/Private/PhysicsEngine/PhysicsReplication.cpp
+++ b/Engine/Source/Runtime/Engine/Private/PhysicsEngine/PhysicsReplication.cpp
@@ -996,7 +996,7 @@ void FPhysicsReplicationAsync::AddResimulationRequest_Internal(const float Delta
 				// Create RewindData if we use resim and not yet created
 				if (!RewindData && EffectiveRepMode == EPhysicsReplicationMode::Resimulation)
 				{
-					if (ensure(Chaos::FPBDRigidsSolver::IsNetworkPhysicsPredictionEnabled() && RigidsSolver->IsUsingFixedDt()))
+					if (Chaos::FPBDRigidsSolver::IsNetworkPhysicsPredictionEnabled() && RigidsSolver->IsUsingFixedDt())
 					{
 						RigidsSolver->EnableRewindCapture();
 						RewindData = RigidsSolver->GetRewindData();
@@ -1112,7 +1112,7 @@ void FPhysicsReplicationAsync::OnPreSimulate_Internal()
 				{
 					if (!bRewindDataExist)
 					{
-						if (ensure(Chaos::FPBDRigidsSolver::IsNetworkPhysicsPredictionEnabled() && RigidsSolver->IsUsingFixedDt()))
+						if (Chaos::FPBDRigidsSolver::IsNetworkPhysicsPredictionEnabled() && RigidsSolver->IsUsingFixedDt())
 						{
 							RigidsSolver->EnableRewindCapture();
 							RewindData = RigidsSolver->GetRewindData();
@@ -1592,6 +1592,8 @@ void FPhysicsReplicationAsync::CheckTargetResimValidity(FReplicatedPhysicsTarget
 	Chaos::FRewindData* RewindData = RigidsSolver->GetRewindData();
 	if (RewindData == nullptr)
 	{
+		// Can't perform resimulation without RewindData, revert to PredictiveInterpolation
+		Target.RepMode = EPhysicsReplicationMode::PredictiveInterpolation;
 		return;
 	}
 
@@ -2605,7 +2607,7 @@ bool FPhysicsReplicationAsync::ResimulationReplication(Chaos::FPBDRigidParticleH
 	Chaos::FRewindData* RewindData = RigidsSolver->GetRewindData();
 	if (!RewindData)
 	{
-		if (ensure(Chaos::FPBDRigidsSolver::IsNetworkPhysicsPredictionEnabled() && RigidsSolver->IsUsingFixedDt()))
+		if (Chaos::FPBDRigidsSolver::IsNetworkPhysicsPredictionEnabled() && RigidsSolver->IsUsingFixedDt())
 		{
 			RigidsSolver->EnableRewindCapture();
 			RewindData = RigidsSolver->GetRewindData();
diff --git a/Engine/Source/Runtime/Experimental/Chaos/Private/Chaos/Framework/PhysicsSolverBase.cpp b/Engine/Source/Runtime/Experimental/Chaos/Private/Chaos/Framework/PhysicsSolverBase.cpp
--- a/Engine/Source/Runtime/Experimental/Chaos/Private/Chaos/Framework/PhysicsSolverBase.cpp
+++ b/Engine/Source/Runtime/Experimental/Chaos/Private/Chaos/Framework/PhysicsSolverBase.cpp
@@ -411,23 +411,19 @@ namespace Chaos
 	
 	void FPhysicsSolverBase::AddResimulationRequest_Internal(FPushPhysicsData& PushData)
 	{
-		if(!MRewindData)
-		{
-			return;
-		}
-
-		// Set all the input on the matching callbacks for future use (PhysicsReplication, RewindCallbacks...)
-		for (FSimCallbackInputAndObject& InputAndCallbackObj : PushData.SimCallbackInputs)
+		// Add resimulation request from rewind callbacks inputs
+		if (!IsShuttingDown() && !bGameThreadFrozen && MRewindCallback)
 		{
-			if(InputAndCallbackObj.CallbackObject)
+			// Set all the input on the matching callbacks so they are available for the resimulation request below.
+			// Note: ProcessPushData_Internal also calls SetCurrentInput_Internal for the general sim callback flow.
+			for (FSimCallbackInputAndObject& InputAndCallbackObj : PushData.SimCallbackInputs)
 			{
-				InputAndCallbackObj.CallbackObject->SetCurrentInput_Internal(InputAndCallbackObj.Input);
+				if(InputAndCallbackObj.CallbackObject)
+				{
+					InputAndCallbackObj.CallbackObject->SetCurrentInput_Internal(InputAndCallbackObj.Input);
+				}
 			}
-		}
 
-		// Add resimulation request from rewind callbacks inputs
-		if (!IsShuttingDown() && !bGameThreadFrozen && MRewindCallback)
-		{
 			QUICK_SCOPE_CYCLE_COUNTER(Chaos_RewindCallback_ProcessInputs_Internal);
 			MRewindCallback->AddResimulationRequest_Internal(PushData.InternalStep, (float)MLastDt);
 		}
