[NetPhysics] Reliable Actions via NetworkPhysicsComponent
- Add support for sending Actions with reliable networking
- Make NetworkPhysicsComponent work with only Actions (no inputs or states)

#Synced-CL 52498458

diff --git a/Engine/Source/Runtime/Engine/Private/PhysicsEngine/NetworkPhysicsComponent.cpp b/Engine/Source/Runtime/Engine/Private/PhysicsEngine/NetworkPhysicsComponent.cpp
--- a/Engine/Source/Runtime/Engine/Private/PhysicsEngine/NetworkPhysicsComponent.cpp
+++ b/Engine/Source/Runtime/Engine/Private/PhysicsEngine/NetworkPhysicsComponent.cpp
@@ -1029,18 +1029,22 @@ void UNetworkPhysicsComponent::PreReplication(IRepChangedPropertyTracker& Change
 {
 	Super::PreReplication(ChangedPropertyTracker);
 
-	// RepGraph / Legacy
-	DOREPLIFETIME_ACTIVE_OVERRIDE_FAST(UNetworkPhysicsComponent, ReplicatedDeltaSourceInput, bIsUsingLegacyData);
-	DOREPLIFETIME_ACTIVE_OVERRIDE_FAST(UNetworkPhysicsComponent, ReplicatedDeltaSourceState, bIsUsingLegacyData);
-	DOREPLIFETIME_ACTIVE_OVERRIDE_FAST(UNetworkPhysicsComponent, ReplicatedOwnerInputs, bIsUsingLegacyData);
-	DOREPLIFETIME_ACTIVE_OVERRIDE_FAST(UNetworkPhysicsComponent, ReplicatedRemoteInputs, bIsUsingLegacyData);
-	DOREPLIFETIME_ACTIVE_OVERRIDE_FAST(UNetworkPhysicsComponent, ReplicatedStates, bIsUsingLegacyData);
-
-	// Iris
-	DOREPLIFETIME_ACTIVE_OVERRIDE_FAST(UNetworkPhysicsComponent, ReplicatedInputCollection, !bIsUsingLegacyData);
-	DOREPLIFETIME_ACTIVE_OVERRIDE_FAST(UNetworkPhysicsComponent, ReplicatedRemoteInputCollection, !bIsUsingLegacyData);
-	DOREPLIFETIME_ACTIVE_OVERRIDE_FAST(UNetworkPhysicsComponent, ReplicatedStateCollection, !bIsUsingLegacyData);
-	DOREPLIFETIME_ACTIVE_OVERRIDE_FAST(UNetworkPhysicsComponent, ReplicatedActionsCollection, true);
+	// Legacy - only active when input/state is registered AND using legacy
+	const bool bLegacyActive = bIsUsingLegacyData && bHasRegisteredInputState;
+	DOREPLIFETIME_ACTIVE_OVERRIDE_FAST(UNetworkPhysicsComponent, ReplicatedDeltaSourceInput, bLegacyActive);
+	DOREPLIFETIME_ACTIVE_OVERRIDE_FAST(UNetworkPhysicsComponent, ReplicatedDeltaSourceState, bLegacyActive);
+	DOREPLIFETIME_ACTIVE_OVERRIDE_FAST(UNetworkPhysicsComponent, ReplicatedOwnerInputs, bLegacyActive);
+	DOREPLIFETIME_ACTIVE_OVERRIDE_FAST(UNetworkPhysicsComponent, ReplicatedRemoteInputs, bLegacyActive);
+	DOREPLIFETIME_ACTIVE_OVERRIDE_FAST(UNetworkPhysicsComponent, ReplicatedStates, bLegacyActive);
+
+	// Iris - only active when input/state is registered AND not using legacy
+	const bool bIrisActive = !bIsUsingLegacyData && bHasRegisteredInputState;
+	DOREPLIFETIME_ACTIVE_OVERRIDE_FAST(UNetworkPhysicsComponent, ReplicatedInputCollection, bIrisActive);
+	DOREPLIFETIME_ACTIVE_OVERRIDE_FAST(UNetworkPhysicsComponent, ReplicatedRemoteInputCollection, bIrisActive);
+	DOREPLIFETIME_ACTIVE_OVERRIDE_FAST(UNetworkPhysicsComponent, ReplicatedStateCollection, bIrisActive);
+
+	// Actions - only active when there are pending actions to replicate
+	DOREPLIFETIME_ACTIVE_OVERRIDE_FAST(UNetworkPhysicsComponent, ReplicatedActionsCollection, bHasPendingActionsToReplicate);
 }
 
 // Called every Game Thread frame
@@ -1065,6 +1069,9 @@ void UNetworkPhysicsComponent::NetworkMarshaledData()
 		return;
 	}
 
+	// Reset pending actions flag, will be set if any actions are produced this frame
+	bHasPendingActionsToReplicate = false;
+
 	const bool bIsServer = HasServerWorld();
 	if (!bIsServer && !IsNetworkPhysicsTickOffsetAssigned())
 	{
@@ -1296,22 +1303,35 @@ void UNetworkPhysicsComponent::NetworkMarshaledData()
 		}
 
 		// Unreliable Actions
-		if (AsyncOutput->ActionData.Num() > 0)
+		if (AsyncOutput->UnreliableActionData.Num() > 0)
 		{
-			ReplicatedActionsCollection.DataArray.SetNum(AsyncOutput->ActionData.Num(), EAllowShrinking::Yes);
-			int32 Idx = 0;
-			for (const TInstancedStruct<FNetworkPhysicsActionPayload>& ActionInstance : AsyncOutput->ActionData)
+			bHasPendingActionsToReplicate = true;
+
+			ReplicatedActionsCollection.DataArray = MoveTemp(AsyncOutput->UnreliableActionData);
+
+			if (bIsServer)
+			{
+				MARK_PROPERTY_DIRTY_FROM_NAME(UNetworkPhysicsComponent, ReplicatedActionsCollection, this);
+			}
+			else
 			{
-				ReplicatedActionsCollection.DataArray[Idx++] = ActionInstance;
+				ServerReceiveUnreliableActionCollection(ReplicatedActionsCollection);
 			}
+		}
+
+		// Reliable Actions (always enabled, not gated by bEnableReliableFlow)
+		if (AsyncOutput->ReliableActionData.Num() > 0)
+		{
+			FNetworkPhysicsActionCollection ReliableActionsToSend;
+			ReliableActionsToSend.DataArray = MoveTemp(AsyncOutput->ReliableActionData);
 
 			if (bIsServer)
 			{
-				MARK_PROPERTY_DIRTY_FROM_NAME(UNetworkPhysicsComponent, ReplicatedActionsCollection, this);
+				MulticastReceiveReliableActionCollection(ReliableActionsToSend);
 			}
 			else
 			{
-				ServerReceiveActionCollection(ReplicatedActionsCollection);
+				ServerReceiveReliableActionCollection(ReliableActionsToSend);
 			}
 		}
 
@@ -1825,16 +1845,16 @@ void UNetworkPhysicsComponent::ServerReceiveInputCollection_Implementation(const
 	}
 }
 
-void UNetworkPhysicsComponent::ServerReceiveActionCollection_Implementation(const FNetworkPhysicsActionCollection& ClientActionCollection)
+void UNetworkPhysicsComponent::ReceiveActionCollectionData(const FNetworkPhysicsActionCollection& ActionCollection)
 {
-	if (!NetworkPhysicsComponent_Internal || ClientActionCollection.DataArray.Num() <= 0)
+	if (!NetworkPhysicsComponent_Internal || ActionCollection.DataArray.Num() <= 0)
 	{
 		return;
 	}
 
 	if (FAsyncNetworkPhysicsComponentInput* AsyncInput = NetworkPhysicsComponent_Internal->GetProducerInputData_External())
 	{
-		for (const TInstancedStruct<FNetworkPhysicsActionPayload>& ActionInstance : ClientActionCollection.DataArray)
+		for (const TInstancedStruct<FNetworkPhysicsActionPayload>& ActionInstance : ActionCollection.DataArray)
 		{
 			if (ActionInstance.IsValid())
 			{
@@ -1844,6 +1864,26 @@ void UNetworkPhysicsComponent::ServerReceiveActionCollection_Implementation(cons
 	}
 }
 
+void UNetworkPhysicsComponent::ServerReceiveUnreliableActionCollection_Implementation(const FNetworkPhysicsActionCollection& ClientActionCollection)
+{
+	ReceiveActionCollectionData(ClientActionCollection);
+}
+
+void UNetworkPhysicsComponent::ServerReceiveReliableActionCollection_Implementation(const FNetworkPhysicsActionCollection& ClientActionCollection)
+{
+	ReceiveActionCollectionData(ClientActionCollection);
+}
+
+void UNetworkPhysicsComponent::MulticastReceiveReliableActionCollection_Implementation(const FNetworkPhysicsActionCollection& ServerActionCollection)
+{
+	// NetMulticast executes on the server too - skip re-ingesting our own actions, the server already has them in ConfirmedActions.
+	if (HasServerWorld())
+	{
+		return;
+	}
+	ReceiveActionCollectionData(ServerActionCollection);
+}
+
 // repnotify for inputs on owner client
 void UNetworkPhysicsComponent::OnRep_SetReplicatedInputCollection()
 {
@@ -1915,49 +1955,35 @@ void UNetworkPhysicsComponent::OnRep_SetReplicatedStateCollection()
 
 void UNetworkPhysicsComponent::OnRep_SetReplicatedActionsCollection()
 {
-	if (!NetworkPhysicsComponent_Internal || ReplicatedActionsCollection.DataArray.Num() <= 0)
-	{
-		return;
-	}
-
-	if (FAsyncNetworkPhysicsComponentInput* AsyncInput = NetworkPhysicsComponent_Internal->GetProducerInputData_External())
-	{
-		for (const TInstancedStruct<FNetworkPhysicsActionPayload>& ActionInstance : ReplicatedActionsCollection.DataArray)
-		{
-			if (ActionInstance.IsValid())
-			{
-				AsyncInput->ReplicatedActionData.Add(ActionInstance);
-			}
-		}
-	}
+	ReceiveActionCollectionData(ReplicatedActionsCollection);
 }
 
-void UNetworkPhysicsComponent::EnqueueImmediateActionInstance_External(TInstancedStruct<FNetworkPhysicsActionPayload>&& ActionInstance, const UObject* SourceObject)
+void UNetworkPhysicsComponent::EnqueueImmediateActionInstance_External(TInstancedStruct<FNetworkPhysicsActionPayload>&& ActionInstance, const UObject* SourceObject, const bool bReliable)
 {
 	const uint32 SourceId = UE::NetworkPhysicsUtils::GetNetworkStableHash_External(SourceObject);
 	ensureMsgf(SourceId > 0u, TEXT("EnqueueImmediateActionInstance_External could not find a network stable source to base an ID from, the object needs to either be replicated or baked in the level. Else it's recommended to pass in a custom unique SourceId as a parameter instead of the source object. Action: %ls  ---  Source: %ls")
 		, ActionInstance.GetScriptStruct() ? *ActionInstance.GetScriptStruct()->GetName() : TEXT("UNKNOWN ACTION")
-		, SourceObject ? *SourceObject->GetName() : TEXT("UNKNOWN SORUCE OBJECT"));
+		, SourceObject ? *SourceObject->GetName() : TEXT("UNKNOWN SOURCE OBJECT"));
 
-	EnqueueImmediateActionInstance_External(MoveTemp(ActionInstance), SourceId);
+	EnqueueImmediateActionInstance_External(MoveTemp(ActionInstance), SourceId, bReliable);
 }
 
-void UNetworkPhysicsComponent::EnqueueImmediateActionInstance_External(TInstancedStruct<FNetworkPhysicsActionPayload>&& ActionInstance, const uint32 SourceId)
+void UNetworkPhysicsComponent::EnqueueImmediateActionInstance_External(TInstancedStruct<FNetworkPhysicsActionPayload>&& ActionInstance, const uint32 SourceId, const bool bReliable)
 {
-	EnqueueScheduledActionInstanceAtFrame_External(MoveTemp(ActionInstance), SourceId, -1);
+	EnqueueScheduledActionInstanceAtFrame_External(MoveTemp(ActionInstance), SourceId, -1, bReliable);
 }
 
-void UNetworkPhysicsComponent::EnqueueScheduledActionInstance_External(TInstancedStruct<FNetworkPhysicsActionPayload>&& ActionInstance, const UObject* SourceObject, const float DelaySeconds)
+void UNetworkPhysicsComponent::EnqueueScheduledActionInstance_External(TInstancedStruct<FNetworkPhysicsActionPayload>&& ActionInstance, const UObject* SourceObject, const float DelaySeconds, const bool bReliable)
 {
 	const uint32 SourceId = UE::NetworkPhysicsUtils::GetNetworkStableHash_External(SourceObject);
 	ensureMsgf(SourceId > 0u, TEXT("EnqueueScheduledActionInstance_External could not find a network stable source to base an ID from, the object needs to either be replicated or baked in the level. Else it's recommended to pass in a custom unique SourceId as a parameter instead of the source object. Action: %ls  ---  Source: %ls")
 		, ActionInstance.GetScriptStruct() ? *ActionInstance.GetScriptStruct()->GetName() : TEXT("UNKNOWN ACTION")
-		, SourceObject ? *SourceObject->GetName() : TEXT("UNKNOWN SORUCE OBJECT"));
+		, SourceObject ? *SourceObject->GetName() : TEXT("UNKNOWN SOURCE OBJECT"));
 
-	EnqueueScheduledActionInstance_External(MoveTemp(ActionInstance), SourceId, DelaySeconds);
+	EnqueueScheduledActionInstance_External(MoveTemp(ActionInstance), SourceId, DelaySeconds, bReliable);
 }
 
-void UNetworkPhysicsComponent::EnqueueScheduledActionInstance_External(TInstancedStruct<FNetworkPhysicsActionPayload>&& ActionInstance, const uint32 SourceId, const float DelaySeconds)
+void UNetworkPhysicsComponent::EnqueueScheduledActionInstance_External(TInstancedStruct<FNetworkPhysicsActionPayload>&& ActionInstance, const uint32 SourceId, const float DelaySeconds, const bool bReliable)
 {
 	int32 LocalFrame = INDEX_NONE;
 	if (UWorld* World = GetWorld())
@@ -1978,20 +2004,20 @@ void UNetworkPhysicsComponent::EnqueueScheduledActionInstance_External(TInstance
 			}
 		}
 	}
-	EnqueueScheduledActionInstanceAtFrame_External(MoveTemp(ActionInstance), SourceId, LocalFrame);
+	EnqueueScheduledActionInstanceAtFrame_External(MoveTemp(ActionInstance), SourceId, LocalFrame, bReliable);
 }
 
-void UNetworkPhysicsComponent::EnqueueScheduledActionInstanceAtFrame_External(TInstancedStruct<FNetworkPhysicsActionPayload>&& ActionInstance, const UObject* SourceObject, const int32 LocalFrame)
+void UNetworkPhysicsComponent::EnqueueScheduledActionInstanceAtFrame_External(TInstancedStruct<FNetworkPhysicsActionPayload>&& ActionInstance, const UObject* SourceObject, const int32 LocalFrame, const bool bReliable)
 {
 	const uint32 SourceId = UE::NetworkPhysicsUtils::GetNetworkStableHash_External(SourceObject);
 	ensureMsgf(SourceId > 0u, TEXT("EnqueueScheduledActionInstanceAtFrame_External could not find a network stable source to base an ID from, the object needs to either be replicated or baked in the level. Else it's recommended to pass in a custom unique SourceId as a parameter instead of the source object. Action: %ls  ---  Source: %ls")
 		, ActionInstance.GetScriptStruct() ? *ActionInstance.GetScriptStruct()->GetName() : TEXT("UNKNOWN ACTION")
-		, SourceObject ? *SourceObject->GetName() : TEXT("UNKNOWN SORUCE OBJECT"));
+		, SourceObject ? *SourceObject->GetName() : TEXT("UNKNOWN SOURCE OBJECT"));
 
-	EnqueueScheduledActionInstanceAtFrame_External(MoveTemp(ActionInstance), SourceId, LocalFrame);
+	EnqueueScheduledActionInstanceAtFrame_External(MoveTemp(ActionInstance), SourceId, LocalFrame, bReliable);
 }
 
-void UNetworkPhysicsComponent::EnqueueScheduledActionInstanceAtFrame_External(TInstancedStruct<FNetworkPhysicsActionPayload>&& ActionInstance, const uint32 SourceId, const int32 LocalFrame)
+void UNetworkPhysicsComponent::EnqueueScheduledActionInstanceAtFrame_External(TInstancedStruct<FNetworkPhysicsActionPayload>&& ActionInstance, const uint32 SourceId, const int32 LocalFrame, const bool bReliable)
 {
 	if (!NetworkPhysicsComponent_Internal)
 	{
@@ -2016,6 +2042,7 @@ void UNetworkPhysicsComponent::EnqueueScheduledActionInstanceAtFrame_External(TI
 	if (FAsyncNetworkPhysicsComponentInput* AsyncInput = NetworkPhysicsComponent_Internal->GetProducerInputData_External())
 	{
 		AsyncInput->LocalActionData.Add(MoveTemp(ActionInstance));
+		AsyncInput->LocalActionDataReliable.Add(bReliable);
 	}
 }
 
@@ -2127,6 +2154,17 @@ void UNetworkPhysicsComponent::SetPhysicsObject(Chaos::FConstPhysicsObjectHandle
 	}
 }
 
+void UNetworkPhysicsComponent::SetActionHandler(INetworkPhysicsActionHandler_Internal* InActionHandler)
+{
+	if (NetworkPhysicsComponent_Internal)
+	{
+		if (FAsyncNetworkPhysicsComponentInput* AsyncInput = NetworkPhysicsComponent_Internal->GetProducerInputData_External())
+		{
+			AsyncInput->ActionHandler_Internal = InActionHandler;
+		}
+	}
+}
+
 void UNetworkPhysicsComponent::UpdateAsyncComponent(const bool bFullUpdate)
 {
 	// Marshal data from Game Thread to Physics Thread
@@ -2377,6 +2415,10 @@ void FAsyncNetworkPhysicsComponent::ConsumeAsyncInput_Internal(const int32 Physi
 			{
 				ImplementationInterface_Internal = *AsyncInput->ImplementationInterface_Internal;
 			}
+			if (AsyncInput->ActionHandler_Internal.IsSet())
+			{
+				ActionHandler_Internal = *AsyncInput->ActionHandler_Internal;
+			}
 			if (AsyncInput->PhysicsObject.IsSet())
 			{
 				if (PhysicsObject == nullptr || PhysicsObject != *AsyncInput->PhysicsObject)
@@ -2584,9 +2626,10 @@ void FAsyncNetworkPhysicsComponent::ConsumeAsyncInput_Internal(const int32 Physi
 			// Receive Local Actions
 			if (AsyncInput->LocalActionData.Num() > 0 && !bTriggerResim) // bTriggerResim TEMPORARY, this function gets called twice, remove when we can add multiple of the same actions without duplication
 			{
-				for (const TInstancedStruct<FNetworkPhysicsActionPayload>& ActionInstance : AsyncInput->LocalActionData)
+				for (int32 ActionIdx = 0; ActionIdx < AsyncInput->LocalActionData.Num(); ++ActionIdx)
 				{
-					RegisterAction_Internal(ActionInstance, /*bWasReceivedFromReplication*/ false);
+					const bool bActionReliable = AsyncInput->LocalActionDataReliable.IsValidIndex(ActionIdx) ? AsyncInput->LocalActionDataReliable[ActionIdx] : false;
+					RegisterAction_Internal(AsyncInput->LocalActionData[ActionIdx], /*bWasReceivedFromReplication*/ false, bActionReliable);
 				}
 			}
 		}
@@ -2595,7 +2638,8 @@ void FAsyncNetworkPhysicsComponent::ConsumeAsyncInput_Internal(const int32 Physi
 
 void FAsyncNetworkPhysicsComponent::RegisterAction_Internal(
 	const TInstancedStruct<FNetworkPhysicsActionPayload>& NewActionInstance,
-	const bool bWasReceivedFromReplication)
+	const bool bWasReceivedFromReplication,
+	const bool bReliable)
 {
 	using EActionAuthorStyle = FNetworkPhysicsActionPayload::EActionAuthorStyle;
 
@@ -2678,7 +2722,7 @@ void FAsyncNetworkPhysicsComponent::RegisterAction_Internal(
 			}
 
 			// Store the proposal exactly as received.
-			PendingProposedActions.AddAction(NewActionInstance, ActionLocalFrame, TickOffset);
+			PendingProposedActions.AddAction(NewActionInstance, ActionLocalFrame, TickOffset, /*OverrideActionId*/0u, /*bClearActionId=*/false, /*bCheckDuplicateInSameFrame=*/false, bReliable);
 
 #if !(UE_BUILD_SHIPPING || UE_BUILD_TEST)
 			if (PhysicsReplicationCVars::ResimulationCVars::bActionsEnableDebugLogs)
@@ -2693,7 +2737,7 @@ void FAsyncNetworkPhysicsComponent::RegisterAction_Internal(
 			// Incoming replicated actions on the client from the server are authoritative confirmed actions.
 
 			// Store the confirmed action exactly as received.
-			ConfirmedActions.AddAction(NewActionInstance, ActionLocalFrame, TickOffset);
+			ConfirmedActions.AddAction(NewActionInstance, ActionLocalFrame, TickOffset, /*OverrideActionId*/0u, /*bClearActionId=*/false, /*bCheckDuplicateInSameFrame=*/false, bReliable);
 
 			// Retire the predicted actions equivalent to the received confirmed action
 			const bool bPredictedActionDeactivated = PredictedActions.DeactivateEquivalentActionsInRange(NewActionInstance, (ActionLocalFrame - PredictionFrameWindow), (ActionLocalFrame + PredictionFrameWindow));
@@ -2751,7 +2795,7 @@ void FAsyncNetworkPhysicsComponent::RegisterAction_Internal(
 						if (bRejectedProposal)
 						{
 							// Override proposal but keep ActionId
-							ConfirmedActions.AddAction(NewActionInstance, ActionLocalFrame, TickOffset, ProposedActionInstance->Get().ActionId);
+							ConfirmedActions.AddAction(NewActionInstance, ActionLocalFrame, TickOffset, ProposedActionInstance->Get().ActionId, /*bClearActionId=*/false, /*bCheckDuplicateInSameFrame=*/false, bReliable);
 
 #if !(UE_BUILD_SHIPPING || UE_BUILD_TEST)
 							if (PhysicsReplicationCVars::ResimulationCVars::bActionsEnableDebugLogs)
@@ -2765,7 +2809,7 @@ void FAsyncNetworkPhysicsComponent::RegisterAction_Internal(
 							const int32 ProposedLocalFrame = PendingProposedActions.GetLocalFrame(MatchingProposalLocation);
 
 							// Store the confirmed proposed action exactly as received from client (keeping the ActionId that the client set)
-							ConfirmedActions.AddAction(*ProposedActionInstance, ProposedLocalFrame, TickOffset);
+							ConfirmedActions.AddAction(*ProposedActionInstance, ProposedLocalFrame, TickOffset, /*OverrideActionId*/0u, /*bClearActionId=*/false, /*bCheckDuplicateInSameFrame=*/false, bReliable);
 
 #if !(UE_BUILD_SHIPPING || UE_BUILD_TEST)
 							if (PhysicsReplicationCVars::ResimulationCVars::bActionsEnableDebugLogs)
@@ -2790,7 +2834,7 @@ void FAsyncNetworkPhysicsComponent::RegisterAction_Internal(
 
 			// Store the confirmed action and set a unique ActionId
 			const uint32 NewActionId = GetUniqueActionId_Internal();
-			ConfirmedActions.AddAction(NewActionInstance, ActionLocalFrame, TickOffset, NewActionId);
+			ConfirmedActions.AddAction(NewActionInstance, ActionLocalFrame, TickOffset, NewActionId, /*bClearActionId=*/false, /*bCheckDuplicateInSameFrame=*/false, bReliable);
 			
 #if !(UE_BUILD_SHIPPING || UE_BUILD_TEST)
 			if (PhysicsReplicationCVars::ResimulationCVars::bActionsEnableDebugLogs)
@@ -2849,7 +2893,7 @@ void FAsyncNetworkPhysicsComponent::RegisterAction_Internal(
 
 				const uint32 NewActionId = GetUniqueActionId_Internal();
 				// Store the proposal action and set a unique ActionId
-				PendingProposedActions.AddAction(NewActionInstance, ActionLocalFrame, TickOffset, NewActionId);
+				PendingProposedActions.AddAction(NewActionInstance, ActionLocalFrame, TickOffset, NewActionId, /*bClearActionId=*/false, /*bCheckDuplicateInSameFrame=*/false, bReliable);
 
 #if !(UE_BUILD_SHIPPING || UE_BUILD_TEST)
 				if (PhysicsReplicationCVars::ResimulationCVars::bActionsEnableDebugLogs)
@@ -2876,7 +2920,7 @@ void FAsyncNetworkPhysicsComponent::RegisterAction_Internal(
 
 void FAsyncNetworkPhysicsComponent::ApplyActionsForFrame(const int32 LocalFrame)
 {
-	if (!ImplementationInterface_Internal || !GetRigidSolver() || !GetEvolution())
+	if (!ActionHandler_Internal || !GetRigidSolver() || !GetEvolution())
 	{
 		return;
 	}
@@ -2912,7 +2956,7 @@ void FAsyncNetworkPhysicsComponent::ApplyActionsForFrame(const int32 LocalFrame)
 			}
 #endif
 
-			ImplementationInterface_Internal->ApplyAction(ConfirmedActionEntry.ActionInstance);
+			ActionHandler_Internal->ApplyAction_Internal(ConfirmedActionEntry.ActionInstance);
 
 			ConfirmedActionEntry.bApplied = true;
 		}
@@ -2939,7 +2983,7 @@ void FAsyncNetworkPhysicsComponent::ApplyActionsForFrame(const int32 LocalFrame)
 			}
 #endif
 
-			ImplementationInterface_Internal->ApplyAction(PredictedActionEntry.ActionInstance);
+			ActionHandler_Internal->ApplyAction_Internal(PredictedActionEntry.ActionInstance);
 
 			PredictedActionEntry.bApplied = true;
 			PredictedActionEntry.bActive = false;
@@ -3397,12 +3441,25 @@ void FAsyncNetworkPhysicsComponent::OnPostProcessInputs_Internal(const int32 Phy
 
 	if (bIsSolverResim == false && PhysicsReplicationCVars::ResimulationCVars::bRecordStatePostSolve == false)
 	{
-		// Marshal inputs and states from PT to GT for networking
-		FAsyncNetworkPhysicsComponentOutput& AsyncOutput = GetAsyncOutput_Internal();
-		SendInputData_Internal(AsyncOutput, PhysicsStep);
-		SendStateData_Internal(AsyncOutput, PhysicsStep);
-		SendActionData_Internal(AsyncOutput, PhysicsStep);
-		FinalizeOutputData_Internal();
+		const bool bHasInputOrState = (InputHistory != nullptr) || (StateHistory != nullptr);
+		const bool bHasActions = ConfirmedActions.IsInitialized() || PendingProposedActions.IsInitialized();
+
+		if (bHasInputOrState || bHasActions)
+		{
+			FAsyncNetworkPhysicsComponentOutput& AsyncOutput = GetAsyncOutput_Internal();
+
+			if (bHasInputOrState)
+			{
+				SendInputData_Internal(AsyncOutput, PhysicsStep);
+				SendStateData_Internal(AsyncOutput, PhysicsStep);
+			}
+			if (bHasActions)
+			{
+				SendActionData_Internal(AsyncOutput, PhysicsStep);
+			}
+
+			FinalizeOutputData_Internal();
+		}
 	}
 }
 
@@ -3542,12 +3599,25 @@ void FAsyncNetworkPhysicsComponent::OnPostSolve_Internal()
 
 	if (bIsSolverResim == false && PhysicsReplicationCVars::ResimulationCVars::bRecordStatePostSolve == true)
 	{
-		// Marshal inputs and states from PT to GT for networking
-		FAsyncNetworkPhysicsComponentOutput& AsyncOutput = GetAsyncOutput_Internal();
-		SendInputData_Internal(AsyncOutput, PhysicsStep);
-		SendStateData_Internal(AsyncOutput, NextPhysicsStep); // We cache state at the end of the solve which is the state that next frame starts with
-		SendActionData_Internal(AsyncOutput, PhysicsStep);
-		FinalizeOutputData_Internal();
+		const bool bHasInputOrState = (InputHistory != nullptr) || (StateHistory != nullptr);
+		const bool bHasActions = ConfirmedActions.IsInitialized() || PendingProposedActions.IsInitialized();
+
+		if (bHasInputOrState || bHasActions)
+		{
+			FAsyncNetworkPhysicsComponentOutput& AsyncOutput = GetAsyncOutput_Internal();
+
+			if (bHasInputOrState)
+			{
+				SendInputData_Internal(AsyncOutput, PhysicsStep);
+				SendStateData_Internal(AsyncOutput, NextPhysicsStep); // We cache state at the end of the solve which is the state that next frame starts with
+			}
+			if (bHasActions)
+			{
+				SendActionData_Internal(AsyncOutput, PhysicsStep);
+			}
+
+			FinalizeOutputData_Internal();
+		}
 	}
 }
 
@@ -3696,11 +3766,11 @@ void FAsyncNetworkPhysicsComponent::SendActionData_Internal(FAsyncNetworkPhysics
 {
 	if (IsServer())
 	{
-		ConfirmedActions.GetActionsToReplicate(AsyncOutput.ActionData);
+		ConfirmedActions.GetActionsToReplicate(AsyncOutput.UnreliableActionData, AsyncOutput.ReliableActionData);
 	}
 	else if (IsLocallyControlled()) // Autonomous Proxy Client
 	{
-		PendingProposedActions.GetActionsToReplicate(AsyncOutput.ActionData);
+		PendingProposedActions.GetActionsToReplicate(AsyncOutput.UnreliableActionData, AsyncOutput.ReliableActionData);
 	}
 }
 
@@ -4148,12 +4218,12 @@ int32 FAsyncNetworkPhysicsComponent::GetLatestReceivedStateFrame() const
 	return INDEX_NONE;
 }
 
-void FAsyncNetworkPhysicsComponent::EnqueueImmediateActionInstance_Internal(TInstancedStruct<FNetworkPhysicsActionPayload>&& ActionInstance, const uint32 SourceId)
+void FAsyncNetworkPhysicsComponent::EnqueueImmediateActionInstance_Internal(TInstancedStruct<FNetworkPhysicsActionPayload>&& ActionInstance, const uint32 SourceId, const bool bReliable)
 {
-	EnqueueScheduledActionInstanceAtFrame_Internal(MoveTemp(ActionInstance), SourceId, -1);
+	EnqueueScheduledActionInstanceAtFrame_Internal(MoveTemp(ActionInstance), SourceId, -1, bReliable);
 }
 
-void FAsyncNetworkPhysicsComponent::EnqueueScheduledActionInstance_Internal(TInstancedStruct<FNetworkPhysicsActionPayload>&& ActionInstance, const uint32 SourceId, const float DelaySeconds)
+void FAsyncNetworkPhysicsComponent::EnqueueScheduledActionInstance_Internal(TInstancedStruct<FNetworkPhysicsActionPayload>&& ActionInstance, const uint32 SourceId, const float DelaySeconds, const bool bReliable)
 {
 	int32 LocalFrame = INDEX_NONE;
 	if (Chaos::FPBDRigidsSolver* RigidSolver = GetRigidSolver())
@@ -4168,10 +4238,10 @@ void FAsyncNetworkPhysicsComponent::EnqueueScheduledActionInstance_Internal(TIns
 			}
 		}
 	}
-	EnqueueScheduledActionInstanceAtFrame_Internal(MoveTemp(ActionInstance), SourceId, LocalFrame);
+	EnqueueScheduledActionInstanceAtFrame_Internal(MoveTemp(ActionInstance), SourceId, LocalFrame, bReliable);
 }
 
-void FAsyncNetworkPhysicsComponent::EnqueueScheduledActionInstanceAtFrame_Internal(TInstancedStruct<FNetworkPhysicsActionPayload>&& ActionInstance, const uint32 SourceId, const int32 LocalFrame)
+void FAsyncNetworkPhysicsComponent::EnqueueScheduledActionInstanceAtFrame_Internal(TInstancedStruct<FNetworkPhysicsActionPayload>&& ActionInstance, const uint32 SourceId, const int32 LocalFrame, const bool bReliable)
 {
 	FNetworkPhysicsActionPayload& Action = ActionInstance.GetMutable<FNetworkPhysicsActionPayload>();
 	Action.SourceId = SourceId;
@@ -4190,7 +4260,7 @@ void FAsyncNetworkPhysicsComponent::EnqueueScheduledActionInstanceAtFrame_Intern
 	}
 #endif
 
-	RegisterAction_Internal(ActionInstance, /*bWasReceivedFromReplication*/ false);
+	RegisterAction_Internal(ActionInstance, /*bWasReceivedFromReplication*/ false, bReliable);
 }
 
 namespace UE::NetworkPhysicsUtils
diff --git a/Engine/Source/Runtime/Engine/Public/Physics/NetworkPhysicsComponent.h b/Engine/Source/Runtime/Engine/Public/Physics/NetworkPhysicsComponent.h
--- a/Engine/Source/Runtime/Engine/Public/Physics/NetworkPhysicsComponent.h
+++ b/Engine/Source/Runtime/Engine/Public/Physics/NetworkPhysicsComponent.h
@@ -692,6 +692,7 @@ struct FActionHistoryEntry
 	bool bActive = false;
 	bool bApplied = false;
 	bool bHasBeenQueuedForReplication = false;
+	bool bReliable = false;
 
 	void Reset()
 	{
@@ -699,6 +700,7 @@ struct FActionHistoryEntry
 		bActive = false;
 		bApplied = false;
 		bHasBeenQueuedForReplication = false;
+		bReliable = false;
 	}
 };
 
@@ -1060,11 +1062,11 @@ struct FActionHistory
 
 	bool HasEquivalentActionInRange(const FActionInstance& CandidateActionInstance, const int32 MinLocalFrame, const int32 MaxLocalFrame, const bool bIncludeApplied = false, const bool bCheckSourceEquivalent = false) const
 	{
-		return FindEquivalentActionInRange(CandidateActionInstance, MinLocalFrame, MaxLocalFrame, bIncludeApplied).IsValid();
+		return FindEquivalentActionInRange(CandidateActionInstance, MinLocalFrame, MaxLocalFrame, bIncludeApplied, bCheckSourceEquivalent).IsValid();
 	}
 
 	bool AddAction(const FActionInstance& ActionInstance, const int32 LocalFrame, const int32 TickOffset
-		, const uint16 OverrideActionId = 0u, const bool bClearActionId = false, const bool bCheckDuplicateInSameFrame = false)
+		, const uint16 OverrideActionId = 0u, const bool bClearActionId = false, const bool bCheckDuplicateInSameFrame = false, const bool bInReliable = false)
 	{
 		if (IsInitialized() == false || ActionInstance.IsValid() == false)
 		{
@@ -1104,6 +1106,7 @@ struct FActionHistory
 		NewEntry.bActive = true;
 		NewEntry.bApplied = false;
 		NewEntry.bHasBeenQueuedForReplication = false;
+		NewEntry.bReliable = bInReliable;
 		return true;
 	}
 
@@ -1249,7 +1252,8 @@ struct FActionHistory
 		return bDeactivatedAny;
 	}
 
-	void GetActionsToReplicate(TArray<FActionInstance>& OutActions, const bool bResetOutArray = true, const bool bMarkAsQueuedForReplication = true, const bool bOnlyActiveActions = true)
+	/** Split actions to replicate into unreliable and reliable output arrays based on FActionHistoryEntry::bReliable */
+	void GetActionsToReplicate(TArray<FActionInstance>& OutUnreliableActions, TArray<FActionInstance>& OutReliableActions, const bool bResetOutArray = true, const bool bMarkAsQueuedForReplication = true, const bool bOnlyActiveActions = true)
 	{
 		if (IsInitialized() == false)
 		{
@@ -1258,7 +1262,8 @@ struct FActionHistory
 
 		if (bResetOutArray == true)
 		{
-			OutActions.Reset();
+			OutUnreliableActions.Reset();
+			OutReliableActions.Reset();
 		}
 
 		struct FPendingReplicationLocation
@@ -1297,7 +1302,6 @@ struct FActionHistory
 				}
 
 				FPendingReplicationLocation& Location = PendingReplicationLocations.Emplace_GetRef();
-
 				Location.FrameIndex = FrameIndex;
 				Location.EntryIndex = EntryIndex;
 				Location.LocalFrame = FrameData.LocalFrame;
@@ -1311,12 +1315,10 @@ struct FActionHistory
 				{
 					return A.LocalFrame < B.LocalFrame;
 				}
-
 				if (A.FrameIndex != B.FrameIndex)
 				{
 					return A.FrameIndex < B.FrameIndex;
 				}
-
 				return A.EntryIndex < B.EntryIndex;
 			});
 
@@ -1325,8 +1327,14 @@ struct FActionHistory
 			FActionHistoryFrameData& FrameData = FrameHistory[Location.FrameIndex];
 			FActionHistoryEntry& Entry = FrameData.Entries[Location.EntryIndex];
 
-			// Populate OUT array with the sorted actions to be replicated
-			OutActions.Add(Entry.ActionInstance);
+			if (Entry.bReliable)
+			{
+				OutReliableActions.Add(Entry.ActionInstance);
+			}
+			else
+			{
+				OutUnreliableActions.Add(Entry.ActionInstance);
+			}
 
 			if (bMarkAsQueuedForReplication == true)
 			{
@@ -1805,6 +1813,16 @@ public:
 	virtual void ValidateInput_External(InputType& Input) const = 0;
 };
 
+// Physics Thread Action Handler Interface API
+class INetworkPhysicsActionHandler_Internal
+{
+protected:
+	virtual ~INetworkPhysicsActionHandler_Internal() = default;
+
+public:
+	virtual void ApplyAction_Internal(const TInstancedStruct<FNetworkPhysicsActionPayload>& ActionInstance) = 0;
+};
+
 // Physics Thread Input and State Interface API
 class INetworkPhysicsInputState_Internal
 {
@@ -1818,8 +1836,6 @@ public:
 
 	virtual void BuildState(FNetworkPhysicsPayload& State) const = 0;
 	virtual void ApplyState(const FNetworkPhysicsPayload& State) = 0;
-
-	virtual void ApplyAction(const TInstancedStruct<FNetworkPhysicsActionPayload>& ActionInstance) = 0;
 };
 
 template<typename InputType, typename StateType>
@@ -1836,7 +1852,6 @@ private:
 
 	void BuildState(FNetworkPhysicsPayload& State) const override { BuildState_Internal(static_cast<StateType&>(State)); }
 	void ApplyState(const FNetworkPhysicsPayload& State) override { ApplyState_Internal(static_cast<const StateType&>(State)); }
-	void ApplyAction(const TInstancedStruct<FNetworkPhysicsActionPayload>& ActionInstance) override { ApplyAction_Internal(ActionInstance); }
 
 public:
 	/** Populate the input struct with current input data */
@@ -1855,9 +1870,6 @@ public:
 
 	/** Apply state struct to implementation */
 	virtual void ApplyState_Internal(const StateType& State) = 0;
-
-	/** Apply action struct to implementation */
-	virtual void ApplyAction_Internal(const TInstancedStruct<FNetworkPhysicsActionPayload>& ActionInstance) { }
 };
 
 /** END Iris Compliant Data */
@@ -2319,87 +2331,99 @@ public:
 	ENGINE_API bool IsNetworkPhysicsTickOffsetAssigned() const;
 
 	/** Enqueue a TInstancedStruct with an action based on FNetworkPhysicsActionPayload, to be executed as soon as possible.
-	* @param SourceObject is the object that is the source of this action, not necessarily the object that the action should run on. */
-	ENGINE_API void EnqueueImmediateActionInstance_External(TInstancedStruct<FNetworkPhysicsActionPayload>&& ActionInstance, const UObject* SourceObject);
-	
+	* @param SourceObject is the object that is the source of this action, not necessarily the object that the action should run on.
+	* @param bReliable if true, the action will be sent via a reliable RPC instead of unreliable property replication. */
+	ENGINE_API void EnqueueImmediateActionInstance_External(TInstancedStruct<FNetworkPhysicsActionPayload>&& ActionInstance, const UObject* SourceObject, const bool bReliable = false);
+
 	/** Enqueue a TInstancedStruct with an action based on FNetworkPhysicsActionPayload, to be executed as soon as possible.
 	* @param SourceId can be left at 0 for actions triggered by the owner of this implementation, else if triggered by an external actor pass in a network stable hash.
+	* @param bReliable if true, the action will be sent via a reliable RPC instead of unreliable property replication.
 	* Note: On the game thread a network stable hash can be retrieved from replicated UObjects or level baked UObjects (for example actors) via UE::NetworkPhysicsUtils::GetNetworkStableHash_External(UObject); */
-	ENGINE_API void EnqueueImmediateActionInstance_External(TInstancedStruct<FNetworkPhysicsActionPayload>&& ActionInstance, const uint32 SourceId);
+	ENGINE_API void EnqueueImmediateActionInstance_External(TInstancedStruct<FNetworkPhysicsActionPayload>&& ActionInstance, const uint32 SourceId, const bool bReliable = false);
 
 	/** Enqueue an action based on FNetworkPhysicsActionPayload, to be executed as soon as possible.
-	* @param SourceObject is the object that is the source of this action, not necessarily the object that the action should run on. */
+	* @param SourceObject is the object that is the source of this action, not necessarily the object that the action should run on.
+	* @param bReliable if true, the action will be sent via a reliable RPC instead of unreliable property replication. */
 	template<typename TActionType>
-	void EnqueueImmediateAction_External(TActionType& Action, const UObject* SourceObject)
+	void EnqueueImmediateAction_External(TActionType& Action, const UObject* SourceObject, const bool bReliable = false)
 	{
 		static_assert(TIsDerivedFrom<TActionType, FNetworkPhysicsActionPayload>::Value, "EnqueueImmediateAction_External - Action needs to derive from FNetworkPhysicsActionPayload");
-		EnqueueImmediateActionInstance_External(TInstancedStruct<FNetworkPhysicsActionPayload>::Make<TActionType>(Action), SourceObject);
+		EnqueueImmediateActionInstance_External(TInstancedStruct<FNetworkPhysicsActionPayload>::Make<TActionType>(Action), SourceObject, bReliable);
 	}
 
 	/** Enqueue an action based on FNetworkPhysicsActionPayload, to be executed as soon as possible.
 	* @param SourceId can be left at 0 for actions triggered by the owner of this implementation, else if triggered by an external actor pass in a network stable hash.
+	* @param bReliable if true, the action will be sent via a reliable RPC instead of unreliable property replication.
 	* Note: On the game thread a network stable hash can be retrieved from replicated UObjects or level baked UObjects (for example actors) via UE::NetworkPhysicsUtils::GetNetworkStableHash_External(UObject); */
 	template<typename TActionType>
-	void EnqueueImmediateAction_External(TActionType& Action, const uint32 SourceId)
+	void EnqueueImmediateAction_External(TActionType& Action, const uint32 SourceId, const bool bReliable = false)
 	{
 		static_assert(TIsDerivedFrom<TActionType, FNetworkPhysicsActionPayload>::Value, "EnqueueImmediateAction_External - Action needs to derive from FNetworkPhysicsActionPayload");
-		EnqueueImmediateActionInstance_External(TInstancedStruct<FNetworkPhysicsActionPayload>::Make<TActionType>(Action), SourceId);
+		EnqueueImmediateActionInstance_External(TInstancedStruct<FNetworkPhysicsActionPayload>::Make<TActionType>(Action), SourceId, bReliable);
 	}
 
 	/** Enqueue a TInstancedStruct with an action based on FNetworkPhysicsActionPayload, to be executed after a delay of @param DelaySeconds.
-	* @param SourceObject is the object that is the source of this action, not necessarily the object that the action should run on. */
-	ENGINE_API void EnqueueScheduledActionInstance_External(TInstancedStruct<FNetworkPhysicsActionPayload>&& ActionInstance, const UObject* SourceObject, const float DelaySeconds);
-	
+	* @param SourceObject is the object that is the source of this action, not necessarily the object that the action should run on.
+	* @param bReliable if true, the action will be sent via a reliable RPC instead of unreliable property replication. */
+	ENGINE_API void EnqueueScheduledActionInstance_External(TInstancedStruct<FNetworkPhysicsActionPayload>&& ActionInstance, const UObject* SourceObject, const float DelaySeconds, const bool bReliable = false);
+
 	/** Enqueue a TInstancedStruct with an action based on FNetworkPhysicsActionPayload, to be executed after a delay of @param DelaySeconds.
 	* @param SourceId can be left at 0 for actions triggered by the owner of this implementation, else if triggered by an external actor pass in a network stable hash.
+	* @param bReliable if true, the action will be sent via a reliable RPC instead of unreliable property replication.
 	* Note: On the game thread a network stable hash can be retrieved from replicated UObjects or level baked UObjects (for example actors) via UE::NetworkPhysicsUtils::GetNetworkStableHash_External(UObject); */
-	ENGINE_API void EnqueueScheduledActionInstance_External(TInstancedStruct<FNetworkPhysicsActionPayload>&& ActionInstance, const uint32 SourceId, const float DelaySeconds);
+	ENGINE_API void EnqueueScheduledActionInstance_External(TInstancedStruct<FNetworkPhysicsActionPayload>&& ActionInstance, const uint32 SourceId, const float DelaySeconds, const bool bReliable = false);
 
 	/** Enqueue an action based on FNetworkPhysicsActionPayload, to be executed after a delay of @param DelaySeconds.
-	* @param SourceObject is the object that is the source of this action, not necessarily the object that the action should run on. */
+	* @param SourceObject is the object that is the source of this action, not necessarily the object that the action should run on.
+	* @param bReliable if true, the action will be sent via a reliable RPC instead of unreliable property replication. */
 	template<typename TActionType>
-	void EnqueueScheduledAction_External(TActionType& Action, const UObject* SourceObject, const float DelaySeconds)
+	void EnqueueScheduledAction_External(TActionType& Action, const UObject* SourceObject, const float DelaySeconds, const bool bReliable = false)
 	{
 		static_assert(TIsDerivedFrom<TActionType, FNetworkPhysicsActionPayload>::Value, "EnqueueScheduledAction_External - Action needs to derive from FNetworkPhysicsActionPayload");
-		EnqueueScheduledActionInstance_External(TInstancedStruct<FNetworkPhysicsActionPayload>::Make<TActionType>(Action), SourceObject, DelaySeconds);
+		EnqueueScheduledActionInstance_External(TInstancedStruct<FNetworkPhysicsActionPayload>::Make<TActionType>(Action), SourceObject, DelaySeconds, bReliable);
 	}
 
 	/** Enqueue an action based on FNetworkPhysicsActionPayload, to be executed after a delay of @param DelaySeconds.
 	* @param SourceId can be left at 0 for actions triggered by the owner of this implementation, else if triggered by an external actor pass in a network stable hash.
+	* @param bReliable if true, the action will be sent via a reliable RPC instead of unreliable property replication.
 	* Note: On the game thread a network stable hash can be retrieved from replicated UObjects or level baked UObjects (for example actors) via UE::NetworkPhysicsUtils::GetNetworkStableHash_External(UObject); */
 	template<typename TActionType>
-	void EnqueueScheduledAction_External(TActionType& Action, const uint32 SourceId, const float DelaySeconds)
+	void EnqueueScheduledAction_External(TActionType& Action, const uint32 SourceId, const float DelaySeconds, const bool bReliable = false)
 	{
 		static_assert(TIsDerivedFrom<TActionType, FNetworkPhysicsActionPayload>::Value, "EnqueueScheduledAction_External - Action needs to derive from FNetworkPhysicsActionPayload");
-		EnqueueScheduledActionInstance_External(TInstancedStruct<FNetworkPhysicsActionPayload>::Make<TActionType>(Action), SourceId, DelaySeconds);
+		EnqueueScheduledActionInstance_External(TInstancedStruct<FNetworkPhysicsActionPayload>::Make<TActionType>(Action), SourceId, DelaySeconds, bReliable);
 	}
 
 	/** Enqueue a TInstancedStruct with an action based on FNetworkPhysicsActionPayload, to be executed at @param LocalFrame.
-	* @param SourceObject is the object that is the source of this action, not necessarily the object that the action should run on. */
-	ENGINE_API void EnqueueScheduledActionInstanceAtFrame_External(TInstancedStruct<FNetworkPhysicsActionPayload>&& ActionInstance, const UObject* SourceObject, const int32 LocalFrame);
-	
+	* @param SourceObject is the object that is the source of this action, not necessarily the object that the action should run on.
+	* @param bReliable if true, the action will be sent via a reliable RPC instead of unreliable property replication. */
+	ENGINE_API void EnqueueScheduledActionInstanceAtFrame_External(TInstancedStruct<FNetworkPhysicsActionPayload>&& ActionInstance, const UObject* SourceObject, const int32 LocalFrame, const bool bReliable = false);
+
 	/** Enqueue a TInstancedStruct with an action based on FNetworkPhysicsActionPayload, to be executed at @param LocalFrame.
 	* @param SourceId can be left at 0 for actions triggered by the owner of this implementation, else if triggered by an external actor pass in a network stable hash.
+	* @param bReliable if true, the action will be sent via a reliable RPC instead of unreliable property replication.
 	* Note: On the game thread a network stable hash can be retrieved from replicated UObjects or level baked UObjects (for example actors) via UE::NetworkPhysicsUtils::GetNetworkStableHash_External(UObject); */
-	ENGINE_API void EnqueueScheduledActionInstanceAtFrame_External(TInstancedStruct<FNetworkPhysicsActionPayload>&& ActionInstance, const uint32 SourceId, const int32 LocalFrame);
+	ENGINE_API void EnqueueScheduledActionInstanceAtFrame_External(TInstancedStruct<FNetworkPhysicsActionPayload>&& ActionInstance, const uint32 SourceId, const int32 LocalFrame, const bool bReliable = false);
 
 	/** Enqueue an action based on FNetworkPhysicsActionPayload, to be executed at @param LocalFrame.
-	* @param SourceObject is the object that is the source of this action, not necessarily the object that the action should run on. */
+	* @param SourceObject is the object that is the source of this action, not necessarily the object that the action should run on.
+	* @param bReliable if true, the action will be sent via a reliable RPC instead of unreliable property replication. */
 	template<typename TActionType>
-	void EnqueueScheduledActionAtFrame_External(TActionType& Action, const UObject* SourceObject, const int32 LocalFrame)
+	void EnqueueScheduledActionAtFrame_External(TActionType& Action, const UObject* SourceObject, const int32 LocalFrame, const bool bReliable = false)
 	{
 		static_assert(TIsDerivedFrom<TActionType, FNetworkPhysicsActionPayload>::Value, "EnqueueScheduledActionAtFrame_External - Action needs to derive from FNetworkPhysicsActionPayload");
-		EnqueueScheduledActionInstanceAtFrame_External(TInstancedStruct<FNetworkPhysicsActionPayload>::Make<TActionType>(Action), SourceObject, LocalFrame);
+		EnqueueScheduledActionInstanceAtFrame_External(TInstancedStruct<FNetworkPhysicsActionPayload>::Make<TActionType>(Action), SourceObject, LocalFrame, bReliable);
 	}
 
 	/** Enqueue a TInstancedStruct with an action based on FNetworkPhysicsActionPayload, to be executed at @param LocalFrame.
 	* @param SourceId can be left at 0 for actions triggered by the owner of this implementation, else if triggered by an external actor pass in a network stable hash.
+	* @param bReliable if true, the action will be sent via a reliable RPC instead of unreliable property replication.
 	* Note: On the game thread a network stable hash can be retrieved from replicated UObjects or level baked UObjects (for example actors) via UE::NetworkPhysicsUtils::GetNetworkStableHash_External(UObject); */
 	template<typename TActionType>
-	void EnqueueScheduledActionAtFrame_External(TActionType& Action, const uint32 SourceId, const int32 LocalFrame)
+	void EnqueueScheduledActionAtFrame_External(TActionType& Action, const uint32 SourceId, const int32 LocalFrame, const bool bReliable = false)
 	{
 		static_assert(TIsDerivedFrom<TActionType, FNetworkPhysicsActionPayload>::Value, "EnqueueScheduledActionAtFrame_External - Action needs to derive from FNetworkPhysicsActionPayload");
-		EnqueueScheduledActionInstanceAtFrame_External(TInstancedStruct<FNetworkPhysicsActionPayload>::Make<TActionType>(Action), SourceId, LocalFrame);
+		EnqueueScheduledActionInstanceAtFrame_External(TInstancedStruct<FNetworkPhysicsActionPayload>::Make<TActionType>(Action), SourceId, LocalFrame, bReliable);
 	}
 
 	/** Mark this as controlled through locally relayed inputs rather than controlled as a pawn through a player controller.
@@ -2522,9 +2546,17 @@ protected:
 	UFUNCTION(Server, unreliable)
 	ENGINE_API void ServerReceiveInputCollection(const FNetworkPhysicsDataCollection& ClientInputCollection);
 
-	// Server RPC to receive actions from client
+	// Server RPC to receive unreliable actions from client
 	UFUNCTION(Server, unreliable)
-	ENGINE_API void ServerReceiveActionCollection(const FNetworkPhysicsActionCollection& ClientActionCollection);
+	ENGINE_API void ServerReceiveUnreliableActionCollection(const FNetworkPhysicsActionCollection& ClientActionCollection);
+
+	// Server RPC to receive reliable actions from client
+	UFUNCTION(Server, reliable)
+	ENGINE_API void ServerReceiveReliableActionCollection(const FNetworkPhysicsActionCollection& ClientActionCollection);
+
+	// Client RPC to receive reliable actions from server
+	UFUNCTION(NetMulticast, reliable)
+	ENGINE_API void MulticastReceiveReliableActionCollection(const FNetworkPhysicsActionCollection& ServerActionCollection);
 
 	// repnotify for inputs on owner client
 	UFUNCTION()
@@ -2563,6 +2595,9 @@ protected:
 #endif
 
 private:
+	/** Marshal received action data (from RPCs or property replication) to the physics thread */
+	void ReceiveActionCollectionData(const FNetworkPhysicsActionCollection& ActionCollection);
+
 	/** Send network data from marshaled data */
 	void NetworkMarshaledData();
 
@@ -2656,15 +2691,10 @@ public:
 	template<typename Input, typename State>
 	void CreateDataHistory(TNetworkPhysicsInputState_Internal<Input, State>* Implementation_Internal, TNetworkPhysicsInputState_External<Input, State>* Implementation_External = nullptr);
 
-	/** Set implementation interface, only needed if using legacy CreateDataHistory() but want to use Actions (which only support interface implementation), else use CreateDataHistory() that takes the interface.
-	* Note: If input and state are not being used, implement with FNetworkPhysicsPayload. */
-	template<typename Input, typename State>
-	void SetImplementationInteraces(TNetworkPhysicsInputState_Internal<Input, State>* Implementation_Internal, TNetworkPhysicsInputState_External<Input, State>* Implementation_External = nullptr)
-	{
-		ImplementationInterface_Internal = Implementation_Internal;
-		ImplementationInterface_External = Implementation_External;
-		UpdateAsyncComponent(true);
-	}
+	/** Set action handler interface for receiving and applying actions on the physics thread.
+	* This can be used standalone (without registering input/state types) for action-only components,
+	* or alongside CreateDataHistory() for components that use both input/state and actions. */
+	ENGINE_API void SetActionHandler(INetworkPhysicsActionHandler_Internal* InActionHandler);
 
 private:
 
@@ -2682,6 +2712,8 @@ private:
 	INetworkPhysicsInputState_External* ImplementationInterface_External = nullptr;
 
 	bool bIsUsingLegacyData = false;
+	bool bHasRegisteredInputState = false;
+	bool bHasPendingActionsToReplicate = false;
 
 	// Network Physics Component data internal to the physics thread
 	FAsyncNetworkPhysicsComponent* NetworkPhysicsComponent_Internal;
@@ -2758,6 +2790,7 @@ template<typename Input, typename State>
 void UNetworkPhysicsComponent::CreateDataHistory(TNetworkPhysicsInputState_Internal<Input, State>* Implementation_Internal, TNetworkPhysicsInputState_External<Input, State>* Implementation_External)
 {
 	bIsUsingLegacyData = false;
+	bHasRegisteredInputState = true;
 
 	ImplementationInterface_Internal = Implementation_Internal;
 	ImplementationInterface_External = Implementation_External;
@@ -2772,6 +2805,7 @@ template<typename PhysicsTraits>
 inline void UNetworkPhysicsComponent::CreateDataHistory(UActorComponent* HistoryComponent)
 {
 	bIsUsingLegacyData = true;
+	bHasRegisteredInputState = true;
 
 	InputHelper = MakeUnique<TNetworkPhysicsDataHelper<typename PhysicsTraits::InputsType, /*bLegacyData*/ true>>();
 	StateHelper = MakeUnique<TNetworkPhysicsDataHelper<typename PhysicsTraits::StatesType, /*bLegacyData*/ true>>();
@@ -2819,6 +2853,7 @@ template<class InputsType>
 inline void UNetworkPhysicsComponent::CreateInputHistory(UActorComponent* HistoryComponent)
 {
 	bIsUsingLegacyData = true;
+	bHasRegisteredInputState = true;
 
 	InputHelper = MakeUnique<TNetworkPhysicsDataHelper<InputsType, /*bLegacyData*/ true>>();
 	
@@ -2872,6 +2907,7 @@ struct FAsyncNetworkPhysicsComponentInput : public Chaos::FSimCallbackInput
 	TOptional<bool> bCompareStateToTriggerRewindIncludeSimProxies;
 	TOptional<bool> bCompareInputToTriggerRewind;
 	TOptional<TWeakPtr<const FNetworkPhysicsSettingsData>> SettingsComponent;
+	TOptional<INetworkPhysicsActionHandler_Internal*> ActionHandler_Internal;
 
 	TUniquePtr<Chaos::FBaseRewindHistory> InputData;
 	TUniquePtr<Chaos::FBaseRewindHistory> StateData;
@@ -2880,6 +2916,7 @@ struct FAsyncNetworkPhysicsComponentInput : public Chaos::FSimCallbackInput
 	TArray<TUniquePtr<Chaos::FBaseRewindHistory>> StateDataImportant;
 
 	TArray<TInstancedStruct<FNetworkPhysicsActionPayload>> LocalActionData;
+	TArray<bool> LocalActionDataReliable; // Parallel array to LocalActionData, indicates if each action should be replicated reliably
 	TArray<TInstancedStruct<FNetworkPhysicsActionPayload>> ReplicatedActionData;
 
 	void Reset()
@@ -2901,6 +2938,7 @@ struct FAsyncNetworkPhysicsComponentInput : public Chaos::FSimCallbackInput
 		bCompareStateToTriggerRewindIncludeSimProxies.Reset();
 		bCompareInputToTriggerRewind.Reset();
 		SettingsComponent.Reset();
+		ActionHandler_Internal.Reset();
 
 		if (InputData)
 		{
@@ -2918,6 +2956,7 @@ struct FAsyncNetworkPhysicsComponentInput : public Chaos::FSimCallbackInput
 		StateDataImportant.Reset();
 
 		LocalActionData.Reset();
+		LocalActionDataReliable.Reset();
 		ReplicatedActionData.Reset();
 	}
 };
@@ -2933,7 +2972,8 @@ struct FAsyncNetworkPhysicsComponentOutput : public Chaos::FSimCallbackOutput
 	TArray<TUniquePtr<Chaos::FBaseRewindHistory>> InputDataImportant;
 	TArray<TUniquePtr<Chaos::FBaseRewindHistory>> StateDataImportant;
 
-	TArray<TInstancedStruct<FNetworkPhysicsActionPayload>> ActionData;
+	TArray<TInstancedStruct<FNetworkPhysicsActionPayload>> UnreliableActionData;
+	TArray<TInstancedStruct<FNetworkPhysicsActionPayload>> ReliableActionData;
 
 	void Reset()
 	{
@@ -2953,7 +2993,8 @@ struct FAsyncNetworkPhysicsComponentOutput : public Chaos::FSimCallbackOutput
 		InputDataImportant.Reset();
 		StateDataImportant.Reset();
 
-		ActionData.Reset();
+		UnreliableActionData.Reset();
+		ReliableActionData.Reset();
 	}
 };
 
@@ -3022,47 +3063,53 @@ public:
 
 	/** Enqueue a TInstancedStruct with an action based on FNetworkPhysicsActionPayload, to be executed as soon as possible.
 	* @param SourceId can be left at 0 for actions triggered by the owner of this implementation, else if triggered by an external actor pass in a network stable hash.
+	* @param bReliable if true, the action will be sent via a reliable RPC instead of unreliable property replication.
 	* Note: On the game thread a network stable hash can be retrieved from replicated UObjects or level baked UObjects (for example actors) via UE::NetworkPhysicsUtils::GetNetworkStableHash_External(UObject); */
-	ENGINE_API void EnqueueImmediateActionInstance_Internal(TInstancedStruct<FNetworkPhysicsActionPayload>&& ActionInstance, const uint32 SourceId);
+	ENGINE_API void EnqueueImmediateActionInstance_Internal(TInstancedStruct<FNetworkPhysicsActionPayload>&& ActionInstance, const uint32 SourceId, const bool bReliable = false);
 
 	/** Enqueue an action based on FNetworkPhysicsActionPayload, to be executed as soon as possible.
 	* @param SourceId can be left at 0 for actions triggered by the owner of this implementation, else if triggered by an external actor pass in a network stable hash.
+	* @param bReliable if true, the action will be sent via a reliable RPC instead of unreliable property replication.
 	* Note: On the game thread a network stable hash can be retrieved from replicated UObjects or level baked UObjects (for example actors) via UE::NetworkPhysicsUtils::GetNetworkStableHash_External(UObject); */
 	template<typename TActionType>
-	void EnqueueImmediateAction_Internal(TActionType& Action, const uint32 SourceId)
+	void EnqueueImmediateAction_Internal(TActionType& Action, const uint32 SourceId, const bool bReliable = false)
 	{
 		static_assert(TIsDerivedFrom<TActionType, FNetworkPhysicsActionPayload>::Value, "EnqueueImmediateAction_Internal - Action needs to derive from FNetworkPhysicsActionPayload");
-		EnqueueImmediateActionInstance_Internal(TInstancedStruct<FNetworkPhysicsActionPayload>::Make<TActionType>(Action), SourceId);
+		EnqueueImmediateActionInstance_Internal(TInstancedStruct<FNetworkPhysicsActionPayload>::Make<TActionType>(Action), SourceId, bReliable);
 	}
 
 	/** Enqueue a TInstancedStruct with an action based on FNetworkPhysicsActionPayload, to be executed after a delay of @param DelaySeconds.
 	* @param SourceId can be left at 0 for actions triggered by the owner of this implementation, else if triggered by an external actor pass in a network stable hash.
+	* @param bReliable if true, the action will be sent via a reliable RPC instead of unreliable property replication.
 	* Note: On the game thread a network stable hash can be retrieved from replicated UObjects or level baked UObjects (for example actors) via UE::NetworkPhysicsUtils::GetNetworkStableHash_External(UObject); */
-	ENGINE_API void EnqueueScheduledActionInstance_Internal(TInstancedStruct<FNetworkPhysicsActionPayload>&& ActionInstance, const uint32 SourceId, const float DelaySeconds);
+	ENGINE_API void EnqueueScheduledActionInstance_Internal(TInstancedStruct<FNetworkPhysicsActionPayload>&& ActionInstance, const uint32 SourceId, const float DelaySeconds, const bool bReliable = false);
 
 	/** Enqueue an action based on FNetworkPhysicsActionPayload, to be executed after a delay of @param DelaySeconds.
 	* @param SourceId can be left at 0 for actions triggered by the owner of this implementation, else if triggered by an external actor pass in a network stable hash.
+	* @param bReliable if true, the action will be sent via a reliable RPC instead of unreliable property replication.
 	* Note: On the game thread a network stable hash can be retrieved from replicated UObjects or level baked UObjects (for example actors) via UE::NetworkPhysicsUtils::GetNetworkStableHash_External(UObject); */
 	template<typename TActionType>
-	void EnqueueScheduledAction_Internal(TActionType& Action, const uint32 SourceId, const float DelaySeconds)
+	void EnqueueScheduledAction_Internal(TActionType& Action, const uint32 SourceId, const float DelaySeconds, const bool bReliable = false)
 	{
 		static_assert(TIsDerivedFrom<TActionType, FNetworkPhysicsActionPayload>::Value, "EnqueueScheduledAction_Internal - Action needs to derive from FNetworkPhysicsActionPayload");
-		EnqueueScheduledActionInstance_Internal(TInstancedStruct<FNetworkPhysicsActionPayload>::Make<TActionType>(Action), SourceId, DelaySeconds);
+		EnqueueScheduledActionInstance_Internal(TInstancedStruct<FNetworkPhysicsActionPayload>::Make<TActionType>(Action), SourceId, DelaySeconds, bReliable);
 	}
 
 	/** Enqueue a TInstancedStruct with an action based on FNetworkPhysicsActionPayload, to be executed at @param LocalFrame.
 	* @param SourceId can be left at 0 for actions triggered by the owner of this implementation, else if triggered by an external actor pass in a network stable hash.
+	* @param bReliable if true, the action will be sent via a reliable RPC instead of unreliable property replication.
 	* Note: On the game thread a network stable hash can be retrieved from replicated UObjects or level baked UObjects (for example actors) via UE::NetworkPhysicsUtils::GetNetworkStableHash_External(UObject); */
-	ENGINE_API void EnqueueScheduledActionInstanceAtFrame_Internal(TInstancedStruct<FNetworkPhysicsActionPayload>&& ActionInstance, const uint32 SourceId, const int32 LocalFrame);
+	ENGINE_API void EnqueueScheduledActionInstanceAtFrame_Internal(TInstancedStruct<FNetworkPhysicsActionPayload>&& ActionInstance, const uint32 SourceId, const int32 LocalFrame, const bool bReliable = false);
 
 	/** Enqueue an action based on FNetworkPhysicsActionPayload, to be executed at @param LocalFrame
 	* @param SourceId can be left at 0 for actions triggered by the owner of this implementation, else if triggered by an external actor pass in a network stable hash.
+	* @param bReliable if true, the action will be sent via a reliable RPC instead of unreliable property replication.
 	* Note:	On the game thread a network stable hash can be retrieved from replicated UObjects or level baked UObjects (for example actors) via UE::NetworkPhysicsUtils::GetNetworkStableHash_External(UObject); */
 	template<typename TActionType>
-	void EnqueueScheduledActionAtFrame_Internal(TActionType& Action, const uint32 SourceId, const int32 LocalFrame)
+	void EnqueueScheduledActionAtFrame_Internal(TActionType& Action, const uint32 SourceId, const int32 LocalFrame, const bool bReliable = false)
 	{
 		static_assert(TIsDerivedFrom<TActionType, FNetworkPhysicsActionPayload>::Value, "EnqueueScheduledActionAtFrame_Internal - Action needs to derive from FNetworkPhysicsActionPayload");
-		EnqueueScheduledActionInstanceAtFrame_Internal(TInstancedStruct<FNetworkPhysicsActionPayload>::Make<TActionType>(Action), SourceId, LocalFrame);
+		EnqueueScheduledActionInstanceAtFrame_Internal(TInstancedStruct<FNetworkPhysicsActionPayload>::Make<TActionType>(Action), SourceId, LocalFrame, bReliable);
 	}
 
 private:
@@ -3093,7 +3140,7 @@ private:
 	void ConsumeAsyncInput_Internal(const int32 PhysicsStep, const bool bTriggerResim);
 
 	// Add new action and handle related proposed actions
-	void RegisterAction_Internal(const TInstancedStruct<FNetworkPhysicsActionPayload>& NewActionInstance, const bool bWasReceivedFromReplication);
+	void RegisterAction_Internal(const TInstancedStruct<FNetworkPhysicsActionPayload>& NewActionInstance, const bool bWasReceivedFromReplication, const bool bReliable = false);
 
 	// Apply actions
 	void ApplyActionsForFrame(const int32 LocalFrame);
@@ -3164,6 +3211,9 @@ private:
 	// Implementation of input / state interface
 	INetworkPhysicsInputState_Internal* ImplementationInterface_Internal;
 
+	// Implementation of action handler interface
+	INetworkPhysicsActionHandler_Internal* ActionHandler_Internal = nullptr;
+
 	// Root components physics object
 	Chaos::FConstPhysicsObjectHandle PhysicsObject;
 
