Improved handling for an empty mesh description, usually due to corrupt data.

#Synced-CL 52440501

diff --git a/Engine/Source/Developer/MeshBuilder/Private/StaticMeshBuilder.cpp b/Engine/Source/Developer/MeshBuilder/Private/StaticMeshBuilder.cpp
--- a/Engine/Source/Developer/MeshBuilder/Private/StaticMeshBuilder.cpp
+++ b/Engine/Source/Developer/MeshBuilder/Private/StaticMeshBuilder.cpp
@@ -1425,7 +1425,6 @@ bool FStaticMeshBuilder::BuildMeshVertexPositions(
 
 	FMeshDescription MeshDescription;
 	const bool bIsMeshDescriptionValid = SourceModel.CloneMeshDescription(MeshDescription);
-	check(bIsMeshDescriptionValid);
 
 	if (MeshDescription.IsEmpty())
 	{
diff --git a/Engine/Source/Runtime/Engine/Private/StaticMesh.cpp b/Engine/Source/Runtime/Engine/Private/StaticMesh.cpp
--- a/Engine/Source/Runtime/Engine/Private/StaticMesh.cpp
+++ b/Engine/Source/Runtime/Engine/Private/StaticMesh.cpp
@@ -7173,7 +7173,7 @@ int32 UStaticMesh::GetNumUVChannels(int32 LODIndex) const
 	int32 NumUVChannels = 0;
 #if WITH_EDITORONLY_DATA
 	FMeshDescription* MeshDescription = GetMeshDescription(LODIndex);
-	if (MeshDescription)
+	if (MeshDescription && !MeshDescription->IsEmpty())
 	{
 		FStaticMeshConstAttributes Attributes(*MeshDescription);
 		NumUVChannels = Attributes.GetVertexInstanceUVs().GetNumChannels();
diff --git a/Engine/Source/Runtime/Engine/Private/StaticMeshSourceData.cpp b/Engine/Source/Runtime/Engine/Private/StaticMeshSourceData.cpp
--- a/Engine/Source/Runtime/Engine/Private/StaticMeshSourceData.cpp
+++ b/Engine/Source/Runtime/Engine/Private/StaticMeshSourceData.cpp
@@ -223,9 +223,9 @@ bool FStaticMeshSourceModel::LoadMeshDescription(FMeshDescription& OutMeshDescri
 	// If we have valid MeshDescriptionBulkData, unpack it and return it...
 	if (StaticMeshDescriptionBulkData->IsBulkDataValid())
 	{
-		// Unpack MeshDescription from the bulk data which was deserialized
+		// Unpack MeshDescription from the bulk data which was deserialized. If there is an error deserializing, then the mesh description will be empty.
 		StaticMeshDescriptionBulkData->GetBulkData().LoadMeshDescription(OutMeshDescription);
-		return true;
+		return !OutMeshDescription.IsEmpty();
 	}
 
 	// RawMeshBulkData should always be empty now (soon to be deprecated)
