MIR: On-demand input resolution for switch-like expressions

Switch-like expression nodes (StaticSwitch, FeatureLevelSwitch, QualitySwitch, etc.) now use MIR_OnDemandInputRequest metadata to only build the taken branch instead of all inputs. This speeds up material translation when large subgraphs are hidden behind switch like expressions as only the live value will be built.

#rnx
#Synced-CL 52470326

diff --git a/Engine/Source/Runtime/Engine/Private/Materials/MaterialExpressionsToMIR.cpp b/Engine/Source/Runtime/Engine/Private/Materials/MaterialExpressionsToMIR.cpp
--- a/Engine/Source/Runtime/Engine/Private/Materials/MaterialExpressionsToMIR.cpp
+++ b/Engine/Source/Runtime/Engine/Private/Materials/MaterialExpressionsToMIR.cpp
@@ -169,6 +169,7 @@
 #include "Materials/MaterialExpressionReflectionVectorWS.h"
 #include "Materials/MaterialExpressionRequiredSamplersSwitch.h"
 #include "Materials/MaterialExpressionReroute.h"
+#include "Materials/MaterialExpressionRerouteBase.h"
 #include "Materials/MaterialExpressionRgbToHsv.h"
 #include "Materials/MaterialExpressionRotateAboutAxis.h"
 #include "Materials/MaterialExpressionRotator.h"
@@ -494,7 +495,7 @@ static FMaterialParameterMetadata GetMaterialExpressionParameterMetadata(MIR::FE
 static FValueRef BuildPrimitiveUniformMaterialParameter(MIR::FEmitter& Em, UMaterialExpressionParameter* ParameterExpr)
 {
 	FMaterialParameterMetadata Metadata = GetMaterialExpressionParameterMetadata(Em, ParameterExpr);
-	if (Em.CurrentExpressionHasErrors())
+	if (Em.IsInvalid())
 	{
 		return Em.Poison();
 	}
@@ -601,7 +602,7 @@ void UMaterialExpressionStaticComponentMaskParameter::Build(MIR::FEmitter& Em)
 void UMaterialExpressionStaticSwitch::Build(MIR::FEmitter& Em)
 {
 	bool bCondition = Em.ToConstantBool(Em.InputDefaultBool(&Value, DefaultValue));
-	UE_MIR_CHECKPOINT(Em); // Make sure that evaluating bCondition didn't raise an error
+	UE_MIR_CHECKPOINT(Em); // Make sure that the bCondition value was available and didn't raise an error
 
 	Em.Output(0, Em.Input(bCondition ? &A : &B));
 }
@@ -613,11 +614,11 @@ void UMaterialExpressionStaticSwitchParameter::Build(MIR::FEmitter& Em)
 
 	check(Metadata.Value.Type == EMaterialParameterType::StaticSwitch);
 
-	// Deliberately fetch both inputs, so we throw an error to the user if either is not set, even though we are only returning one of the inputs.
-	FValueRef AValue = Em.Input(&A);
-	FValueRef BValue = Em.Input(&B);
-
-	Em.Output(0, Metadata.Value.Bool[0] ? AValue : BValue);
+	// Make sure both input pins are always connected to some expression even if we don't build one.
+	MaterialToMIR::CheckInputIsConnected(Em, A, TEXT("A"));
+	MaterialToMIR::CheckInputIsConnected(Em, B, TEXT("B"));
+	
+	Em.Output(0, Em.Input(Metadata.Value.Bool[0] ? &A : &B));
 }
 
 void UMaterialExpressionDataDrivenShaderPlatformInfoSwitch::Build(MIR::FEmitter& Em)
@@ -638,23 +639,28 @@ void UMaterialExpressionDistanceFieldsRenderingSwitch::Build(MIR::FEmitter& Em)
 	EShaderPlatform ShaderPlatform = Em.GetShaderPlatform();
 	bool bDistanceFieldsEnabled = IsMobilePlatform(ShaderPlatform) ? IsMobileDistanceFieldEnabled(ShaderPlatform) : IsUsingDistanceFields(ShaderPlatform);
 
-	// Deliberately fetch both inputs, so we throw an error to the user if either is not set, even though we are only returning one of the inputs.
-	FValueRef YesValue = Em.Input(&Yes);
-	FValueRef NoValue = Em.Input(&No);
-
-	Em.Output(0, bDistanceFieldsEnabled ? YesValue : NoValue);
+	// Make sure both input pins are always connected to some expression even if we don't build one.
+	MaterialToMIR::CheckInputIsConnected(Em, Yes, TEXT("Yes"));
+	MaterialToMIR::CheckInputIsConnected(Em, No, TEXT("No"));
+	
+	Em.Output(0, Em.Input(bDistanceFieldsEnabled ? &Yes : &No));
 }
 
 void UMaterialExpressionFeatureLevelSwitch::Build(MIR::FEmitter& Em)
 {
-	// Always fetch the Default input, so we throw an error to the user if it's not set, even if it doesn't end up being used.
-	FValueRef DefaultValue = Em.Input(&Default);
 	FValueRef Result = Em.TryInput(&Inputs[GetFeatureLevelToCompile(Em.GetShaderPlatform(), Em.GetFeatureLevel())]);
+	UE_MIR_CHECKPOINT(Em);
 
 	if (!Result)
 	{
-		Result = DefaultValue;
+		Result = Em.Input(&Default);
 	}
+	else
+	{
+		// Make sure the default pin is always connected to some expression even if we don't build it.
+		MaterialToMIR::CheckInputIsConnected(Em, Default, TEXT("Default"));
+	}
+
 	Em.Output(0, Result);
 }
 
@@ -672,8 +678,8 @@ void UMaterialExpressionQualitySwitch::Build(MIR::FEmitter& Em)
 		}
 		else
 		{
-			// Deliberately fetch the Default input, so we throw an error to the user if it's not set, even when it's not being used.
-			Em.Input(&Default);
+			// Make sure the default pin is always connected to some expression even if we don't build it.
+			MaterialToMIR::CheckInputIsConnected(Em, Default, TEXT("Default"));
 		}
 		Em.Output(0, Result);
 	}
@@ -721,19 +727,19 @@ void UMaterialExpressionShadingPathSwitch::Build(MIR::FEmitter& Em)
 	}
 	else
 	{
-		// Deliberately fetch the Default input, so we throw an error to the user if it's not set, even when it's not being used.
-		Em.Input(&Default);
+		// Make sure the default pin is always connected to some expression even if we don't build it.
+		MaterialToMIR::CheckInputIsConnected(Em, Default, TEXT("Default"));
 	}
 	Em.Output(0, Result);
 }
 
 void UMaterialExpressionVirtualTextureFeatureSwitch::Build(MIR::FEmitter& Em)
 {
-	// Deliberately fetch both inputs, so we throw an error to the user if either is not set, even though we are only returning one of the inputs.
-	FValueRef YesValue = Em.Input(&Yes);
-	FValueRef NoValue = Em.Input(&No);
+	// Make sure both input pins are always connected to some expression even if we don't build one.
+	MaterialToMIR::CheckInputIsConnected(Em, Yes, TEXT("Yes"));
+	MaterialToMIR::CheckInputIsConnected(Em, No, TEXT("No"));
 
-	Em.Output(0, UseVirtualTexturing(Em.GetShaderPlatform()) ? YesValue : NoValue);
+	Em.Output(0, Em.Input(UseVirtualTexturing(Em.GetShaderPlatform()) ? &Yes : &No));
 }
 
 static void EmitExternalCodeConditionalReplace(MIR::FEmitter& Em, FName ExternalCodeIdentifier, FExpressionInput& Default, FExpressionInput& Replace)
@@ -5372,7 +5378,7 @@ static FValueRef EmitProcessPhysicsField(MIR::FEmitter& Em, FExpressionInput* Po
 	}
 
 	Em.CheckFeatureLevelIsAtLeast(ERHIFeatureLevel::SM5);
-	if (Em.CurrentExpressionHasErrors())
+	if (Em.IsInvalid())
 	{
 		return Em.Poison();
 	}
diff --git a/Engine/Source/Runtime/Engine/Private/Materials/MaterialExpressionsToMIRCommon.cpp b/Engine/Source/Runtime/Engine/Private/Materials/MaterialExpressionsToMIRCommon.cpp
--- a/Engine/Source/Runtime/Engine/Private/Materials/MaterialExpressionsToMIRCommon.cpp
+++ b/Engine/Source/Runtime/Engine/Private/Materials/MaterialExpressionsToMIRCommon.cpp
@@ -8,6 +8,8 @@
 #include "Materials/MaterialIRExtern.h"
 #include "Materials/Material.h"
 #include "Materials/MaterialInterface.h"
+#include "Materials/MaterialExpressionReroute.h"
+#include "Materials/MaterialExpressionNamedReroute.h"
 
 namespace MaterialToMIR
 {
@@ -140,5 +142,56 @@ MIR::FValueRef EmitCustomPrimitiveDataFloat4(MIR::FEmitter& Em, int32 Index)
 		EmitCustomPrimitiveDataFloat1(Em, Index + 3));
 }
 
+static bool TryGetRerouteInput(UMaterialExpressionRerouteBase* Reroute, FExpressionInput& OutInput)
+{
+	if (UMaterialExpressionReroute* SimpleReroute = Cast<UMaterialExpressionReroute>(Reroute))
+	{
+		OutInput = SimpleReroute->Input;
+		return true;
+	}
+	if (UMaterialExpressionNamedRerouteDeclaration* Declaration = Cast<UMaterialExpressionNamedRerouteDeclaration>(Reroute))
+	{
+		OutInput = Declaration->Input;
+		return true;
+	}
+	if (UMaterialExpressionNamedRerouteUsage* Usage = Cast<UMaterialExpressionNamedRerouteUsage>(Reroute))
+	{
+		if (Usage->IsDeclarationValid())
+		{
+			OutInput = Usage->Declaration->Input;
+			return true;
+		}
+	}
+	return false;
+}
+
+void CheckInputIsConnected(MIR::FEmitter& Em, FExpressionInput& Input, const TCHAR* InputName)
+{
+	// Iteratively trace through reroute nodes to find the real connected expression.
+	FExpressionInput TracedInput = Input;
+	TArray<const UMaterialExpression*, TInlineAllocator<4>> Visited;
+	while (UMaterialExpressionRerouteBase* Reroute = Cast<UMaterialExpressionRerouteBase>(TracedInput.Expression))
+	{
+		if (Visited.Contains(Reroute))
+		{
+			TracedInput.Expression = nullptr;
+			break;
+		}
+		
+		Visited.Add(Reroute);
+
+		if (!TryGetRerouteInput(Reroute, TracedInput))
+		{
+			TracedInput.Expression = nullptr;
+			break;
+		}
+	}
+
+	if (!TracedInput.Expression)
+	{
+		Em.Errorf(TEXT("Input '%s' must be connected."), InputName);
+	}
+}
+
 } // MaterialToMIR
 #endif // WITH_EDITOR
diff --git a/Engine/Source/Runtime/Engine/Private/Materials/MaterialIREmitter.cpp b/Engine/Source/Runtime/Engine/Private/Materials/MaterialIREmitter.cpp
--- a/Engine/Source/Runtime/Engine/Private/Materials/MaterialIREmitter.cpp
+++ b/Engine/Source/Runtime/Engine/Private/Materials/MaterialIREmitter.cpp
@@ -368,7 +368,8 @@ void FEmitter::Error(FStringView Message)
 	}
 	
 	Module->Errors.Push(Error);
-	bCurrentExpressionHasErrors = true;
+
+	State = EState::Error;
 }
 
 /*--------------------------------- Type handling ----------------------------------*/
@@ -532,6 +533,14 @@ FType FEmitter::GetMaterialAggregateAttributeType(const UMaterialAggregate* Aggr
 
 FValueRef FEmitter::TryInput(const FExpressionInput* InInput)
 {
+	if (bExpressionRequestsInputsOnDemand)
+	{
+		if (!Internal::RequestInputValue(BuilderImpl, InInput))
+		{
+			State = EState::Pending;
+			return { FPoison::Get(), InInput };
+		}
+	}
 	return FValueRef{ Internal::FetchValueFromExpressionInput(BuilderImpl, InInput), InInput };
 }
 
@@ -815,7 +824,10 @@ bool FEmitter::ToConstantBool(FValueRef Value)
 
 FEmitter& FEmitter::Output(int32 OutputIndex, FValueRef Value)
 {
-	Output(Expression->GetOutput(OutputIndex), Value);
+	if (Value.IsValid())
+	{
+		Output(Expression->GetOutput(OutputIndex), Value);
+	}
 	return *this;
 }
 
@@ -1330,7 +1342,7 @@ void FEmitter::SetCustomOutputs(FStringView Name, TConstArrayView<FValueRef> Arg
 	}
 
 	// Avoid work if some error occurred.
-	if (CurrentExpressionHasErrors() || IsAnyNotValid(Args))
+	if (IsInvalid() || IsAnyNotValid(Args))
 	{
 		return;
 	}
diff --git a/Engine/Source/Runtime/Engine/Private/Materials/MaterialIRInternal.h b/Engine/Source/Runtime/Engine/Private/Materials/MaterialIRInternal.h
--- a/Engine/Source/Runtime/Engine/Private/Materials/MaterialIRInternal.h
+++ b/Engine/Source/Runtime/Engine/Private/Materials/MaterialIRInternal.h
@@ -8,11 +8,16 @@
 
 namespace MIR::Internal {
 
-// Returns the material value tpye of the specified UTexture or URuntimeVirtualTexture.
+// Returns the material value type of the specified UTexture or URuntimeVirtualTexture.
 EMaterialValueType GetTextureMaterialValueType(const UObject* TextureObject);
 
 // @massimo.tristano expose an internal "Context" struct between the builder and the emitter, and move these methods
 // to the context itself
+
+// Requests the value for the given input, pushing its source expression as a dependency if not yet built.
+// Returns true if the value is already available, false if a new dependency was pushed.
+bool RequestInputValue(FMaterialIRModuleBuilderImpl* Builder, const FExpressionInput* Input);
+
 // Returns the value flowing into given expression input (previously set through `BindValueToExpressionInput`).
 FValue* FetchValueFromExpressionInput(FMaterialIRModuleBuilderImpl* Builder, const FExpressionInput* Input);
 
diff --git a/Engine/Source/Runtime/Engine/Private/Materials/MaterialIRModuleBuilder.cpp b/Engine/Source/Runtime/Engine/Private/Materials/MaterialIRModuleBuilder.cpp
--- a/Engine/Source/Runtime/Engine/Private/Materials/MaterialIRModuleBuilder.cpp
+++ b/Engine/Source/Runtime/Engine/Private/Materials/MaterialIRModuleBuilder.cpp
@@ -190,7 +190,7 @@ struct FMaterialIRModuleBuilderImpl
 		{
 			FMaterialInputDescription Input;
 			GetExpressionInputDescription(MP_MaterialAttributes, Input);
-			PushDependency(Input.Input->Expression);
+			PushDependencyOnStack(Input.Input->Expression);
 		}
 		else
 		{
@@ -214,7 +214,7 @@ struct FMaterialIRModuleBuilderImpl
 				GetExpressionInputDescription(Property, Input);
 
 				// Push the connected expression to this material attribute input as a dependency, if any.
-				PushDependency(Input.Input->Expression);
+				PushDependencyOnStack(Input.Input->Expression);
 			}
 		}
 
@@ -226,7 +226,7 @@ struct FMaterialIRModuleBuilderImpl
 		{
 			if (Cast<UMaterialExpressionCustomOutput>(Expression) && !Cast<UMaterialExpressionVertexInterpolator>(Expression))
 			{
-				PushDependency(Expression);
+				PushDependencyOnStack(Expression);
 			}
 		}
 	}
@@ -350,9 +350,11 @@ struct FMaterialIRModuleBuilderImpl
 
 	void BuildTopMaterialExpression()
 	{
+		static const FName OnDemandInputRequestName(TEXT("MIR_OnDemandInputRequest"));
+
 		FAnalysisContext& CurrContext = AnalysisContextStack.Last();
 		Emitter.Expression = CurrContext.ExpressionStack.Last();
-
+	
 		// If expression is clean, nothing to be done.
 		if (CurrContext.BuiltExpressions.Contains(Emitter.Expression))
 		{
@@ -360,50 +362,33 @@ struct FMaterialIRModuleBuilderImpl
 			return;
 		}
 
-		// Push to the expression stack all dependencies that still need to be analyzed.
-		for (FExpressionInputIterator It{ Emitter.Expression }; It; ++It)
-		{
-			PushDependency(It.Input->Expression);
-		}
+		// Set the current expression to the emitter
+		Emitter.State = MIR::FEmitter::EState::None;
+		Emitter.bExpressionRequestsInputsOnDemand = Emitter.Expression->GetClass()->HasMetaData(OnDemandInputRequestName);
 
-		// Named reroute usage nodes should add their declaration as dependency, as they
-		// do not have any direct input (they simply forward to the declaration input).
-		if (auto NamedReroute = Cast<UMaterialExpressionNamedRerouteUsage>(Emitter.Expression))
+		// Auto-request all inputs unless the expression opts out via MIR_OnDemandInputRequest meta flag.
+		if (!Emitter.bExpressionRequestsInputsOnDemand)
 		{
-			if (NamedReroute->IsDeclarationValid())
+			for (FExpressionInputIterator It{ Emitter.Expression }; It; ++It)
 			{
-				PushDependency(NamedReroute->Declaration.Get());
+				(void)::MIR::Internal::RequestInputValue(this, It.Input);
 			}
-		}
-
-		// If on top of the stack there's a different expression, we have a dependency to analyze first.
-		if (CurrContext.ExpressionStack.Last() != Emitter.Expression)
-		{
-			return;
-		}
 
-		// Take the top expression out of the stack as ready for analysis. Also mark it as built.
-		CurrContext.ExpressionStack.Pop();
-		CurrContext.BuiltExpressions.Add(Emitter.Expression);
-
-		// Flow the value into this expression's inputs from their connected outputs.
-		for (FExpressionInputIterator It{ Emitter.Expression}; It; ++It)
-		{
-			FlowValueThroughConnection(It.Input);
-		}
-		
-		// Usage reroute nodes should forward the value coming in from their declaration input (since
-		// they don't have an input of their own).
-		if (auto NamedReroute = Cast<UMaterialExpressionNamedRerouteUsage>(Emitter.Expression))
-		{
-			if (NamedReroute->IsDeclarationValid())
+			// If top expression on the stack isn't this, some new dependency was pushed on it.
+			if (CurrContext.ExpressionStack.Last() != Emitter.Expression)
 			{
-				FlowValueThroughConnection(&NamedReroute->Declaration->Input);
+				return;
 			}
 		}
-
+		
+		// @massimo.tristano The call mechanism should be moved to the emitter itself by pushing a "context".
 		if (auto Call = Cast<UMaterialExpressionMaterialFunctionCall>(Emitter.Expression))
 		{
+			// Pop and mark built before PushFunctionCall, which may Emplace onto AnalysisContextStack
+			// and invalidate the CurrContext reference.
+			CurrContext.ExpressionStack.Pop();
+			CurrContext.BuiltExpressions.Add(Emitter.Expression);
+
 			// Function calls are handled internally as they manipulate the analysis context stack.
 			PushFunctionCall(Call);
 		}
@@ -415,8 +400,18 @@ struct FMaterialIRModuleBuilderImpl
 			// emit IR values for its outputs (which will flow into connected expressions inputs).
 			Emitter.Expression->Build(Emitter);
 
+			// If the expression has a pending dependency, defer building it.
+			if (Emitter.State == MIR::FEmitter::EState::Pending)
+			{
+				return;
+			}
+
 			// Populate the insight information about this expression pins.
 			AddExpressionConnectionInsights(Emitter.Expression);
+
+			// Take the top expression out of the stack as ready for analysis. Also mark it as built.
+			CurrContext.ExpressionStack.Pop();
+			CurrContext.BuiltExpressions.Add(Emitter.Expression);
 		}
 	}
 
@@ -565,8 +560,9 @@ struct FMaterialIRModuleBuilderImpl
 		}
 	}
 	
-	// Pushes an expression dependency on this context expression stack.
-	void PushDependency(UMaterialExpression* Expression)
+	// Pushes the given expression onto the build stack if it hasn't been built yet.
+	// Returns true if already resolved (null or already built), false if pushed.
+	bool PushDependencyOnStack(UMaterialExpression* Expression)
 	{
 		FAnalysisContext& CurrContext = AnalysisContextStack.Last();
 
@@ -574,11 +570,14 @@ struct FMaterialIRModuleBuilderImpl
 		if (Expression && !CurrContext.BuiltExpressions.Contains(Expression))
 		{
 			CurrContext.ExpressionStack.Push(Expression);
+			return false;
 		}
+
+		return true;
 	}
 
 	// Flows the value from connected output to specified input.
-	void FlowValueThroughConnection(FExpressionInput* Input)
+	void FlowValueThroughConnection(const FExpressionInput* Input)
 	{
 		FAnalysisContext& CurrContext = AnalysisContextStack.Last();
 
@@ -595,6 +594,8 @@ struct FMaterialIRModuleBuilderImpl
 		}
 	}
 
+	// @massimo.tristano Function call could be optimized by having the caller flag which outputs are actually not used, so that we can skip
+	// pushing those function outputs dependencies
 	void PushFunctionCall(UMaterialExpressionMaterialFunctionCall* Call)
 	{
 		MIR::TTemporaryArray<MIR::FValue*> CallInputValues{ Call->FunctionInputs.Num() };
@@ -620,7 +621,7 @@ struct FMaterialIRModuleBuilderImpl
 		}
 
 		// If some error occurred (e.g. some function input wasn't linked in) early out.
-		if (Emitter.CurrentExpressionHasErrors())
+		if (Emitter.IsInvalid())
 		{
 			return;
 		}
@@ -1459,6 +1460,18 @@ struct FMaterialIRModuleBuilderImpl
 
 namespace MIR::Internal {
 
+	bool RequestInputValue(FMaterialIRModuleBuilderImpl* Builder, const FExpressionInput* Input)
+	{
+		// Safe for disconnected inputs: PushDependencyOnStack returns true for null expressions,
+		// and FlowValueThroughConnection is a no-op when there is no connected output.
+		if (Builder->PushDependencyOnStack(Input->Expression))
+		{
+			Builder->FlowValueThroughConnection(Input);
+			return true;
+		}
+		return false;
+	}
+
 	FValue* FetchValueFromExpressionInput(FMaterialIRModuleBuilderImpl* Builder, const FExpressionInput* Input)
 	{
 		return Builder->AnalysisContextStack.Last().GetInputValue(Input);
diff --git a/Engine/Source/Runtime/Engine/Private/Materials/MaterialShared.cpp b/Engine/Source/Runtime/Engine/Private/Materials/MaterialShared.cpp
--- a/Engine/Source/Runtime/Engine/Private/Materials/MaterialShared.cpp
+++ b/Engine/Source/Runtime/Engine/Private/Materials/MaterialShared.cpp
@@ -423,7 +423,7 @@ FExpressionInput FExpressionInput::GetTracedInput() const
 	return *this;
 }
 
-FExpressionOutput* FExpressionInput::GetConnectedOutput()
+FExpressionOutput* FExpressionInput::GetConnectedOutput() const
 {
 	return IsConnected() ? &Expression->GetOutputs()[OutputIndex] : nullptr;
 }
diff --git a/Engine/Source/Runtime/Engine/Public/MaterialExpressionIO.h b/Engine/Source/Runtime/Engine/Public/MaterialExpressionIO.h
--- a/Engine/Source/Runtime/Engine/Public/MaterialExpressionIO.h
+++ b/Engine/Source/Runtime/Engine/Public/MaterialExpressionIO.h
@@ -94,7 +94,7 @@ struct FExpressionInput
 	}
 
 	/** Returns the connected expression output if this input is connected to one,	nullptr otherwise. */
-	FExpressionOutput* GetConnectedOutput();
+	FExpressionOutput* GetConnectedOutput() const;
 
 #endif // WITH_EDITOR
 };
diff --git a/Engine/Source/Runtime/Engine/Public/Materials/MaterialExpressionDataDrivenShaderPlatformInfoSwitch.h b/Engine/Source/Runtime/Engine/Public/Materials/MaterialExpressionDataDrivenShaderPlatformInfoSwitch.h
--- a/Engine/Source/Runtime/Engine/Public/Materials/MaterialExpressionDataDrivenShaderPlatformInfoSwitch.h
+++ b/Engine/Source/Runtime/Engine/Public/Materials/MaterialExpressionDataDrivenShaderPlatformInfoSwitch.h
@@ -32,7 +32,7 @@ struct FDataDrivenShaderPlatformInfoInput
 		TEnumAsByte<EDataDrivenShaderPlatformInfoCondition> PropertyCondition = EDataDrivenShaderPlatformInfoCondition::COND_True;
 };
 
-UCLASS(collapsecategories, hidecategories = Object, MinimalAPI)
+UCLASS(collapsecategories, hidecategories = Object, MinimalAPI, meta=(MIR_OnDemandInputRequest))
 class UMaterialExpressionDataDrivenShaderPlatformInfoSwitch : public UMaterialExpression
 {
 	GENERATED_UCLASS_BODY()
diff --git a/Engine/Source/Runtime/Engine/Public/Materials/MaterialExpressionDistanceFieldsRenderingSwitch.h b/Engine/Source/Runtime/Engine/Public/Materials/MaterialExpressionDistanceFieldsRenderingSwitch.h
--- a/Engine/Source/Runtime/Engine/Public/Materials/MaterialExpressionDistanceFieldsRenderingSwitch.h
+++ b/Engine/Source/Runtime/Engine/Public/Materials/MaterialExpressionDistanceFieldsRenderingSwitch.h
@@ -9,7 +9,7 @@
 #include "MaterialExpressionDistanceFieldsRenderingSwitch.generated.h"
 
 /** Material output expression to switch logic according to whether distance fields renderering is supported on this project and feature level. */
-UCLASS(collapsecategories, hidecategories = Object)
+UCLASS(collapsecategories, hidecategories = Object, meta=(MIR_OnDemandInputRequest))
 class UMaterialExpressionDistanceFieldsRenderingSwitch : public UMaterialExpression
 {
 	GENERATED_BODY()
diff --git a/Engine/Source/Runtime/Engine/Public/Materials/MaterialExpressionFeatureLevelSwitch.h b/Engine/Source/Runtime/Engine/Public/Materials/MaterialExpressionFeatureLevelSwitch.h
--- a/Engine/Source/Runtime/Engine/Public/Materials/MaterialExpressionFeatureLevelSwitch.h
+++ b/Engine/Source/Runtime/Engine/Public/Materials/MaterialExpressionFeatureLevelSwitch.h
@@ -13,7 +13,7 @@
 
 enum EShaderPlatform : uint16;
 
-UCLASS(collapsecategories, hidecategories=Object, MinimalAPI)
+UCLASS(collapsecategories, hidecategories=Object, MinimalAPI, meta=(MIR_OnDemandInputRequest))
 class UMaterialExpressionFeatureLevelSwitch : public UMaterialExpression
 {
 	GENERATED_BODY()
diff --git a/Engine/Source/Runtime/Engine/Public/Materials/MaterialExpressionIfThenElse.h b/Engine/Source/Runtime/Engine/Public/Materials/MaterialExpressionIfThenElse.h
--- a/Engine/Source/Runtime/Engine/Public/Materials/MaterialExpressionIfThenElse.h
+++ b/Engine/Source/Runtime/Engine/Public/Materials/MaterialExpressionIfThenElse.h
@@ -10,7 +10,7 @@
 #include "MaterialExpressionIfThenElse.generated.h"
 
 
-UCLASS(collapsecategories, hidecategories = Object, MinimalAPI, meta=(NewMaterialTranslator))
+UCLASS(collapsecategories, hidecategories = Object, MinimalAPI, meta=(NewMaterialTranslator, MIR_OnDemandInputRequest))
 class UMaterialExpressionIfThenElse : public UMaterialExpression
 {
 	GENERATED_BODY()
@@ -28,8 +28,8 @@ public:
 	
 	//~ Begin UMaterialExpression Interface
 #if WITH_EDITOR
-	virtual int32 Compile(class FMaterialCompiler* Compiler, int32 OutputIndex) override;
 	virtual void Build(MIR::FEmitter& Emitter) override;
+	virtual int32 Compile(class FMaterialCompiler* Compiler, int32 OutputIndex) override;
 	virtual void GetCaption(TArray<FString>& OutCaptions) const override;
 	virtual EMaterialValueType GetInputValueType(int32 InputIndex) override;
 	virtual EMaterialValueType GetOutputValueType(int32 OutputIndex) override;
diff --git a/Engine/Source/Runtime/Engine/Public/Materials/MaterialExpressionNamedReroute.h b/Engine/Source/Runtime/Engine/Public/Materials/MaterialExpressionNamedReroute.h
--- a/Engine/Source/Runtime/Engine/Public/Materials/MaterialExpressionNamedReroute.h
+++ b/Engine/Source/Runtime/Engine/Public/Materials/MaterialExpressionNamedReroute.h
@@ -93,7 +93,7 @@ private:
 	void MakeNameUnique();
 };
 
-UCLASS(collapsecategories, hidecategories=Object, DisplayName = "Named Reroute Usage", MinimalAPI)
+UCLASS(collapsecategories, hidecategories=Object, DisplayName = "Named Reroute Usage", MinimalAPI, meta=(MIR_OnDemandInputRequest))
 class UMaterialExpressionNamedRerouteUsage : public UMaterialExpressionNamedRerouteBase
 {
 	GENERATED_BODY()
diff --git a/Engine/Source/Runtime/Engine/Public/Materials/MaterialExpressionQualitySwitch.h b/Engine/Source/Runtime/Engine/Public/Materials/MaterialExpressionQualitySwitch.h
--- a/Engine/Source/Runtime/Engine/Public/Materials/MaterialExpressionQualitySwitch.h
+++ b/Engine/Source/Runtime/Engine/Public/Materials/MaterialExpressionQualitySwitch.h
@@ -10,7 +10,7 @@
 #include "Materials/MaterialExpression.h"
 #include "MaterialExpressionQualitySwitch.generated.h"
 
-UCLASS(collapsecategories, hidecategories=Object, MinimalAPI)
+UCLASS(collapsecategories, hidecategories=Object, MinimalAPI, meta=(MIR_OnDemandInputRequest))
 class UMaterialExpressionQualitySwitch : public UMaterialExpression
 {
 	GENERATED_BODY()
diff --git a/Engine/Source/Runtime/Engine/Public/Materials/MaterialExpressionRequiredSamplersSwitch.h b/Engine/Source/Runtime/Engine/Public/Materials/MaterialExpressionRequiredSamplersSwitch.h
--- a/Engine/Source/Runtime/Engine/Public/Materials/MaterialExpressionRequiredSamplersSwitch.h
+++ b/Engine/Source/Runtime/Engine/Public/Materials/MaterialExpressionRequiredSamplersSwitch.h
@@ -11,7 +11,7 @@
 #include "RHIFeatureLevel.h"
 #include "MaterialExpressionRequiredSamplersSwitch.generated.h"
 
-UCLASS(MinimalAPI)
+UCLASS(MinimalAPI, meta=(MIR_OnDemandInputRequest))
 class UMaterialExpressionRequiredSamplersSwitch : public UMaterialExpression
 {
 	GENERATED_BODY()
diff --git a/Engine/Source/Runtime/Engine/Public/Materials/MaterialExpressionShadingPathSwitch.h b/Engine/Source/Runtime/Engine/Public/Materials/MaterialExpressionShadingPathSwitch.h
--- a/Engine/Source/Runtime/Engine/Public/Materials/MaterialExpressionShadingPathSwitch.h
+++ b/Engine/Source/Runtime/Engine/Public/Materials/MaterialExpressionShadingPathSwitch.h
@@ -13,7 +13,7 @@
 
 enum EShaderPlatform : uint16;
 
-UCLASS(collapsecategories, hidecategories=Object, MinimalAPI)
+UCLASS(collapsecategories, hidecategories=Object, MinimalAPI, meta=(MIR_OnDemandInputRequest))
 class UMaterialExpressionShadingPathSwitch : public UMaterialExpression
 {
 	GENERATED_BODY()
diff --git a/Engine/Source/Runtime/Engine/Public/Materials/MaterialExpressionStaticSwitch.h b/Engine/Source/Runtime/Engine/Public/Materials/MaterialExpressionStaticSwitch.h
--- a/Engine/Source/Runtime/Engine/Public/Materials/MaterialExpressionStaticSwitch.h
+++ b/Engine/Source/Runtime/Engine/Public/Materials/MaterialExpressionStaticSwitch.h
@@ -10,7 +10,7 @@
 #include "Materials/MaterialExpression.h"
 #include "MaterialExpressionStaticSwitch.generated.h"
 
-UCLASS(collapsecategories, hidecategories=Object, MinimalAPI)
+UCLASS(collapsecategories, hidecategories=Object, MinimalAPI, meta=(MIR_OnDemandInputRequest))
 class UMaterialExpressionStaticSwitch : public UMaterialExpression
 {
 	GENERATED_BODY()
diff --git a/Engine/Source/Runtime/Engine/Public/Materials/MaterialExpressionStaticSwitchParameter.h b/Engine/Source/Runtime/Engine/Public/Materials/MaterialExpressionStaticSwitchParameter.h
--- a/Engine/Source/Runtime/Engine/Public/Materials/MaterialExpressionStaticSwitchParameter.h
+++ b/Engine/Source/Runtime/Engine/Public/Materials/MaterialExpressionStaticSwitchParameter.h
@@ -10,7 +10,7 @@
 #include "Materials/MaterialExpressionStaticBoolParameter.h"
 #include "MaterialExpressionStaticSwitchParameter.generated.h"
 
-UCLASS(collapsecategories, hidecategories=Object, MinimalAPI)
+UCLASS(collapsecategories, hidecategories=Object, MinimalAPI, meta=(MIR_OnDemandInputRequest))
 class UMaterialExpressionStaticSwitchParameter : public UMaterialExpressionStaticBoolParameter
 {
 	GENERATED_BODY()
diff --git a/Engine/Source/Runtime/Engine/Public/Materials/MaterialExpressionVirtualTextureFeatureSwitch.h b/Engine/Source/Runtime/Engine/Public/Materials/MaterialExpressionVirtualTextureFeatureSwitch.h
--- a/Engine/Source/Runtime/Engine/Public/Materials/MaterialExpressionVirtualTextureFeatureSwitch.h
+++ b/Engine/Source/Runtime/Engine/Public/Materials/MaterialExpressionVirtualTextureFeatureSwitch.h
@@ -10,7 +10,7 @@
 #include "MaterialExpressionVirtualTextureFeatureSwitch.generated.h"
 
 /** Material output expression to switch logic according to whether virtual texturing is supported on this project, platform and feature level. */
-UCLASS(collapsecategories, hidecategories = Object)
+UCLASS(collapsecategories, hidecategories = Object, meta=(MIR_OnDemandInputRequest))
 class UMaterialExpressionVirtualTextureFeatureSwitch : public UMaterialExpression
 {
 	GENERATED_BODY()
diff --git a/Engine/Source/Runtime/Engine/Public/Materials/MaterialExpressionsToMIRCommon.h b/Engine/Source/Runtime/Engine/Public/Materials/MaterialExpressionsToMIRCommon.h
--- a/Engine/Source/Runtime/Engine/Public/Materials/MaterialExpressionsToMIRCommon.h
+++ b/Engine/Source/Runtime/Engine/Public/Materials/MaterialExpressionsToMIRCommon.h
@@ -27,5 +27,9 @@ MIR::FValueRef EmitCustomPrimitiveDataFloat1(MIR::FEmitter& Em, int32 Index);
 // This is equilent to `Em.Vector(EmitCustomPrimitiveDataFloat1(Em, 10), ..., EmitCustomPrimitiveDataFloat1(Em, 13))`.
 MIR::FValueRef EmitCustomPrimitiveDataFloat4(MIR::FEmitter& Em, int32 Index);
 
+// Verifies that Input is connected to a real expression, tracing through any reroute nodes iteratively.
+// Emits a compile error via Em.Errorf if the input is unconnected (or only connected through a dangling reroute chain).
+void CheckInputIsConnected(MIR::FEmitter& Em, FExpressionInput& Input, const TCHAR* InputName);
+
 } // MaterialToMIR
 #endif // WITH_EDITOR
diff --git a/Engine/Source/Runtime/Engine/Public/Materials/MaterialIREmitter.h b/Engine/Source/Runtime/Engine/Public/Materials/MaterialIREmitter.h
--- a/Engine/Source/Runtime/Engine/Public/Materials/MaterialIREmitter.h
+++ b/Engine/Source/Runtime/Engine/Public/Materials/MaterialIREmitter.h
@@ -9,7 +9,7 @@
 
 #if WITH_EDITOR
 
-#define UE_MIR_CHECKPOINT(Emitter) if (Emitter.CurrentExpressionHasErrors()) return
+#define UE_MIR_CHECKPOINT(Emitter) if (Emitter.IsInvalid()) return
 
 namespace MIR {
 
@@ -299,9 +299,10 @@ class FEmitter
 public:
 	/*--------------------------------- Error handling ---------------------------------*/
 
-	// Returns whether the current expression generated an error.
-	// You can also use UE_MIR_CHECKPOINT() in your code to check this and automatically return if an error occurred.
-	bool CurrentExpressionHasErrors() const { return bCurrentExpressionHasErrors; }
+	// Returns whether the emitter is in a non-clean state (pending dependency or error).
+	// You can also use UE_MIR_CHECKPOINT() in your code to check this and automatically
+	// return if a state is invalid and expression Build() execution should abort.
+	bool IsInvalid() const { return State != EState::None; }
 
 	// Reports an error using printf-like format and arguments to format the error message.
 	template <typename... TArgs>
@@ -938,9 +939,6 @@ private:
 	// The substrate translator common data representing the material topology.
 	struct Substrate::FSubstrateTranslatorData* SubstrateTranslatorData;
 
-	// Whether the current expression has reported any error.
-	bool bCurrentExpressionHasErrors = false;
-
 	// Global "true" constant.
 	FValue* TrueConstant{};
 
@@ -962,6 +960,20 @@ private:
 	// The set of unique HLSL functions emitted so far, used to deduplicate identical function definitions.
 	TSet<FFunctionHLSL*, FFunctionHLSLKeyFuncs> FunctionHLSLSet{};
 
+	// Tracks the current expression build state.
+	enum class EState : uint8
+	{
+		None,
+		Pending, // The expression has an unresolved input dependency and must be deferred.
+		Error,   // An error was reported.
+	};
+
+	// Current emitter state for the expression being built.
+	EState State = EState::None;
+
+	// Whether the current expression opts into on-demand input resolution via MIR_OnDemandInputRequest metadata.
+	bool bExpressionRequestsInputsOnDemand = false;
+
 	friend FMaterialIRModuleBuilder;
 	friend FMaterialIRModuleBuilderImpl;
 };
