Partial rollback of change in 52446534 that is breaking some substrate materials.
This will bring back the potential mismatch between the Shader Parameters used to generate DDC Key and the ones used to generate the Shader Layout.
So we are trading a visual issue with going back to conservative shaders generation.
The shader generation issue will be revisited in another changelist.

#rnx
#Synced-CL 52498028

diff --git a/Engine/Source/Runtime/Engine/Private/Materials/MaterialCachedData.cpp b/Engine/Source/Runtime/Engine/Private/Materials/MaterialCachedData.cpp
--- a/Engine/Source/Runtime/Engine/Private/Materials/MaterialCachedData.cpp
+++ b/Engine/Source/Runtime/Engine/Private/Materials/MaterialCachedData.cpp
@@ -836,6 +836,7 @@ void FMaterialCachedExpressionData::DoPostAnalyzeChecks(const UMaterial& BaseMat
 
 	for (const FParsedExpressionData& ParameterData : ParsedExpressions)
 	{
+#if 0
 		// Legacy materials may have a static switch with MaterialAttributes on one branch
 		// and a non-MaterialAttributes value (e.g. scalar constant) on the other.
 		// During translation, the non-MaterialAttributes branch gets broadcast to ALL material properties, causing
@@ -873,7 +874,7 @@ void FMaterialCachedExpressionData::DoPostAnalyzeChecks(const UMaterial& BaseMat
 				MarkAllPropertiesConnectedForMismatchedMaterialAttributesSwitch(StaticSwitchParam->A, StaticSwitchParam->B);
 			}
 		}
-
+#endif
 
 		if (bPedanticErrorChecksEnabled)
 		{
