[WorldPartition] Fix access to GC'ed URuntimePartitionCell in server
- Cells injected before BeginPlay would be removed at OnBeginPlay, but made their way into the SpatialIndex, which wasn't rebuilt
- SpatialIndex only keeps weak references to cells (except in editor) leading to crashes during queries on server instances

#Synced-CL 52465190

diff --git a/Engine/Source/Runtime/Engine/Private/WorldPartition/RuntimeHashSet/WorldPartitionRuntimeHashSet.cpp b/Engine/Source/Runtime/Engine/Private/WorldPartition/RuntimeHashSet/WorldPartitionRuntimeHashSet.cpp
--- a/Engine/Source/Runtime/Engine/Private/WorldPartition/RuntimeHashSet/WorldPartitionRuntimeHashSet.cpp
+++ b/Engine/Source/Runtime/Engine/Private/WorldPartition/RuntimeHashSet/WorldPartitionRuntimeHashSet.cpp
@@ -504,12 +504,9 @@ bool UWorldPartitionRuntimeHashSet::InjectExternalStreamingObject(URuntimeHashEx
 		URuntimeHashSetExternalStreamingObject* HashSetExternalStreamingObject = CastChecked<URuntimeHashSetExternalStreamingObject>(ExternalStreamingObject);
 
 #if !WITH_EDITOR
-		if (UWorld* OuterWorld = GetTypedOuter<UWorld>(); OuterWorld && OuterWorld->HasBegunPlay())
+		for (FRuntimePartitionStreamingData& StreamingData : HashSetExternalStreamingObject->RuntimeStreamingData)
 		{
-			for (FRuntimePartitionStreamingData& StreamingData : HashSetExternalStreamingObject->RuntimeStreamingData)
-			{
-				RemoveIrrelevantCells(StreamingData);
-			}
+			RemoveIrrelevantCells(StreamingData);
 		}
 #endif
 
@@ -1418,8 +1415,14 @@ void UWorldPartitionRuntimeHashSet::ForEachStreamingData(TFunctionRef<bool(const
 
 void UWorldPartitionRuntimeHashSet::RemoveIrrelevantCells(FRuntimePartitionStreamingData& StreamingData)
 {
+	const int32 PreviousSpatialCellCount = StreamingData.SpatiallyLoadedCells.Num();
+
 	StreamingData.SpatiallyLoadedCells.SetNum(Algo::RemoveIf(StreamingData.SpatiallyLoadedCells, [this](UWorldPartitionRuntimeCell* Cell) { return !IsCellRelevantFor(Cell->GetClientOnlyVisible()); }));
 	StreamingData.NonSpatiallyLoadedCells.SetNum(Algo::RemoveIf(StreamingData.NonSpatiallyLoadedCells, [this](UWorldPartitionRuntimeCell* Cell) { return !IsCellRelevantFor(Cell->GetClientOnlyVisible()); }));
+
+	// See URuntimeHashSetExternalStreamingObject::AddReferencedObjects - cells in spatial indices are not tracked via ARO in non-editor builds, will lead to deallocated mem access
+	checkf(StreamingData.SpatiallyLoadedCells.Num() == PreviousSpatialCellCount || !StreamingData.HasSpatialIndex(),
+		TEXT("Spatial cells removed after SpatialIndex creation."));
 }
 
 void UWorldPartitionRuntimeHashSet::UpdateRuntimeDataGridMap()
diff --git a/Engine/Source/Runtime/Engine/Public/WorldPartition/RuntimeHashSet/WorldPartitionRuntimeHashSet.h b/Engine/Source/Runtime/Engine/Public/WorldPartition/RuntimeHashSet/WorldPartitionRuntimeHashSet.h
--- a/Engine/Source/Runtime/Engine/Public/WorldPartition/RuntimeHashSet/WorldPartitionRuntimeHashSet.h
+++ b/Engine/Source/Runtime/Engine/Public/WorldPartition/RuntimeHashSet/WorldPartitionRuntimeHashSet.h
@@ -100,6 +100,7 @@ struct FRuntimePartitionStreamingData
 
 	void CreatePartitionsSpatialIndex() const;
 	void DestroyPartitionsSpatialIndex() const;
+	bool HasSpatialIndex() const { return SpatialIndex.IsValid() || SpatialIndexForce2D.IsValid() || SpatialIndex2D.IsValid(); }
 
 	ENGINE_API int32 GetLoadingRange() const;
 
