[SkeletalMesh] Add NoRefStreamingLODBias to the SkeletalMesh LODSettings.

#rnx

#Synced-CL 52493164

diff --git a/Engine/Source/Runtime/Engine/Classes/Engine/SkeletalMeshLODSettings.h b/Engine/Source/Runtime/Engine/Classes/Engine/SkeletalMeshLODSettings.h
--- a/Engine/Source/Runtime/Engine/Classes/Engine/SkeletalMeshLODSettings.h
+++ b/Engine/Source/Runtime/Engine/Classes/Engine/SkeletalMeshLODSettings.h
@@ -152,6 +152,10 @@ protected:
 	UPROPERTY(globalconfig, EditAnywhere, Category=LODGroups, meta=(EditCondition="bOverrideLODStreamingSettings"))
 	FPerPlatformInt MaxNumOptionalLODs;
 
+	/** Default LOD bias for no-ref meshes. */
+	UPROPERTY(globalconfig, EditAnywhere, Category = LODGroups)
+	FPerQualityLevelInt NoRefStreamingLODBias;
+
 	UPROPERTY(globalconfig, EditAnywhere, Category=LODGroups)
 	TArray<FSkeletalMeshLODGroupSettings> LODGroups;
 
diff --git a/Engine/Source/Runtime/Engine/Private/SkeletalMeshLODSettings.cpp b/Engine/Source/Runtime/Engine/Private/SkeletalMeshLODSettings.cpp
--- a/Engine/Source/Runtime/Engine/Private/SkeletalMeshLODSettings.cpp
+++ b/Engine/Source/Runtime/Engine/Private/SkeletalMeshLODSettings.cpp
@@ -15,6 +15,8 @@ DEFINE_LOG_CATEGORY_STATIC(LogSkeletalMeshLODSettings, Warning, All)
 
 extern const TCHAR* GSkeletalMeshMinLodQualityLevelCVarName;
 extern const TCHAR* GSkeletalMeshMinLodQualityLevelScalabilitySection;
+extern const TCHAR* GNoRefBiasQualityLevelCVarName;
+extern const TCHAR* GNoRefBiasQualityLevelScalabilitySection;
 
 USkeletalMeshLODSettings::USkeletalMeshLODSettings(const FObjectInitializer& ObjectInitializer)
 	: Super(ObjectInitializer)
@@ -25,6 +27,9 @@ USkeletalMeshLODSettings::USkeletalMeshLODSettings(const FObjectInitializer& Obj
 	MaxNumOptionalLODs.Default = 0;
 
 	MinQualityLevelLod.SetQualityLevelCVarForCooking(GSkeletalMeshMinLodQualityLevelCVarName, GSkeletalMeshMinLodQualityLevelScalabilitySection);
+
+	NoRefStreamingLODBias.Default = -1;
+	NoRefStreamingLODBias.SetQualityLevelCVarForCooking(GNoRefBiasQualityLevelCVarName, GNoRefBiasQualityLevelScalabilitySection);
 }
 
 const FSkeletalMeshLODGroupSettings& USkeletalMeshLODSettings::GetSettingsForLODLevel(const int32 LODIndex) const
@@ -206,6 +211,7 @@ int32 USkeletalMeshLODSettings::SetLODSettingsToMesh(USkeletalMesh* InMesh) cons
 		InMesh->SetMinLod(MinLod);
 		InMesh->SetQualityLevelMinLod(MinQualityLevelLod);
 		InMesh->SetDisableBelowMinLodStripping(DisableBelowMinLodStripping);
+		InMesh->SetNoRefStreamingLODBias(NoRefStreamingLODBias);
 #if WITH_EDITORONLY_DATA
 		InMesh->SetOverrideLODStreamingSettings(bOverrideLODStreamingSettings);
 		InMesh->SetSupportLODStreaming(bSupportLODStreaming);
@@ -233,6 +239,7 @@ int32 USkeletalMeshLODSettings::SetLODSettingsFromMesh(USkeletalMesh* InMesh)
 		MinLod = InMesh->GetMinLod();
 		MinQualityLevelLod = InMesh->GetQualityLevelMinLod();
 		DisableBelowMinLodStripping = InMesh->GetDisableBelowMinLodStripping();
+		NoRefStreamingLODBias = InMesh->GetNoRefStreamingLODBias();
 #if WITH_EDITORONLY_DATA
 		bOverrideLODStreamingSettings = InMesh->GetOverrideLODStreamingSettings();
 		bSupportLODStreaming = InMesh->GetSupportLODStreaming();
diff --git a/Engine/Source/Runtime/Engine/Private/StreamableRenderAsset.cpp b/Engine/Source/Runtime/Engine/Private/StreamableRenderAsset.cpp
--- a/Engine/Source/Runtime/Engine/Private/StreamableRenderAsset.cpp
+++ b/Engine/Source/Runtime/Engine/Private/StreamableRenderAsset.cpp
@@ -11,8 +11,8 @@
 
 #include UE_INLINE_GENERATED_CPP_BY_NAME(StreamableRenderAsset)
 
-static const TCHAR* GNoRefBiasQualityLevelCVarName = TEXT("r.Streaming.NoRefLODBiasQualityLevel");
-static const TCHAR* GNoRefBiasQualityLevelScalabilitySection = TEXT("ViewDistanceQuality");
+const TCHAR* GNoRefBiasQualityLevelCVarName = TEXT("r.Streaming.NoRefLODBiasQualityLevel");
+const TCHAR* GNoRefBiasQualityLevelScalabilitySection = TEXT("ViewDistanceQuality");
 static int32 GNoRefBiasQualityLevel = -1;
 static FAutoConsoleVariableRef CVarNoRefBiasQualityLevel(
 	GNoRefBiasQualityLevelCVarName,
