Reenable fix for the potential mismatch between the Shader Parameters used to generate DDC Key and the ones used to generate the Shader Layout.
Deal with specific substrate case. Don't mark the substrate front material property as connected unless we are specifically fixing up a mismatch between substrate attribute and non substrate attribute inputs.
#rnx
#Synced-CL 52498028

diff --git a/Engine/Source/Runtime/Engine/Private/Materials/MaterialCachedData.cpp b/Engine/Source/Runtime/Engine/Private/Materials/MaterialCachedData.cpp
--- a/Engine/Source/Runtime/Engine/Private/Materials/MaterialCachedData.cpp
+++ b/Engine/Source/Runtime/Engine/Private/Materials/MaterialCachedData.cpp
@@ -836,7 +836,6 @@ void FMaterialCachedExpressionData::DoPostAnalyzeChecks(const UMaterial& BaseMat
 
 	for (const FParsedExpressionData& ParameterData : ParsedExpressions)
 	{
-#if 0
 		// Legacy materials may have a static switch with MaterialAttributes on one branch
 		// and a non-MaterialAttributes value (e.g. scalar constant) on the other.
 		// During translation, the non-MaterialAttributes branch gets broadcast to ALL material properties, causing
@@ -844,26 +843,38 @@ void FMaterialCachedExpressionData::DoPostAnalyzeChecks(const UMaterial& BaseMat
 		// To fix this, detect the mismatch here and conservatively mark all properties as connected for those legacy materials.
 		{
 			auto MarkAllPropertiesConnectedForMismatchedMaterialAttributesSwitch = [&](const FExpressionInput& InputA, const FExpressionInput& InputB)
-				{
-					const FExpressionInput TracedA = InputA.GetTracedInput();
-					const FExpressionInput TracedB = InputB.GetTracedInput();
+			{
+				const FExpressionInput TracedA = InputA.GetTracedInput();
+				const FExpressionInput TracedB = InputB.GetTracedInput();
 
-					if (!TracedA.IsConnected() || !TracedB.IsConnected())
-					{
-						return;
-					}
+				if (!TracedA.IsConnected() || !TracedB.IsConnected())
+				{
+					return;
+				}
 
-					const bool bAIsMaterialAttributes = TracedA.Expression->IsResultMaterialAttributes(TracedA.OutputIndex);
-					const bool bBIsMaterialAttributes = TracedB.Expression->IsResultMaterialAttributes(TracedB.OutputIndex);
+				const bool bAIsMaterialAttributes = TracedA.Expression->IsResultMaterialAttributes(TracedA.OutputIndex);
+				const bool bBIsMaterialAttributes = TracedB.Expression->IsResultMaterialAttributes(TracedB.OutputIndex);
 
-					if (bAIsMaterialAttributes != bBIsMaterialAttributes)
+				if (bAIsMaterialAttributes != bBIsMaterialAttributes)
+				{
+					for (int32 PropertyIndex = 0; PropertyIndex < MP_MaterialAttributes; ++PropertyIndex)
 					{
-						for (int32 PropertyIndex = 0; PropertyIndex < MP_MaterialAttributes; ++PropertyIndex)
+						// Substrate specific property handled separately.
+						if (PropertyIndex != MP_FrontMaterial)
 						{
 							SetPropertyConnected((EMaterialProperty)PropertyIndex);
 						}
 					}
-				};
+				}
+
+				const bool bAIsSubstrateAttributes = TracedA.Expression->IsResultSubstrateMaterial(TracedA.OutputIndex);
+				const bool bBIsSubstrateAttributes = TracedB.Expression->IsResultSubstrateMaterial(TracedB.OutputIndex);
+
+				if (bAIsSubstrateAttributes != bBIsSubstrateAttributes)
+				{
+					SetPropertyConnected(MP_FrontMaterial);
+				}
+			};
 
 			if (UMaterialExpressionStaticSwitch* StaticSwitch = Cast<UMaterialExpressionStaticSwitch>(ParameterData.Expression))
 			{
@@ -874,7 +885,6 @@ void FMaterialCachedExpressionData::DoPostAnalyzeChecks(const UMaterial& BaseMat
 				MarkAllPropertiesConnectedForMismatchedMaterialAttributesSwitch(StaticSwitchParam->A, StaticSwitchParam->B);
 			}
 		}
-#endif
 
 		if (bPedanticErrorChecksEnabled)
 		{
