Fix: 'ke * MyFunc' console command (KismetEvent) won't deadlock anymore if MyFunc flushes async loading
#Synced-CL 52506375

diff --git a/Engine/Source/Runtime/Engine/Private/UnrealEngine.cpp b/Engine/Source/Runtime/Engine/Private/UnrealEngine.cpp
--- a/Engine/Source/Runtime/Engine/Private/UnrealEngine.cpp
+++ b/Engine/Source/Runtime/Engine/Private/UnrealEngine.cpp
@@ -6748,8 +6748,12 @@ bool UEngine::HandleKismetEventCommand(UWorld* InWorld, const TCHAR* Cmd, FOutpu
 			return false;
 		};
 
-		// Send the command to everything in the world we're dealing with...
+		// Gather all non-CDO, non-default subobject objects in the target world.
+		// Call the function in a separate pass to avoid deadlocks from the target
+		// function triggering and waiting for off-thread work that needs the 
+		// GUObjectArray lock, like FlushAsyncLoading.
 		int32 NumInstanceCallsSucceeded = 0;
+		TArray<UObject*> CandidateObjects;
 		for (TObjectIterator<UObject> It(EObjectFlags::RF_ClassDefaultObject, true, EInternalObjectFlags::Garbage); It; ++It)
 		{
 			UObject* const Obj = *It;
@@ -6757,14 +6761,24 @@ bool UEngine::HandleKismetEventCommand(UWorld* InWorld, const TCHAR* Cmd, FOutpu
 			const bool bIsCDOOwnedArchtype = IsCDOOwnedArchtype(Obj);
 			if (!bIsCDOOwnedArchtype)
 			{
-			UWorld const* const ObjWorld = Obj->GetWorld();
-			if (ObjWorld == InWorld)
+				UWorld const* const ObjWorld = Obj->GetWorld();
+				if (ObjWorld == InWorld)
+				{
+					CandidateObjects.Add(Obj);
+				}
+			}
+		}
+
+		// Try to call the target function on all candidate objects. Only succeeds
+		// if the object has a matching function and signature.
+		for (UObject* Obj : CandidateObjects)
+		{
+			if (IsValid(Obj))
 			{
 				const bool bSucceeded = Obj->CallFunctionByNameWithArguments(Cmd, Ar, nullptr, true);
 				NumInstanceCallsSucceeded += bSucceeded ? 1 : 0;
 			}
 		}
-		}
 
 		Ar.Logf(TEXT("Called '%s' on everything in the world and %d instances succeeded"), Cmd, NumInstanceCallsSucceeded);
 	}
@@ -6780,17 +6794,30 @@ bool UEngine::HandleKismetEventCommand(UWorld* InWorld, const TCHAR* Cmd, FOutpu
 		{
 			if (UClass* ClassToMatch = Cast<UClass>(ObjectToMatch))
 			{
-				// Call it on all instances of the class
+				// Gather all eligible instances of the class. 
+				// Call the function in a separate pass to avoid deadlocks from the target
+				// function triggering and waiting for off-thread work that needs the 
+				// GUObjectArray lock, like FlushAsyncLoading.
 				int32 NumInstancesFound = 0;
 				int32 NumInstanceCallsSucceeded = 0;
+				TArray<UObject*> MatchingObjects;
 				for (FThreadSafeObjectIterator It(ClassToMatch); It; ++It)
 				{
 					UObject* const Obj = *It;
 					UWorld const* const ObjWorld = Obj->GetWorld();
 					if (ObjWorld == InWorld)
 					{
-						const bool bSucceeded = Obj->CallFunctionByNameWithArguments(Cmd, Ar, nullptr, true);
 						++NumInstancesFound;
+						MatchingObjects.Add(Obj);
+					}
+				}
+
+				// Call the target function all eligible objects.
+				for (UObject* Obj : MatchingObjects)
+				{
+					if (IsValid(Obj))
+					{
+						const bool bSucceeded = Obj->CallFunctionByNameWithArguments(Cmd, Ar, nullptr, true);
 						NumInstanceCallsSucceeded += bSucceeded ? 1 : 0;
 					}
 				}
