[StandaloneHLOD] Fix GC crash in SetupHLODActors. Protect StandaloneHLODWorld with FGCObjectScopeGuard

#Synced-CL 52152802

diff --git a/Engine/Source/Runtime/Engine/Private/WorldPartition/RuntimeHashSet/WorldPartitionRuntimeHashSetHLODGeneration.cpp b/Engine/Source/Runtime/Engine/Private/WorldPartition/RuntimeHashSet/WorldPartitionRuntimeHashSetHLODGeneration.cpp
--- a/Engine/Source/Runtime/Engine/Private/WorldPartition/RuntimeHashSet/WorldPartitionRuntimeHashSetHLODGeneration.cpp
+++ b/Engine/Source/Runtime/Engine/Private/WorldPartition/RuntimeHashSet/WorldPartitionRuntimeHashSetHLODGeneration.cpp
@@ -27,6 +27,7 @@
 #include "EngineUtils.h"
 #endif
 
+#include "UObject/GCObjectScopeGuard.h"
 #include "UObject/SavePackage.h"
 #include "HAL/PlatformFileManager.h"
 
@@ -413,6 +414,9 @@ bool UWorldPartitionRuntimeHashSet::SetupHLODActors(const IStreamingGenerationCo
 			});
 		}
 
+		// Prevent StandaloneHLODWorld from being garbage collected during GameTick calls
+		FGCObjectScopeGuard StandaloneHLODWorldGuard(StandaloneHLODWorld);
+
 		TMap<URuntimePartition*, TArray<URuntimePartition::FCellDescInstance>> RuntimePartitionsStreamingDescs;
 		GenerateRuntimePartitionsStreamingDescs(CurrentHLODStreamingGenerationContext.Get(), RuntimePartitionsStreamingDescs);
 
