Fixed all tick records on an anim instance being congregated on the main instance track in Rewind Debugger making it hard to debug complex setups.

#rnx

#Synced-CL 52475608

diff --git a/Engine/Source/Runtime/AnimGraphRuntime/Private/AnimNodes/AnimNode_BlendSpacePlayer.cpp b/Engine/Source/Runtime/AnimGraphRuntime/Private/AnimNodes/AnimNode_BlendSpacePlayer.cpp
--- a/Engine/Source/Runtime/AnimGraphRuntime/Private/AnimNodes/AnimNode_BlendSpacePlayer.cpp
+++ b/Engine/Source/Runtime/AnimGraphRuntime/Private/AnimNodes/AnimNode_BlendSpacePlayer.cpp
@@ -99,7 +99,9 @@ void FAnimNode_BlendSpacePlayerBase::UpdateInternal(const FAnimationUpdateContex
 		TickRecord.RootMotionWeightModifier = Context.GetRootMotionWeightModifier();
 		TickRecord.DeltaTimeRecord = &DeltaTimeRecord;
 		TickRecord.bRequestedInertialization = Context.GetMessage<UE::Anim::FAnimInertializationSyncScope>() != nullptr;
-		
+#if ANIM_TRACE_ENABLED
+		TickRecord.SourceProxy = Context.AnimInstanceProxy; 
+#endif
 		TickRecord.GatherContextData(Context);
 
 		SyncScope.AddTickRecord(TickRecord, GetSyncParams(TickRecord.bRequestedInertialization), UE::Anim::FAnimSyncDebugInfo(Context));
diff --git a/Engine/Source/Runtime/AnimGraphRuntime/Private/AnimNodes/AnimNode_RandomPlayer.cpp b/Engine/Source/Runtime/AnimGraphRuntime/Private/AnimNodes/AnimNode_RandomPlayer.cpp
--- a/Engine/Source/Runtime/AnimGraphRuntime/Private/AnimNodes/AnimNode_RandomPlayer.cpp
+++ b/Engine/Source/Runtime/AnimGraphRuntime/Private/AnimNodes/AnimNode_RandomPlayer.cpp
@@ -239,6 +239,9 @@ void FAnimNode_RandomPlayer::Update_AnyThread(const FAnimationUpdateContext& Con
 
 	FAnimTickRecord TickRecord(CurrentData->Entry->Sequence, true, CurrentData->PlayRate, false, CurrentData->BlendWeight * BlendWeight, CurrentData->CurrentPlayTime, CurrentData->MarkerTickRecord);
 	TickRecord.DeltaTimeRecord = &CurrentData->DeltaTimeRecord;
+#if ANIM_TRACE_ENABLED
+	TickRecord.SourceProxy = Context.AnimInstanceProxy;
+#endif
 	TickRecord.GatherContextData(Context);
 
 	UE::Anim::FAnimSyncGroupScope& SyncScope = Context.GetMessageChecked<UE::Anim::FAnimSyncGroupScope>();
@@ -248,6 +251,9 @@ void FAnimNode_RandomPlayer::Update_AnyThread(const FAnimationUpdateContext& Con
 	{
 		FAnimTickRecord NextTickRecord(NextData->Entry->Sequence, true, NextData->PlayRate, false, NextData->BlendWeight * BlendWeight, NextData->CurrentPlayTime, NextData->MarkerTickRecord);
 		NextTickRecord.DeltaTimeRecord = &NextData->DeltaTimeRecord;
+#if ANIM_TRACE_ENABLED
+		NextTickRecord.SourceProxy = Context.AnimInstanceProxy;
+#endif
 		NextTickRecord.GatherContextData(Context);
 
 		SyncScope.AddTickRecord(NextTickRecord, UE::Anim::FAnimSyncParams(), UE::Anim::FAnimSyncDebugInfo(Context));
diff --git a/Engine/Source/Runtime/Engine/Classes/Animation/AnimationAsset.h b/Engine/Source/Runtime/Engine/Classes/Animation/AnimationAsset.h
--- a/Engine/Source/Runtime/Engine/Classes/Animation/AnimationAsset.h
+++ b/Engine/Source/Runtime/Engine/Classes/Animation/AnimationAsset.h
@@ -16,6 +16,7 @@
 #include "Interfaces/Interface_AssetUserData.h"
 #include "AnimInterpFilter.h"
 #include "AnimEnums.h"
+#include "Animation/AnimTraceDefines.h"
 #include "Interfaces/Interface_PreviewMeshProvider.h"
 #if WITH_EDITOR
 #include "IO/IoHash.h"
@@ -34,6 +35,8 @@ class UMirrorDataTable;
 class USkeletalMesh;
 class ITargetPlatform;
 struct FAnimationUpdateContext;
+struct FAnimInstanceProxy;
+
 namespace SmartName
 {
 typedef uint16 UID_Type;
@@ -458,7 +461,14 @@ struct FAnimTickRecord
 	FMarkerTickRecord* MarkerTickRecord = nullptr;
 	bool bCanUseMarkerSync = false;
 	float LeaderScore = 0.0f;
-
+	
+#if ANIM_TRACE_ENABLED
+	// Source proxy for deferred tracing - captures the originating anim instance proxy at tick record
+	// creation time so that linked graph tick records are correctly attributed during the deferred trace loop.
+	// Raw pointer is safe here: proxy lifetime is tied to its UAnimInstance which outlives the tick record.
+	FAnimInstanceProxy* SourceProxy = nullptr;
+#endif
+	
 	// Return the root motion weight for this tick record
 	float GetRootMotionWeight() const { return EffectiveBlendWeight * RootMotionWeightModifier; }
 
diff --git a/Engine/Source/Runtime/Engine/Private/Animation/AnimInstance.cpp b/Engine/Source/Runtime/Engine/Private/Animation/AnimInstance.cpp
--- a/Engine/Source/Runtime/Engine/Private/Animation/AnimInstance.cpp
+++ b/Engine/Source/Runtime/Engine/Private/Animation/AnimInstance.cpp
@@ -606,7 +606,9 @@ void UAnimInstance::UpdateMontageSyncGroup()
 					MontageInstance->MarkerTickRecord
 				);
 				TickRecord.DeltaTimeRecord = &MontageInstance->DeltaTimeRecord;
-
+#if ANIM_TRACE_ENABLED
+				TickRecord.SourceProxy = &GetProxyOnAnyThread<FAnimInstanceProxy>();
+#endif
 				UE::Anim::FAnimSyncParams Params(GroupNameToUse);
 				GetProxyOnGameThread<FAnimInstanceProxy>().AddTickRecord(TickRecord, Params);
 
diff --git a/Engine/Source/Runtime/Engine/Private/Animation/AnimInstanceProxy.cpp b/Engine/Source/Runtime/Engine/Private/Animation/AnimInstanceProxy.cpp
--- a/Engine/Source/Runtime/Engine/Private/Animation/AnimInstanceProxy.cpp
+++ b/Engine/Source/Runtime/Engine/Private/Animation/AnimInstanceProxy.cpp
@@ -1249,7 +1249,7 @@ void FAnimInstanceProxy::UpdateAnimation()
 				TRACE_ANIM_TICK_RECORD(Context, TickRecord)
 			}
 		}
-	
+
 		for (const FAnimTickRecord& TickRecord : Sync.GetUngroupedActivePlayersRead())
 		{
 			TRACE_ANIM_TICK_RECORD(Context, TickRecord)
@@ -1993,6 +1993,8 @@ void FAnimInstanceProxy::SlotEvaluatePose(const FName& SlotNodeName, const FAnim
 					TickRecord.TimeAccumulator = &CurrentTime;
 					TickRecord.DeltaTimeRecord = &ExtractionContext.DeltaTimeRecord;
 					TickRecord.EffectiveBlendWeight = MontageWeight;
+					TickRecord.SourceProxy = Context.AnimInstanceProxy;
+					
 					TRACE_ANIM_TICK_RECORD(Context, TickRecord);
 				}
 			}
diff --git a/Engine/Source/Runtime/Engine/Private/Animation/AnimNode_AssetPlayerBase.cpp b/Engine/Source/Runtime/Engine/Private/Animation/AnimNode_AssetPlayerBase.cpp
--- a/Engine/Source/Runtime/Engine/Private/Animation/AnimNode_AssetPlayerBase.cpp
+++ b/Engine/Source/Runtime/Engine/Private/Animation/AnimNode_AssetPlayerBase.cpp
@@ -46,7 +46,10 @@ void FAnimNode_AssetPlayerBase::CreateTickRecordForNode(const FAnimationUpdateCo
 	TickRecord.RootMotionWeightModifier = Context.GetRootMotionWeightModifier();
 	TickRecord.DeltaTimeRecord = &DeltaTimeRecord;
 	TickRecord.bRequestedInertialization = Context.GetMessage<UE::Anim::FAnimInertializationSyncScope>() != nullptr;
-
+#if ANIM_TRACE_ENABLED
+	TickRecord.SourceProxy = Context.AnimInstanceProxy;
+#endif
+	
 	if (bHasValidAssetPlayerInstanceID)
 	{
 		TickRecord.MakeContextData<UE::Anim::FAnimNotifyAssetPlayerInstanceContext>(AssetPlayerInstanceID);
diff --git a/Engine/Source/Runtime/Engine/Private/Animation/AnimSingleNodeInstanceProxy.cpp b/Engine/Source/Runtime/Engine/Private/Animation/AnimSingleNodeInstanceProxy.cpp
--- a/Engine/Source/Runtime/Engine/Private/Animation/AnimSingleNodeInstanceProxy.cpp
+++ b/Engine/Source/Runtime/Engine/Private/Animation/AnimSingleNodeInstanceProxy.cpp
@@ -585,7 +585,9 @@ void FAnimNode_SingleNode::Update_AnyThread(const FAnimationUpdateContext& Conte
 				BlendSpace, Proxy->BlendSpacePosition, Proxy->BlendSampleData, Proxy->BlendFilter, Proxy->bLooping, 
 				NewPlayRate, false, false, 1.f, /*inout*/ Proxy->CurrentTime, Proxy->MarkerTickRecord);
 			TickRecord.DeltaTimeRecord = &(Proxy->DeltaTimeRecord);
-			
+#if ANIM_TRACE_ENABLED
+			TickRecord.SourceProxy = Context.AnimInstanceProxy;
+#endif
 			SyncScope.AddTickRecord(TickRecord);
 
 #if WITH_EDITORONLY_DATA
@@ -596,7 +598,9 @@ void FAnimNode_SingleNode::Update_AnyThread(const FAnimationUpdateContext& Conte
 		{
 			FAnimTickRecord TickRecord(Sequence, Proxy->bLooping, NewPlayRate, false, 1.f, /*inout*/ Proxy->CurrentTime, Proxy->MarkerTickRecord);
 			TickRecord.DeltaTimeRecord = &(Proxy->DeltaTimeRecord);
-			
+#if ANIM_TRACE_ENABLED
+			TickRecord.SourceProxy = Context.AnimInstanceProxy;
+#endif
 			SyncScope.AddTickRecord(TickRecord);
 
 			// if it's not looping, just set play to be false when reached to end
@@ -613,6 +617,9 @@ void FAnimNode_SingleNode::Update_AnyThread(const FAnimationUpdateContext& Conte
 		{
 			FAnimTickRecord TickRecord(Streamable, Proxy->bLooping, NewPlayRate, false, 1.f, /*inout*/ Proxy->CurrentTime, Proxy->MarkerTickRecord);
 			TickRecord.DeltaTimeRecord = &(Proxy->DeltaTimeRecord);
+#if ANIM_TRACE_ENABLED
+			TickRecord.SourceProxy = Context.AnimInstanceProxy;
+#endif
 			
 			SyncScope.AddTickRecord(TickRecord);
 
@@ -630,6 +637,9 @@ void FAnimNode_SingleNode::Update_AnyThread(const FAnimationUpdateContext& Conte
 		{
 			FAnimTickRecord TickRecord(Composite, Proxy->bLooping, NewPlayRate, false, 1.f, /*inout*/ Proxy->CurrentTime, Proxy->MarkerTickRecord);
 			TickRecord.DeltaTimeRecord = &(Proxy->DeltaTimeRecord);
+#if ANIM_TRACE_ENABLED
+			TickRecord.SourceProxy = Context.AnimInstanceProxy;
+#endif
 			
 			SyncScope.AddTickRecord(TickRecord);
 
@@ -670,6 +680,9 @@ void FAnimNode_SingleNode::Update_AnyThread(const FAnimationUpdateContext& Conte
 		{
 			FAnimTickRecord TickRecord(PoseAsset, 1.f);
 			TickRecord.DeltaTimeRecord = &(Proxy->DeltaTimeRecord);
+#if ANIM_TRACE_ENABLED
+			TickRecord.SourceProxy = Context.AnimInstanceProxy;
+#endif
 
 			SyncScope.AddTickRecord(TickRecord);
 		}
diff --git a/Engine/Source/Runtime/Engine/Private/Animation/AnimTrace.cpp b/Engine/Source/Runtime/Engine/Private/Animation/AnimTrace.cpp
--- a/Engine/Source/Runtime/Engine/Private/Animation/AnimTrace.cpp
+++ b/Engine/Source/Runtime/Engine/Private/Animation/AnimTrace.cpp
@@ -469,9 +469,12 @@ void FAnimTrace::OutputAnimTickRecord(const FAnimationBaseContext& InContext, co
 		return;
 	}
 
-	if(InTickRecord.SourceAsset)
+	if (InTickRecord.SourceAsset)
 	{
-		UObject* AnimInstance = InContext.AnimInstanceProxy->GetAnimInstanceObject();
+		// Use the source proxy for attribution if the tick record originated from a different proxy (e.g. a linked anim graph), 
+		// while keeping the context proxy for sync group lookup since tick records always live in the main proxy's sync regardless of which linked graph created them.
+		FAnimInstanceProxy* AttributionProxy = (InTickRecord.SourceProxy && InTickRecord.SourceProxy != InContext.AnimInstanceProxy) ? InTickRecord.SourceProxy : InContext.AnimInstanceProxy;
+		UObject* AnimInstance = AttributionProxy->GetAnimInstanceObject();
 		TRACE_OBJECT(AnimInstance);
 
 		TRACE_OBJECT(InTickRecord.SourceAsset);
diff --git a/Engine/Source/Runtime/Engine/Public/Animation/AnimTrace.h b/Engine/Source/Runtime/Engine/Public/Animation/AnimTrace.h
--- a/Engine/Source/Runtime/Engine/Public/Animation/AnimTrace.h
+++ b/Engine/Source/Runtime/Engine/Public/Animation/AnimTrace.h
@@ -8,9 +8,7 @@
 #include "ObjectTrace.h"
 #endif // UE_ENABLE_INCLUDE_ORDER_DEPRECATED_IN_5_7
 
-#include "ObjectTraceDefines.h"
-
-#define ANIM_TRACE_ENABLED (OBJECT_TRACE_ENABLED && !(UE_BUILD_SHIPPING || UE_BUILD_TEST))
+#include "Animation/AnimTraceDefines.h"
 
 #if ANIM_TRACE_ENABLED
 
diff --git a/Engine/Source/Runtime/Engine/Public/Animation/AnimTraceDefines.h b/Engine/Source/Runtime/Engine/Public/Animation/AnimTraceDefines.h
new file mode 100644
--- /dev/null
+++ b/Engine/Source/Runtime/Engine/Public/Animation/AnimTraceDefines.h
@@ -0,0 +1,9 @@
+// Copyright Epic Games, Inc. All Rights Reserved.
+
+#pragma once
+
+#include "ObjectTraceDefines.h"
+
+#ifndef ANIM_TRACE_ENABLED
+#define ANIM_TRACE_ENABLED (OBJECT_TRACE_ENABLED && !(UE_BUILD_SHIPPING || UE_BUILD_TEST))
+#endif
\ No newline at end of file
