Simplify the PPV storage in UWorld by deprecating non-UObject PPVs.

#rnx
#Synced-CL 52479141

diff --git a/Engine/Plugins/Experimental/Water/Source/Runtime/Private/WaterSubsystem.cpp b/Engine/Plugins/Experimental/Water/Source/Runtime/Private/WaterSubsystem.cpp
--- a/Engine/Plugins/Experimental/Water/Source/Runtime/Private/WaterSubsystem.cpp
+++ b/Engine/Plugins/Experimental/Water/Source/Runtime/Private/WaterSubsystem.cpp
@@ -318,7 +318,7 @@ void UWaterSubsystem::Initialize(FSubsystemCollectionBase& Collection)
 	UnderwaterPostProcessVolume = NewObject<UUnderwaterPostProcessVolume>(this);
 
 	World->OnBeginPostProcessSettings.AddUObject(this, &UWaterSubsystem::ComputeUnderwaterPostProcess);
-	World->InsertPostProcessVolume(UnderwaterPostProcessVolume);
+	World->AddPostProcessVolume(UnderwaterPostProcessVolume);
 	{
 		FActorSpawnParameters SpawnInfo;
 		SpawnInfo.SpawnCollisionHandlingOverride = ESpawnActorCollisionHandlingMethod::AlwaysSpawn;
diff --git a/Engine/Source/Runtime/Engine/Classes/Engine/World.h b/Engine/Source/Runtime/Engine/Classes/Engine/World.h
--- a/Engine/Source/Runtime/Engine/Classes/Engine/World.h
+++ b/Engine/Source/Runtime/Engine/Classes/Engine/World.h
@@ -1913,15 +1913,11 @@ public:
 	// We want to store IInterface_PostProcessVolume*, but callers do not always properly call the Remove function,
 	// which leads to dangling pointers.
 	// Most PPVs are UObject-based, and so can use TWeakInterfacePtr to detect object expiration.
-	// However, some usages implement IInterface_PostProcessVolume directly without the UObject inheritance, 
-	// in which case TWeakInterfacePtr doesn't work. For these cases we have TWeakPtr as a safe option.
 	// The API signature historically allows for passing a raw pointer though, so we also support those cases for now.
 	// TLDR: Variant type for PPV storage: UObject-based PPVs use TWeakInterfacePtr (detects stale entries),
-	// non-UObject shared-ownership PPVs use TWeakPtr (detects stale entries),
 	// non-UObject raw pointer PPVs are assumed always valid.
 	using FPostProcessVolumeEntry = TVariant<
 		TWeakInterfacePtr<IInterface_PostProcessVolume>,
-		TWeakPtr<IInterface_PostProcessVolume>,
 		IInterface_PostProcessVolume*>;
 
 	/** Iterator range over post process volumes in ascending priority order, filtering out invalid/expired entries. */
@@ -1953,7 +1949,6 @@ public:
 
 			const TConstArrayView<FPostProcessVolumeEntry> Entries;
 			int32 Index;
-			TSharedPtr<IInterface_PostProcessVolume> PinnedPtr;
 			IInterface_PostProcessVolume* Current = nullptr;
 		};
 
@@ -1968,20 +1963,17 @@ public:
 		const TConstArrayView<FPostProcessVolumeEntry> Entries;
 	};
 
-	/** Inserts a non-UObject post process volume with shared ownership into the world in priority order */
-	UE_API void InsertPostProcessVolume(TSharedPtr<IInterface_PostProcessVolume>& InVolume);
-
 	/** Inserts a UObject-based post process volume into the world in priority order */
+	UE_API void AddPostProcessVolume(TWeakInterfacePtr<IInterface_PostProcessVolume> InVolume);
+
+	UE_DEPRECATED(5.8, "Use AddPostProcessVolume(TWeakInterfacePtr<IInterface_PostProcessVolume>): for memory-safety, all IInterface_PostProcessVolume need to be UObject-based")
 	UE_API void InsertPostProcessVolume(IInterface_PostProcessVolume* InVolume);
 
-	/** Removes a post process volume from the world.  Returns true if it was successfully removed. */
+	/** Removes a post process volume from the world.  Returns true if it was successfully removed. 
+	 *  Accepts both UObject and non-UObject based PPVs.
+	 */
 	UE_API bool RemovePostProcessVolume(IInterface_PostProcessVolume* InVolume);
 
-	bool RemovePostProcessVolume(TSharedPtr<IInterface_PostProcessVolume>& InVolume)
-	{
-		return RemovePostProcessVolume(InVolume.Get());
-	}
-
 	/** Called when a scene view for this world needs the worlds post process settings computed */
 	UE_API void AddPostProcessingSettings(FVector ViewLocation, FSceneView* SceneView);
 
diff --git a/Engine/Source/Runtime/Engine/Private/Components/PostProcessComponent.cpp b/Engine/Source/Runtime/Engine/Private/Components/PostProcessComponent.cpp
--- a/Engine/Source/Runtime/Engine/Private/Components/PostProcessComponent.cpp
+++ b/Engine/Source/Runtime/Engine/Private/Components/PostProcessComponent.cpp
@@ -26,7 +26,7 @@ UPostProcessComponent::UPostProcessComponent(const FObjectInitializer& ObjectIni
 void UPostProcessComponent::OnRegister()
 {
 	Super::OnRegister();
-	GetWorld()->InsertPostProcessVolume(this);
+	GetWorld()->AddPostProcessVolume(this);
 }
 
 void UPostProcessComponent::OnUnregister()
diff --git a/Engine/Source/Runtime/Engine/Private/PostProcessVolume.cpp b/Engine/Source/Runtime/Engine/Private/PostProcessVolume.cpp
--- a/Engine/Source/Runtime/Engine/Private/PostProcessVolume.cpp
+++ b/Engine/Source/Runtime/Engine/Private/PostProcessVolume.cpp
@@ -46,7 +46,7 @@ static void PostProcessVolumeOnEditorActorMoved(AActor* InActor)
 			// Don't re-insert it if we didn't actually remove it
 			if (World->RemovePostProcessVolume(VolumeActor))
 			{
-				World->InsertPostProcessVolume(VolumeActor);
+				World->AddPostProcessVolume(VolumeActor);
 			}
 		}
 	}
@@ -83,7 +83,7 @@ void APostProcessVolume::PostRegisterAllComponents()
 {
 	// Route update to super first.
 	Super::PostRegisterAllComponents();
-	GetWorld()->InsertPostProcessVolume(this);
+	GetWorld()->AddPostProcessVolume(this);
 
 #if WITH_EDITOR
 	if (GEngine)
diff --git a/Engine/Source/Runtime/Engine/Private/World.cpp b/Engine/Source/Runtime/Engine/Private/World.cpp
--- a/Engine/Source/Runtime/Engine/Private/World.cpp
+++ b/Engine/Source/Runtime/Engine/Private/World.cpp
@@ -10386,20 +10386,12 @@ FPrimaryAssetId UWorld::GetPrimaryAssetId() const
 	return FPrimaryAssetId();
 }
 
-static IInterface_PostProcessVolume* ResolvePPVEntry(
-	const UWorld::FPostProcessVolumeEntry& Entry,
-	TSharedPtr<IInterface_PostProcessVolume>& OutPinned)
+static IInterface_PostProcessVolume* ResolvePPVEntry(const UWorld::FPostProcessVolumeEntry& Entry)
 {
-	OutPinned = nullptr;
 	if (Entry.IsType<TWeakInterfacePtr<IInterface_PostProcessVolume>>())
 	{
 		return Entry.Get<TWeakInterfacePtr<IInterface_PostProcessVolume>>().Get();
 	}
-	if (Entry.IsType<TWeakPtr<IInterface_PostProcessVolume>>())
-	{
-		OutPinned = Entry.Get<TWeakPtr<IInterface_PostProcessVolume>>().Pin();
-		return OutPinned.Get();
-	}
 	return Entry.Get<IInterface_PostProcessVolume*>();
 }
 
@@ -10426,12 +10418,11 @@ static bool PostProcessVolumeLess(const FPostProcessVolumeProperties& Properties
 
 void UWorld::FPostProcessVolumeIterator::FIterator::AdvanceToValid()
 {
-	PinnedPtr = nullptr;
 	Current = nullptr;
 	while (Index < Entries.Num())
 	{
 		const FPostProcessVolumeEntry& Entry = Entries[Index];
-		Current = ResolvePPVEntry(Entry, PinnedPtr);
+		Current = ResolvePPVEntry(Entry);
 
 		if (Current)
 		{
@@ -10446,17 +10437,15 @@ void UWorld::FPostProcessVolumeIterator::FIterator::AdvanceToValid()
 
 void UWorld::InsertPostProcessVolumeEntry(FPostProcessVolumeEntry InVolume)
 {
-	TSharedPtr<IInterface_PostProcessVolume> InVolumeSharedPtr = nullptr;
-	IInterface_PostProcessVolume* InVolumeRaw = ResolvePPVEntry(InVolume, InVolumeSharedPtr);
+	IInterface_PostProcessVolume* InVolumeRaw = ResolvePPVEntry(InVolume);
 
 	int32 NumVolumes = PostProcessVolumeEntries.Num();
 	FPostProcessVolumeProperties TargetProperties = InVolumeRaw->GetProperties();
 	int32 InsertIndex = 0;
-	TSharedPtr<IInterface_PostProcessVolume> CurrentPinned;
 	// TODO: replace with binary search.
 	while (InsertIndex < NumVolumes)
 	{
-		IInterface_PostProcessVolume* CurrentVolume = ResolvePPVEntry(PostProcessVolumeEntries[InsertIndex], CurrentPinned);
+		IInterface_PostProcessVolume* CurrentVolume = ResolvePPVEntry(PostProcessVolumeEntries[InsertIndex]);
 		if (CurrentVolume == InVolumeRaw)
 		{
 			return;
@@ -10484,7 +10473,7 @@ void UWorld::InsertPostProcessVolumeEntry(FPostProcessVolumeEntry InVolume)
 	// have been reached in the loop.
 	for (int32 Index = InsertIndex + 1; Index < NumVolumes; )
 	{
-		IInterface_PostProcessVolume* CurrentVolume = ResolvePPVEntry(PostProcessVolumeEntries[Index], CurrentPinned);
+		IInterface_PostProcessVolume* CurrentVolume = ResolvePPVEntry(PostProcessVolumeEntries[Index]);
 		if (CurrentVolume == InVolumeRaw)
 		{
 			return;
@@ -10508,42 +10497,43 @@ PRAGMA_DISABLE_DEPRECATION_WARNINGS
 PRAGMA_ENABLE_DEPRECATION_WARNINGS
 }
 
+// Deprecated
 void UWorld::InsertPostProcessVolume(IInterface_PostProcessVolume* InVolume)
 {
 	check(InVolume);
 
-	// Is InVolume a UObject?
 	TWeakInterfacePtr<IInterface_PostProcessVolume> WeakInterfacePtr { InVolume };
-	if (WeakInterfacePtr.Get())
+	bool bIsUObject = WeakInterfacePtr.Get() != nullptr;
+	if (bIsUObject)
 	{
-		FPostProcessVolumeEntry Entry;
-		Entry.Set<TWeakInterfacePtr<IInterface_PostProcessVolume>>(WeakInterfacePtr);
-		InsertPostProcessVolumeEntry(MoveTemp(Entry));
+		AddPostProcessVolume(WeakInterfacePtr);
 	}
 	else
 	{
-		// InVolume is not a UObject; store as a raw pointer.
+#if DEBUG_POST_PROCESS_VOLUME_ENABLE
+		UE_LOG(LogWorld, Warning, TEXT("Post process volume %s does not derive from UObject. This path is deprecated."), *InVolume->GetDebugName());
+#endif
+
 		FPostProcessVolumeEntry Entry;
 		Entry.Set<IInterface_PostProcessVolume*>(InVolume);
 		InsertPostProcessVolumeEntry(MoveTemp(Entry));
 	}
 }
 
-void UWorld::InsertPostProcessVolume(TSharedPtr<IInterface_PostProcessVolume>& InVolume)
+void UWorld::AddPostProcessVolume(TWeakInterfacePtr<IInterface_PostProcessVolume> InVolume)
 {
-	check(InVolume);
+	check(InVolume.Get());
 
 	FPostProcessVolumeEntry Entry;
-	Entry.Set<TWeakPtr<IInterface_PostProcessVolume>>(InVolume.ToWeakPtr());
+	Entry.Set<TWeakInterfacePtr<IInterface_PostProcessVolume>>(InVolume);
 	InsertPostProcessVolumeEntry(MoveTemp(Entry));
 }
 
 bool UWorld::RemovePostProcessVolume(IInterface_PostProcessVolume* InVolume)
 {
-	TSharedPtr<IInterface_PostProcessVolume> CurrentPinned;
 	for (int32 Index = 0, NumVolumes = PostProcessVolumeEntries.Num(); Index < NumVolumes; )
 	{
-		IInterface_PostProcessVolume* CurrentVolume = ResolvePPVEntry(PostProcessVolumeEntries[Index], CurrentPinned);
+		IInterface_PostProcessVolume* CurrentVolume = ResolvePPVEntry(PostProcessVolumeEntries[Index]);
 		if (!CurrentVolume) // stale entry, volume was destroyed
 		{
 			RemovePostProcessVolume(Index);
