Re-enable RPC DoS Detection in Iris.

(Manual validations run - P4Dedupe)
#Synced-CL -1

#Synced-CL -1

diff --git a/Engine/Source/Runtime/Engine/Private/DataReplication.cpp b/Engine/Source/Runtime/Engine/Private/DataReplication.cpp
--- a/Engine/Source/Runtime/Engine/Private/DataReplication.cpp
+++ b/Engine/Source/Runtime/Engine/Private/DataReplication.cpp
@@ -25,7 +25,7 @@
 #include "Net/Core/NetCoreModule.h"
 #include "HAL/LowLevelMemStats.h"
 #include "Net/Core/PushModel/Types/PushModelPerNetDriverState.h"
-#include "Net/RPCDoSDetection.h"
+#include "Net/Core/Connection/RPCDoSDetection.h"
 #include "UObject/CoreNetContext.h"
 #include "UObject/RemoteExecutor.h"
 
@@ -1369,11 +1369,11 @@ bool FObjectReplicator::ReceivedRPC(FNetBitReader& Reader, TSet<FNetworkGUID>& U
 
 		if (UNLIKELY(RPCDoS->ShouldMonitorReceivedRPC()))
 		{
-			ERPCNotifyResult Result = RPCDoS->NotifyReceivedRPC(Reader, UnmappedGuids, Object, Function, FunctionName);
+			ERPCNotifyResult Result = RPCDoS->NotifyReceivedRPC(Object, Function, FunctionName);
 
 			if (UNLIKELY(Result == ERPCNotifyResult::BlockRPC))
 			{
-				UE_LOG(LogRepTraffic, Log, TEXT("      Blocked RPC: %s"), *FunctionName.ToString());
+				RPCDoS->TrackBlockedRPC(FunctionName);
 
 				return true;
 			}
