Input: Propagate more information to primitive component touch events, support keying by globally unique touch ID

[RN] Primitive component has new touch events that provide full context of the touch (which player, device, hit). The old events will be deprecated in the next release.
#Synced-CL 52523701

diff --git a/Engine/Plugins/Runtime/CommonUI/Source/CommonUI/Private/CommonGameViewportClient.cpp b/Engine/Plugins/Runtime/CommonUI/Source/CommonUI/Private/CommonGameViewportClient.cpp
--- a/Engine/Plugins/Runtime/CommonUI/Source/CommonUI/Private/CommonGameViewportClient.cpp
+++ b/Engine/Plugins/Runtime/CommonUI/Source/CommonUI/Private/CommonGameViewportClient.cpp
@@ -76,8 +76,8 @@ bool UCommonGameViewportClient::InputAxis(const FInputKeyEventArgs& Args)
 	return Super::InputAxis(Args);
 }
 
-bool UCommonGameViewportClient::InputTouch(FViewport* const InViewport, const FInputDeviceId DeviceId, const ETouchIndex::Type TouchId, const ETouchType::Type Type,
-	const FVector2D& TouchLocation, const float Force, const uint64 Timestamp)
+bool UCommonGameViewportClient::InputTouch(FViewport* const InViewport, const FTouchId TouchId, const ETouchType::Type Type, const FVector2D& TouchLocation,
+	const float Force, const uint64 Timestamp)
 {
 #if ALLOW_CONSOLE
 	if (ViewportConsole != nullptr && (ViewportConsole->ConsoleState != NAME_Typing) && (ViewportConsole->ConsoleState != NAME_Open))
@@ -85,16 +85,11 @@ bool UCommonGameViewportClient::InputTouch(FViewport* const InViewport, const FI
 	{
 		FReply Result = FReply::Unhandled();
 
-		// Remap the newer FInputDeviceId to the old int32 ControllerId for deprecated code. This can be removed
-		// when OnRerouteTouch() is removed.
-		const FPlatformUserId UserId = IPlatformInputDeviceMapper::Get().GetUserForInputDevice(DeviceId);
-		int32 ControllerId = IPlatformInputDeviceMapper::Get().GetUserIndexForPlatformUser(UserId);
-		 
 		PRAGMA_DISABLE_DEPRECATION_WARNINGS
-		if (!OnRerouteTouchInput().ExecuteIfBound(DeviceId, TouchId, Type, TouchLocation, Result) &&
-			!OnRerouteTouch().ExecuteIfBound(ControllerId, TouchId, Type, TouchLocation, Result))
+		if (!GetRerouteTouchDelegate().ExecuteIfBound(TouchId, Type, TouchLocation, Result) &&
+			!OnRerouteTouchInput().ExecuteIfBound(TouchId.GetDeviceId(), TouchId.GetIndex(), Type, TouchLocation, Result))
 		{
-			HandleRerouteTouch(DeviceId, TouchId, Type, TouchLocation, Result);
+			HandleRerouteTouch(TouchId, Type, TouchLocation, Result);
 		}
 		PRAGMA_ENABLE_DEPRECATION_WARNINGS
 
@@ -104,7 +99,7 @@ bool UCommonGameViewportClient::InputTouch(FViewport* const InViewport, const FI
 		}
 	}
 
-	return Super::InputTouch(InViewport, DeviceId, TouchId, Type, TouchLocation, Force, Timestamp);
+	return Super::InputTouch(InViewport, TouchId, Type, TouchLocation, Force, Timestamp);
 }
 
 void UCommonGameViewportClient::MouseMove(FViewport* InViewport, int32 X, int32 Y)
@@ -172,14 +167,21 @@ void UCommonGameViewportClient::HandleRerouteAxis(FInputDeviceId DeviceId, FKey
 	}
 }
 
+PRAGMA_DISABLE_DEPRECATION_WARNINGS
 void UCommonGameViewportClient::HandleRerouteTouch(FInputDeviceId DeviceId, uint32 TouchId, ETouchType::Type TouchType, const FVector2D& TouchLocation, FReply& Reply)
 {
-	ULocalPlayer* LocalPlayer = GameInstance->FindLocalPlayerFromDeviceId(DeviceId);
+	return HandleRerouteTouch(FTouchId(DeviceId, static_cast<ETouchIndex::Type>(TouchId)), TouchType, TouchLocation, Reply);
+}
+PRAGMA_ENABLE_DEPRECATION_WARNINGS
+
+void UCommonGameViewportClient::HandleRerouteTouch(FTouchId TouchId, ETouchType::Type TouchType, const FVector2D& TouchLocation, FReply& Reply)
+{
+	ULocalPlayer* LocalPlayer = GameInstance->FindLocalPlayerFromDeviceId(TouchId.GetDeviceId());
 	Reply = FReply::Unhandled();
 
-	if (LocalPlayer && TouchId < EKeys::NUM_TOUCH_KEYS)
+	if (LocalPlayer && TouchId.GetIndex() < EKeys::NUM_TOUCH_KEYS)
 	{
-		FKey KeyPressed = EKeys::TouchKeys[TouchId];
+		FKey KeyPressed = EKeys::TouchKeys[TouchId.GetIndex()];
 		if (KeyPressed.IsValid())
 		{
 			UCommonUIActionRouterBase* ActionRouter = LocalPlayer->GetSubsystem<UCommonUIActionRouterBase>();
@@ -208,18 +210,6 @@ void UCommonGameViewportClient::HandleRerouteTouch(FInputDeviceId DeviceId, uint
 	}
 }
 
-PRAGMA_DISABLE_DEPRECATION_WARNINGS
-void UCommonGameViewportClient::HandleRerouteTouch(int32 ControllerId, uint32 TouchId, ETouchType::Type TouchType, const FVector2D& TouchLocation, FReply& Reply)
-{
-	// Remap the old int32 ControllerId to the new platform user and input device ID
-	FPlatformUserId UserId = FGenericPlatformMisc::GetPlatformUserForUserIndex(ControllerId);
-	FInputDeviceId DeviceID = INPUTDEVICEID_NONE;
-	IPlatformInputDeviceMapper::Get().RemapControllerIdToPlatformUserAndDevice(ControllerId, UserId, DeviceID);
-	
-	return HandleRerouteTouch(DeviceID, TouchId, TouchType, TouchLocation, Reply);
-}
-PRAGMA_ENABLE_DEPRECATION_WARNINGS
-
 bool UCommonGameViewportClient::IsKeyPriorityAboveUI(const FInputKeyEventArgs& EventArgs)
 {
 #if ALLOW_CONSOLE
diff --git a/Engine/Plugins/Runtime/CommonUI/Source/CommonUI/Public/CommonGameViewportClient.h b/Engine/Plugins/Runtime/CommonUI/Source/CommonUI/Public/CommonGameViewportClient.h
--- a/Engine/Plugins/Runtime/CommonUI/Source/CommonUI/Public/CommonGameViewportClient.h
+++ b/Engine/Plugins/Runtime/CommonUI/Source/CommonUI/Public/CommonGameViewportClient.h
@@ -12,11 +12,11 @@ class FReply;
 DECLARE_DELEGATE_FourParams(FOnRerouteInputDelegate, FInputDeviceId /* InputDeviceId */, FKey /* Key */, EInputEvent /* EventType */, FReply& /* Reply */);
 DECLARE_DELEGATE_FourParams(FOnRerouteAxisDelegate, FInputDeviceId /* InputDeviceId */, FKey /* Key */, float /* Delta */, FReply& /* Reply */);
 
-UE_DEPRECATED(5.6, "Use the version which takes an FInputDeviceId instead (FOnRerouteTouchInputDelegate).")
-DECLARE_DELEGATE_FiveParams(FOnRerouteTouchDelegate, int32 /* ControllerId */, uint32 /* TouchId */, ETouchType::Type /* TouchType */, const FVector2D& /* TouchLocation */, FReply& /* Reply */);
-
+UE_DEPRECATED(5.8, "Use the version which takes an FTouchId instead (FTouchRerouteDelegate).")
 DECLARE_DELEGATE_FiveParams(FOnRerouteTouchInputDelegate, FInputDeviceId /* Deviceid */, uint32 /* TouchId */, ETouchType::Type /* TouchType */, const FVector2D& /* TouchLocation */, FReply& /* Reply */);
 
+DECLARE_DELEGATE_FourParams(FTouchRerouteDelegate, FTouchId /* TouchId */, ETouchType::Type /* TouchType */, const FVector2D& /* TouchLocation */, FReply& /* Reply */);
+
 /**
 * CommonUI Viewport to reroute input to UI first. Needed to allow CommonUI to route / handle inputs.
 */
@@ -32,20 +32,26 @@ public:
 	// UGameViewportClient interface begin
 	UE_API virtual bool InputKey(const FInputKeyEventArgs& EventArgs) override;
 	UE_API virtual bool InputAxis(const FInputKeyEventArgs& Args) override;
-	UE_API virtual bool InputTouch(FViewport* const InViewport, const FInputDeviceId DeviceId, const ETouchIndex::Type TouchId, const ETouchType::Type Type, const FVector2D& TouchLocation, const float Force, const uint64 Timestamp) override;
+	UE_API virtual bool InputTouch(FViewport* const InViewport, const FTouchId TouchId, const ETouchType::Type Type, const FVector2D& TouchLocation, const float Force, const uint64 Timestamp) override;
 	UE_API virtual void MouseMove(FViewport* InViewport, int32 X, int32 Y) override;
 	UE_API virtual void CapturedMouseMove(FViewport* InViewport, int32 X, int32 Y) override;
 	// UGameViewportClient interface end
 
 	FOnRerouteInputDelegate& OnRerouteInput() { return RerouteInput; }
 	FOnRerouteAxisDelegate& OnRerouteAxis() { return RerouteAxis; }
-	
-	UE_DEPRECATED(5.6, "Use the version which takes a FInputDeviceId instead (OnRerouteTouchInput).")
+
+	UE_DEPRECATED(5.8, "Use GetRerouteTouchRegistration() instead.")
 	PRAGMA_DISABLE_DEPRECATION_WARNINGS
-	FOnRerouteTouchDelegate& OnRerouteTouch() { return RerouteTouch; }
+	FOnRerouteTouchInputDelegate& OnRerouteTouchInput()
+	{
+		return RerouteTouchInput;
+	}
 	PRAGMA_ENABLE_DEPRECATION_WARNINGS
-	
-	FOnRerouteTouchInputDelegate& OnRerouteTouchInput() { return RerouteTouchInput; }
+
+	FTouchRerouteDelegate::RegistrationType& GetRerouteTouchRegistration()
+	{
+		return RerouteTouchDelegate;
+	}
 
 	FOnRerouteInputDelegate& OnRerouteBlockedInput() { return RerouteBlockedInput; }
 
@@ -55,28 +61,35 @@ public:
 	/** Default Handler for Axis input. */
 	UE_API virtual void HandleRerouteAxis(FInputDeviceId DeviceId, FKey Key, float Delta, FReply& Reply);
 
+	UE_DEPRECATED(5.8, "Use the version which takes a FTouchId instead.")
+	UE_API virtual void HandleRerouteTouch(FInputDeviceId DeviceId, uint32 TouchId, ETouchType::Type TouchType, const FVector2D& TouchLocation, FReply& Reply) final;
+
 	/** Default Handler for Touch input. */
-	UE_DEPRECATED(5.6, "Use the version which takes a FInputDeviceId instead.")
-	UE_API virtual void HandleRerouteTouch(int32 ControllerId, uint32 TouchId, ETouchType::Type TouchType, const FVector2D& TouchLocation, FReply& Reply) final;
-	
-	UE_API virtual void HandleRerouteTouch(FInputDeviceId DeviceId, uint32 TouchId, ETouchType::Type TouchType, const FVector2D& TouchLocation, FReply& Reply);
+	UE_API virtual void HandleRerouteTouch(FTouchId TouchId, ETouchType::Type TouchType, const FVector2D& TouchLocation, FReply& Reply);
 
 protected:
-	
+
+	FTouchRerouteDelegate& GetRerouteTouchDelegate()
+	{
+		return RerouteTouchDelegate;
+	}
+
 	/** Console window & fullscreen shortcut have higher priority than UI */
 	UE_API virtual bool IsKeyPriorityAboveUI(const FInputKeyEventArgs& EventArgs);
 
 	FOnRerouteInputDelegate RerouteInput;
 	FOnRerouteAxisDelegate RerouteAxis;
 	
-	UE_DEPRECATED(5.6, "Use RerouteTouchInput instead.")
+	UE_DEPRECATED(5.8, "Use GetRerouteTouchDelegate() instead.")
 	PRAGMA_DISABLE_DEPRECATION_WARNINGS
-	FOnRerouteTouchDelegate RerouteTouch;
-	PRAGMA_ENABLE_DEPRECATION_WARNINGS
-	
 	FOnRerouteTouchInputDelegate RerouteTouchInput;
+	PRAGMA_ENABLE_DEPRECATION_WARNINGS
 
 	FOnRerouteInputDelegate RerouteBlockedInput;
+
+private:
+
+	FTouchRerouteDelegate RerouteTouchDelegate;
 };
 
 #undef UE_API
diff --git a/Engine/Source/Runtime/Engine/Classes/Components/PrimitiveComponent.h b/Engine/Source/Runtime/Engine/Classes/Components/PrimitiveComponent.h
--- a/Engine/Source/Runtime/Engine/Classes/Components/PrimitiveComponent.h
+++ b/Engine/Source/Runtime/Engine/Classes/Components/PrimitiveComponent.h
@@ -44,6 +44,7 @@ DECLARE_CYCLE_STAT_EXTERN(TEXT("BeginComponentOverlap"), STAT_BeginComponentOver
 DECLARE_CYCLE_STAT_EXTERN(TEXT("MoveComponent FastOverlap"), STAT_MoveComponent_FastOverlap, STATGROUP_Game, ENGINE_API);
 
 class AController;
+class APlayerController;
 class FPrimitiveSceneProxy;
 class UMaterialInterface;
 class UPhysicalMaterial;
@@ -274,6 +275,10 @@ DECLARE_DYNAMIC_MULTICAST_SPARSE_DELEGATE_OneParam( FComponentBeginCursorOverSig
 DECLARE_DYNAMIC_MULTICAST_SPARSE_DELEGATE_OneParam( FComponentEndCursorOverSignature, UPrimitiveComponent, OnEndCursorOver, UPrimitiveComponent*, TouchedComponent );
 DECLARE_DYNAMIC_MULTICAST_SPARSE_DELEGATE_TwoParams( FComponentOnClickedSignature, UPrimitiveComponent, OnClicked, UPrimitiveComponent*, TouchedComponent , FKey, ButtonPressed);
 DECLARE_DYNAMIC_MULTICAST_SPARSE_DELEGATE_TwoParams( FComponentOnReleasedSignature, UPrimitiveComponent, OnReleased, UPrimitiveComponent*, TouchedComponent, FKey, ButtonReleased);
+DECLARE_DYNAMIC_MULTICAST_SPARSE_DELEGATE_FourParams( FComponentOnTouchBeganSignature, UPrimitiveComponent, OnTouchBegan, const FTouchId, TouchId, APlayerController*, Controller, const FHitResult&, HitResult, UPrimitiveComponent*, TouchedComponent );
+DECLARE_DYNAMIC_MULTICAST_SPARSE_DELEGATE_FourParams( FComponentOnTouchEndedSignature, UPrimitiveComponent, OnTouchEnded, const FTouchId, TouchId, APlayerController*, Controller, const FHitResult&, HitResult, UPrimitiveComponent*, TouchedComponent );
+DECLARE_DYNAMIC_MULTICAST_SPARSE_DELEGATE_FourParams( FComponentOnTouchEnteredSignature, UPrimitiveComponent, OnTouchEntered, const FTouchId, TouchId, APlayerController*, Controller, const FHitResult&, HitResult, UPrimitiveComponent*, TouchedComponent );
+DECLARE_DYNAMIC_MULTICAST_SPARSE_DELEGATE_FourParams( FComponentOnTouchExitedSignature, UPrimitiveComponent, OnTouchExited, const FTouchId, TouchId, APlayerController*, Controller, const FHitResult&, HitResult, UPrimitiveComponent*, TouchedComponent );
 DECLARE_DYNAMIC_MULTICAST_SPARSE_DELEGATE_TwoParams( FComponentOnInputTouchBeginSignature, UPrimitiveComponent, OnInputTouchBegin, ETouchIndex::Type, FingerIndex, UPrimitiveComponent*, TouchedComponent );
 DECLARE_DYNAMIC_MULTICAST_SPARSE_DELEGATE_TwoParams( FComponentOnInputTouchEndSignature, UPrimitiveComponent, OnInputTouchEnd, ETouchIndex::Type, FingerIndex, UPrimitiveComponent*, TouchedComponent );
 DECLARE_DYNAMIC_MULTICAST_SPARSE_DELEGATE_TwoParams( FComponentBeginTouchOverSignature, UPrimitiveComponent, OnInputTouchEnter, ETouchIndex::Type, FingerIndex, UPrimitiveComponent*, TouchedComponent );
@@ -1493,22 +1498,46 @@ public:
 	/** Event called when the left mouse button is released while the mouse is over this component click events are enabled in the player controller */
 	UPROPERTY(BlueprintAssignable, Category="Input|Mouse Input")
 	FComponentOnReleasedSignature OnReleased;
-		 
+
+	/** Event called when a touch input is received over this component when touch events are enabled in the player controller
+	 * This will be deprecated in the next engine release.  Use OnTouchBegan with TouchId, Player Controller, and Hit Result arguments instead.
+	 */
+	UPROPERTY(BlueprintAssignable, Category="Input|Touch Input", DisplayName="On Touch Began (finger only)")
+	FComponentOnInputTouchBeginSignature OnInputTouchBegin;
+
+	/** Event called when a touch input is released over this component when touch events are enabled in the player controller
+	 * This will be deprecated in the next engine release.  Use OnTouchEnded with TouchId, Player Controller, and Hit Result arguments instead.
+	 */
+	UPROPERTY(BlueprintAssignable, Category="Input|Touch Input", DisplayName="On Touch Ended (finger only)")
+	FComponentOnInputTouchEndSignature OnInputTouchEnd;
+
+	/** Event called when a finger is moved over this component when touch over events are enabled in the player controller
+	 * This will be deprecated in the next engine release.  Use OnTouchEntered with TouchId, Player Controller, and Hit Result arguments instead.
+	 */
+	UPROPERTY(BlueprintAssignable, Category="Input|Touch Input", DisplayName="On Touch Entered (finger only)")
+	FComponentBeginTouchOverSignature OnInputTouchEnter;
+
+	/** Event called when a finger is moved off this component when touch over events are enabled in the player controller
+	 * This will be deprecated in the next engine release.  Use OnTouchExited with TouchId, Player Controller, and Hit Result arguments instead.
+	 */
+	UPROPERTY(BlueprintAssignable, Category="Input|Touch Input", DisplayName="On Touch Exited (finger only)")
+	FComponentEndTouchOverSignature OnInputTouchLeave;
+
 	/** Event called when a touch input is received over this component when touch events are enabled in the player controller */
 	UPROPERTY(BlueprintAssignable, Category="Input|Touch Input")
-	FComponentOnInputTouchBeginSignature OnInputTouchBegin;
+	FComponentOnTouchBeganSignature OnTouchBegan;
 
 	/** Event called when a touch input is released over this component when touch events are enabled in the player controller */
 	UPROPERTY(BlueprintAssignable, Category="Input|Touch Input")
-	FComponentOnInputTouchEndSignature OnInputTouchEnd;
+	FComponentOnTouchEndedSignature OnTouchEnded;
 
 	/** Event called when a finger is moved over this component when touch over events are enabled in the player controller */
 	UPROPERTY(BlueprintAssignable, Category="Input|Touch Input")
-	FComponentBeginTouchOverSignature OnInputTouchEnter;
+	FComponentOnTouchEnteredSignature OnTouchEntered;
 
 	/** Event called when a finger is moved off this component when touch over events are enabled in the player controller */
 	UPROPERTY(BlueprintAssignable, Category="Input|Touch Input")
-	FComponentEndTouchOverSignature OnInputTouchLeave;
+	FComponentOnTouchExitedSignature OnTouchExited;
 
 	/** Scale the bounds of this object, used for frustum culling. Useful for features like WorldPositionOffset. */
 	UFUNCTION(BlueprintCallable, Category = "Rendering")
@@ -3165,12 +3194,25 @@ public:
 	ENGINE_API void SetCustomNavigableGeometry(const EHasCustomNavigableGeometry::Type InType);
 
 	static ENGINE_API void DispatchMouseOverEvents(UPrimitiveComponent* CurrentComponent, UPrimitiveComponent* NewComponent);
+
+	UE_DEPRECATED(5.8, "Use overload with Touch ID, Player Controller, and Hit Result." )
 	static ENGINE_API void DispatchTouchOverEvents(ETouchIndex::Type FingerIndex, UPrimitiveComponent* CurrentComponent, UPrimitiveComponent* NewComponent);
+
+	static ENGINE_API void DispatchTouchOverEvents(const FTouchId TouchId, APlayerController* const Controller, const FHitResult& HitResult, UPrimitiveComponent* const CurrentComponent, UPrimitiveComponent* const NewComponent);
+
 	ENGINE_API void DispatchOnClicked(FKey ButtonClicked = EKeys::LeftMouseButton);
 	ENGINE_API void DispatchOnReleased(FKey ButtonReleased = EKeys::LeftMouseButton);
+
+	UE_DEPRECATED(5.8, "Use overload with Touch ID, Player Controller, and Hit Result." )
 	ENGINE_API void DispatchOnInputTouchBegin(const ETouchIndex::Type Key);
+
+	UE_DEPRECATED(5.8, "Use overload with Touch ID, Player Controller, and Hit Result." )
 	ENGINE_API void DispatchOnInputTouchEnd(const ETouchIndex::Type Key);
 
+	ENGINE_API void DispatchOnInputTouchBegin(const FTouchId TouchId, APlayerController* const Controller, const FHitResult& HitResult);
+
+	ENGINE_API void DispatchOnInputTouchEnd(const FTouchId TouchId, APlayerController* const Controller, const FHitResult& HitResult);
+
 	//~ Begin IPhysicsComponent Interface.
 public:
 	ENGINE_API virtual Chaos::FPhysicsObject* GetPhysicsObjectById(Chaos::FPhysicsObjectId Id) const override;
diff --git a/Engine/Source/Runtime/Engine/Classes/Engine/Console.h b/Engine/Source/Runtime/Engine/Classes/Engine/Console.h
--- a/Engine/Source/Runtime/Engine/Classes/Engine/Console.h
+++ b/Engine/Source/Runtime/Engine/Classes/Engine/Console.h
@@ -273,25 +273,24 @@ class UConsole
 	
 	virtual bool InputAxis(FInputDeviceId DevideId, FKey Key, float Delta, float DeltaTime, int32 NumSamples=1, bool bGamepad=false) { return false; };
 
-	UE_DEPRECATED(5.8, "Use the overload that omits TouchpadIndex instead")
+	UE_DEPRECATED(5.8, "Use the overload that uses FTouchId instead")
 	virtual bool InputTouch(const FInputDeviceId DevideId, const uint32 Handle, const ETouchType::Type Type, const FVector2D& TouchLocation, const float Force, const FDateTime DeviceTimestamp, const uint32 TouchpadIndex) final
 	{
-		return InputTouch(DevideId, static_cast<ETouchIndex::Type>(Handle), Type, TouchLocation, Force, DeviceTimestamp.GetTicks());
+		return InputTouch(FTouchId(DevideId, static_cast<ETouchIndex::Type>(Handle)), Type, TouchLocation, Force, DeviceTimestamp.GetTicks());
 	}
 
 	/**
 	 * Allow the console to process a touch event received by the viewport-client.
 	 * If the console uses the event, it should return true.  If true is returned, the viewport-client will not offer it to the next candidate.
 	 * 
-	 * @param	DeviceId - The input device that was touched.
-	 * @param	TouchId - Identifier unique to this ongoing touch (i.e. finger).
+	 * @param	TouchId - Identifier unique to this ongoing touch (i.e. device + finger).
 	 * @param	Type - What kind of touch event this is (see ETouchType).
 	 * @param	TouchLocation - Viewport-relative position of the touch.
 	 * @param	Force - How hard the touch is.
 	 * @param	Timestamp - The timestamp of when this input event was originally polled in the form of CPU cycles.
 	 * @return	True to consume the key event, false to pass it on.
 	 */
-	virtual bool InputTouch(const FInputDeviceId DeviceId, const ETouchIndex::Type TouchId, const ETouchType::Type Type, const FVector2D& TouchLocation, const float Force, const uint64 Timestamp)
+	virtual bool InputTouch(const FTouchId TouchId, const ETouchType::Type Type, const FVector2D& TouchLocation, const float Force, const uint64 Timestamp)
 	{
 		return false;
 	}
diff --git a/Engine/Source/Runtime/Engine/Classes/Engine/GameViewportClient.h b/Engine/Source/Runtime/Engine/Classes/Engine/GameViewportClient.h
--- a/Engine/Source/Runtime/Engine/Classes/Engine/GameViewportClient.h
+++ b/Engine/Source/Runtime/Engine/Classes/Engine/GameViewportClient.h
@@ -172,7 +172,7 @@ public:
 	ENGINE_API virtual bool InputKey(const FInputKeyEventArgs& EventArgs) override;
 	ENGINE_API virtual bool InputAxis(const FInputKeyEventArgs& EventArgs) override;
 	ENGINE_API virtual bool InputChar(FViewport* Viewport,int32 ControllerId, TCHAR Character) override;
-	ENGINE_API virtual bool InputTouch(FViewport* const InViewport, const FInputDeviceId DeviceId, const ETouchIndex::Type TouchId, const ETouchType::Type Type, const FVector2D& TouchLocation, const float Force, const uint64 Timestamp) override;
+	ENGINE_API virtual bool InputTouch(FViewport* const InViewport, const FTouchId TouchId, const ETouchType::Type Type, const FVector2D& TouchLocation, const float Force, const uint64 Timestamp) override;
 	ENGINE_API virtual bool InputMotion(FViewport* Viewport, const FInputDeviceId DeviceId, const FVector& Tilt, const FVector& RotationRate, const FVector& Gravity, const FVector& Acceleration, const uint64 Timestamp) override;
 	ENGINE_API virtual EMouseCursor::Type GetCursor(FViewport* Viewport, int32 X, int32 Y ) override;
 	ENGINE_API virtual TOptional<TSharedRef<SWidget>> MapCursor(FViewport* Viewport, const FCursorReply& CursorReply) override;
diff --git a/Engine/Source/Runtime/Engine/Classes/GameFramework/PlayerController.h b/Engine/Source/Runtime/Engine/Classes/GameFramework/PlayerController.h
--- a/Engine/Source/Runtime/Engine/Classes/GameFramework/PlayerController.h
+++ b/Engine/Source/Runtime/Engine/Classes/GameFramework/PlayerController.h
@@ -1789,7 +1789,7 @@ public:
 	ENGINE_API virtual bool InputKey(const FInputKeyEventArgs& Params);
 
 	/** Handles a touch screen action */
-	UE_DEPRECATED(5.8, "Use the overload that omits TouchpadIndex instead")
+	UE_DEPRECATED(5.8, "Use the overload that uses FTouchId instead")
 	ENGINE_API virtual bool InputTouch(
 		const FInputDeviceId DeviceId,
 		const uint32 Handle,
@@ -1804,8 +1804,7 @@ public:
 	 * If the player-controller uses the event, it should return true.  If true is returned, the viewport-client will not offer it to the next
 	 * candidate.
 	 * 
-	 * @param	DeviceId - The input device that was touched.
-	 * @param	TouchId - Identifier unique to this ongoing touch (i.e. finger).
+	 * @param	TouchId - Identifier unique to this ongoing touch (i.e. device + finger).
 	 * @param	Type - What kind of touch event this is (see ETouchType).
 	 * @param	TouchLocation - Viewport-relative position of the touch.
 	 * @param	Force - How hard the touch is.
@@ -1813,8 +1812,7 @@ public:
 	 * @return	True to consume the event, false to pass it on.
 	 */
 	ENGINE_API virtual bool InputTouch(
-		const FInputDeviceId DeviceId,
-		const ETouchIndex::Type TouchId,
+		const FTouchId TouchId,
 		const ETouchType::Type Type,
 		const FVector2D& TouchLocation,
 		const float Force,
@@ -1998,7 +1996,11 @@ protected:
 	ENGINE_API virtual bool IsInViewportClient(UGameViewportClient* ViewportClient) const;
 	ENGINE_API virtual int32 GetInputIndex() const;
 	ENGINE_API virtual ACameraActor* GetAutoActivateCameraForPlayer() const;
-	ENGINE_API virtual void ProcessTouchHitResult(const FInputDeviceId DeviceId, uint32 Handle, ETouchType::Type Type, const FVector2D& TouchLocation, const FHitResult& HitResult, const bool bHit);
+
+	UE_DEPRECATED(5.8, "Use the overload with FTouchId instead")
+	ENGINE_API virtual void ProcessTouchHitResult(const FInputDeviceId DeviceId, uint32 Handle, ETouchType::Type Type, const FVector2D& TouchLocation, const FHitResult& HitResult, const bool bHit) final;
+
+	ENGINE_API virtual void ProcessTouchHitResult(const FTouchId TouchId, const ETouchType::Type Type, const FVector2D& TouchLocation, const FHitResult& HitResult, const bool bHit);
 
 	/** Allows the PlayerController to set up custom input bindings. */
 	ENGINE_API virtual void SetupInputComponent();
diff --git a/Engine/Source/Runtime/Engine/Classes/GameFramework/PlayerInput.h b/Engine/Source/Runtime/Engine/Classes/GameFramework/PlayerInput.h
--- a/Engine/Source/Runtime/Engine/Classes/GameFramework/PlayerInput.h
+++ b/Engine/Source/Runtime/Engine/Classes/GameFramework/PlayerInput.h
@@ -14,6 +14,7 @@
 #include "InputCoreTypes.h"
 #include "InputKeyEventArgs.h"
 #include "KeyState.h"
+#include "Containers/StaticArray.h"
 #include "UObject/Object.h"
 #include "UObject/ObjectMacros.h"
 #include "UObject/ScriptMacros.h"
@@ -411,6 +412,16 @@ struct UE_DEPRECATED(5.6, "Use FInputKeyEventArgs instead") FInputKeyParams
 
 ENGINE_API DECLARE_LOG_CATEGORY_EXTERN(LogPlayerInput, Log, All);
 
+namespace UE
+{
+	struct FTouchState
+	{
+		FInputDeviceId DeviceId = INPUTDEVICEID_NONE;
+		FVector2f ViewportRelativeLocation = FVector2f(-1.0f, -1.0f);
+		float Pressure = 0.0f;
+	};
+}
+
 /**
  * Object within PlayerController that processes player input.
  * Only exists on the client in network games.
@@ -427,10 +438,16 @@ public:
 	ENGINE_API UPlayerInput();
 
 	// NOTE: These touch vectors are calculated and set directly, they do not go through the .ini Bindings
-	// Touch locations, from 0..1 (0,0 is top left, 1,1 is bottom right), the Z component is > 0 if the touch is currently held down
 	// @todo: We have 10 touches to match the number of Touch* entries in EKeys (not easy to make this an enum or define or anything)
+	UE_DEPRECATED(5.8, "Use GetTouchStates() instead")
 	FVector Touches[EKeys::NUM_TOUCH_KEYS];
 
+	/** The most recent state reported for each finger. */
+	TArrayView<const UE::FTouchState> GetTouchStates() const
+	{
+		return TouchStates;
+	}
+
 	/** Used to store paired touch locations for event ids during the frame and flushed when processed. */
 	TMap<uint32, FVector> TouchEventLocations;
 
@@ -605,7 +622,7 @@ public:
 	ENGINE_API virtual bool InputKey(const FInputKeyEventArgs& Params);
 
 	/** Handles a touch input event.  Returns true. */
-	UE_DEPRECATED(5.8, "Use the overload that omits TouchpadIndex instead")
+	UE_DEPRECATED(5.8, "Use the overload that uses FTouchId instead")
 	ENGINE_API bool InputTouch(
 		const FInputDeviceId DeviceId,
 		const uint32 Handle,
@@ -617,8 +634,7 @@ public:
 
 	/** Handles a touch input event.  Returns true. */
 	ENGINE_API bool InputTouch(
-		const FInputDeviceId DeviceId,
-		const ETouchIndex::Type TouchId,
+		const FTouchId TouchId,
 		const ETouchType::Type Type,
 		const FVector2D& TouchLocation,
 		const float Force,
@@ -878,6 +894,8 @@ protected:
 	/** Retrieves the actions bound to the given input component which are triggered by a given key. */
 	ENGINE_API void GetActionsBoundToKey(UInputComponent* InputComponent, FKey Key, TArray<TSharedPtr<FInputActionBinding>>& Actions);
 
+	TStaticArray<UE::FTouchState, EKeys::NUM_TOUCH_KEYS> TouchStates;
+
 private:
 
 	/** Gesture recognizer object */
diff --git a/Engine/Source/Runtime/Engine/Private/Components/PrimitiveComponent.cpp b/Engine/Source/Runtime/Engine/Private/Components/PrimitiveComponent.cpp
--- a/Engine/Source/Runtime/Engine/Private/Components/PrimitiveComponent.cpp
+++ b/Engine/Source/Runtime/Engine/Private/Components/PrimitiveComponent.cpp
@@ -4953,6 +4953,49 @@ void UPrimitiveComponent::DispatchOnInputTouchEnd(const ETouchIndex::Type Finger
 	}
 }
 
+void UPrimitiveComponent::DispatchOnInputTouchBegin(const FTouchId TouchId, APlayerController* const Controller, const FHitResult& HitResult)
+{
+PRAGMA_DISABLE_DEPRECATION_WARNINGS
+	DispatchOnInputTouchBegin(TouchId.GetIndex());
+PRAGMA_ENABLE_DEPRECATION_WARNINGS
+
+	if (IsValidChecked(this))
+	{
+		OnTouchBegan.Broadcast(TouchId, Controller, HitResult, this);
+	}
+}
+
+void UPrimitiveComponent::DispatchOnInputTouchEnd(const FTouchId TouchId, APlayerController* const Controller, const FHitResult& HitResult)
+{
+PRAGMA_DISABLE_DEPRECATION_WARNINGS
+	DispatchOnInputTouchEnd(TouchId.GetIndex());
+PRAGMA_ENABLE_DEPRECATION_WARNINGS
+
+	if (IsValidChecked(this))
+	{
+		OnTouchEnded.Broadcast(TouchId, Controller, HitResult, this);
+	}
+}
+
+void UPrimitiveComponent::DispatchTouchOverEvents(const FTouchId TouchId, APlayerController* const Controller, const FHitResult& HitResult, UPrimitiveComponent* const CurrentComponent, UPrimitiveComponent* const NewComponent)
+{
+PRAGMA_DISABLE_DEPRECATION_WARNINGS
+	DispatchTouchOverEvents(TouchId.GetIndex(), CurrentComponent, NewComponent);
+PRAGMA_ENABLE_DEPRECATION_WARNINGS
+
+	if (CurrentComponent != NewComponent)
+	{
+		if (IsValid(CurrentComponent))
+		{
+			CurrentComponent->OnTouchExited.Broadcast(TouchId, Controller, HitResult, CurrentComponent);
+		}
+		if (IsValid(NewComponent))
+		{
+			NewComponent->OnTouchEntered.Broadcast(TouchId, Controller, HitResult, NewComponent);
+		}
+	}
+}
+
 void UPrimitiveComponent::GetResourceSizeEx(FResourceSizeEx& CumulativeResourceSize)
 {
 	Super::GetResourceSizeEx(CumulativeResourceSize);
diff --git a/Engine/Source/Runtime/Engine/Private/GameViewportClient.cpp b/Engine/Source/Runtime/Engine/Private/GameViewportClient.cpp
--- a/Engine/Source/Runtime/Engine/Private/GameViewportClient.cpp
+++ b/Engine/Source/Runtime/Engine/Private/GameViewportClient.cpp
@@ -901,8 +901,13 @@ bool UGameViewportClient::InputChar(FViewport* InViewport, int32 ControllerId, T
 	return bResult;
 }
 
-bool UGameViewportClient::InputTouch(FViewport* const InViewport, const FInputDeviceId DeviceId, const ETouchIndex::Type TouchId, const ETouchType::Type Type,
-	const FVector2D& TouchLocation, const float Force, const uint64 Timestamp)
+bool UGameViewportClient::InputTouch(
+	FViewport* const InViewport,
+	const FTouchId TouchId,
+	const ETouchType::Type Type,
+	const FVector2D& TouchLocation,
+	const float Force,
+	const uint64 Timestamp)
 {
 	if (IgnoreInput())
 	{
@@ -910,13 +915,13 @@ bool UGameViewportClient::InputTouch(FViewport* const InViewport, const FInputDe
 	}
 
 	// route to subsystems that care
-	bool bResult = (ViewportConsole ? ViewportConsole->InputTouch(DeviceId, TouchId, Type, TouchLocation, Force, Timestamp) : false);
+	bool bResult = (ViewportConsole ? ViewportConsole->InputTouch(TouchId, Type, TouchLocation, Force, Timestamp) : false);
 	if (!bResult)
 	{
-		ULocalPlayer* const TargetPlayer = GEngine->GetLocalPlayerFromInputDevice(this, DeviceId);
+		ULocalPlayer* const TargetPlayer = GEngine->GetLocalPlayerFromInputDevice(this, TouchId.GetDeviceId());
 		if (TargetPlayer && TargetPlayer->PlayerController)
 		{
-			bResult = TargetPlayer->PlayerController->InputTouch(DeviceId, TouchId, Type, TouchLocation, Force, Timestamp);
+			bResult = TargetPlayer->PlayerController->InputTouch(TouchId, Type, TouchLocation, Force, Timestamp);
 		}
 	}
 
diff --git a/Engine/Source/Runtime/Engine/Private/GestureRecognizer.cpp b/Engine/Source/Runtime/Engine/Private/GestureRecognizer.cpp
--- a/Engine/Source/Runtime/Engine/Private/GestureRecognizer.cpp
+++ b/Engine/Source/Runtime/Engine/Private/GestureRecognizer.cpp
@@ -9,13 +9,13 @@
 #include "Engine/LocalPlayer.h"
 #include "GameFramework/PlayerInput.h"
 
-void FGestureRecognizer::DetectGestures(const FVector (&Touches)[EKeys::NUM_TOUCH_KEYS], UPlayerInput* PlayerInput, float DeltaTime)
+void FGestureRecognizer::DetectGestures(TArrayView<const UE::FTouchState> TouchStates, UPlayerInput* PlayerInput, float DeltaTime)
 {
 	// how many fingers currently held down?
 	int32 TouchCount = 0;
-	for (int32 TouchIndex = 0; TouchIndex < EKeys::NUM_TOUCH_KEYS; TouchIndex++)
+	for (int32 TouchIndex = 0; TouchIndex < TouchStates.Num(); TouchIndex++)
 	{
-		if (Touches[TouchIndex].Z != 0)
+		if (TouchStates[TouchIndex].Pressure != 0)
 		{
 			TouchCount++;
 		}
@@ -29,7 +29,7 @@ void FGestureRecognizer::DetectGestures(const FVector (&Touches)[EKeys::NUM_TOUC
 		{
 			if (PreviousTouchCount < Index + 1 && TouchCount >= Index + 1)
 			{
-				AnchorPoints[Index] = FVector2D(Touches[Index]);
+				AnchorPoints[Index] = FVector2D(TouchStates[Index].ViewportRelativeLocation);
 			}
 		}
 
@@ -39,15 +39,15 @@ void FGestureRecognizer::DetectGestures(const FVector (&Touches)[EKeys::NUM_TOUC
 			float* CurrentAlpha = CurrentGestureValues.Find(EKeys::Gesture_Pinch);
 
 			//#jira UE-51262 values for pinch input produce very different results for same area on android device
-			FVector2D CurrentPinchPoint_Start = FVector2D(Touches[0]);
-			FVector2D CurrentPinchPoint_End = FVector2D(Touches[1]);
+			FVector2D CurrentPinchPoint_Start(TouchStates[0].ViewportRelativeLocation);
+			FVector2D CurrentPinchPoint_End(TouchStates[1].ViewportRelativeLocation);
 			const float PinchThreshold = 10000.0f;
 
 			if (CurrentAlpha == nullptr)
 			{
 				//if the pinch points are close enough to those from the last pinch, continue with the previous starting distance
 				if (AnchorDistanceSq == 0 ||
-					(LastPinchPoint_Start - CurrentPinchPoint_Start).SizeSquared() > PinchThreshold || 
+					(LastPinchPoint_Start - CurrentPinchPoint_Start).SizeSquared() > PinchThreshold ||
 					(LastPinchPoint_End - CurrentPinchPoint_End).SizeSquared() > PinchThreshold)
 				{
 					// remember the starting distance
@@ -68,7 +68,7 @@ void FGestureRecognizer::DetectGestures(const FVector (&Touches)[EKeys::NUM_TOUC
 			LastPinchPoint_End = CurrentPinchPoint_End;
 
 			// calculate the angle of the vector between the touch points
-			float NewAngle = FMath::Atan2(Touches[0].Y - Touches[1].Y, Touches[0].X - Touches[1].X);
+			float NewAngle = FMath::Atan2(CurrentPinchPoint_Start.Y - CurrentPinchPoint_End.Y, CurrentPinchPoint_Start.X - CurrentPinchPoint_End.X);
 			NewAngle = FRotator::ClampAxis(FMath::RadiansToDegrees(NewAngle));
 
 			float* CurrentAngle = CurrentGestureValues.Find(EKeys::Gesture_Rotate);
@@ -105,7 +105,8 @@ void FGestureRecognizer::DetectGestures(const FVector (&Touches)[EKeys::NUM_TOUC
 		else if (PreviousTouchCount == 1 && TouchCount == 1)
 		{
 			// remember the position for when we let go
-			FlickCurrent = FVector2D(Touches[0].X, Touches[0].Y);
+			const FVector2f& Loc0 = TouchStates[0].ViewportRelativeLocation;
+			FlickCurrent = FVector2D(Loc0);
 			FlickTime += DeltaTime;
 		}
 		else if (PreviousTouchCount >= 1 && TouchCount == 0)
diff --git a/Engine/Source/Runtime/Engine/Private/PlayerController.cpp b/Engine/Source/Runtime/Engine/Private/PlayerController.cpp
--- a/Engine/Source/Runtime/Engine/Private/PlayerController.cpp
+++ b/Engine/Source/Runtime/Engine/Private/PlayerController.cpp
@@ -2493,33 +2493,39 @@ bool APlayerController::InputTouch(
 	const uint32 TouchpadIndex,
 	const uint64 Timestamp)
 {
-	return InputTouch(DeviceId, static_cast<ETouchIndex::Type>(Handle), Type, TouchLocation, Force, Timestamp);
+	return InputTouch(FTouchId(DeviceId, static_cast<ETouchIndex::Type>(Handle)), Type, TouchLocation, Force, Timestamp);
 }
 
-bool APlayerController::InputTouch(const FInputDeviceId DeviceId, const ETouchIndex::Type TouchId, const ETouchType::Type Type, const FVector2D& TouchLocation,
-                                   const float Force, const uint64 Timestamp)
+bool APlayerController::InputTouch(const FTouchId TouchId, const ETouchType::Type Type, const FVector2D& TouchLocation, const float Force,
+	const uint64 Timestamp)
 {
 	bool bResult = false;
 	if (PlayerInput)
 	{
-		bResult = PlayerInput->InputTouch(DeviceId, TouchId, Type, TouchLocation, Force, Timestamp);
+		bResult = PlayerInput->InputTouch(TouchId, Type, TouchLocation, Force, Timestamp);
 
 		if (bEnableTouchEvents || bEnableTouchOverEvents)
 		{
 			FHitResult HitResult;
 			const bool bHit = GetHitResultAtScreenPosition(TouchLocation, CurrentClickTraceChannel, true, HitResult);
 
-			ProcessTouchHitResult(DeviceId, TouchId, Type, TouchLocation, HitResult, bHit);
+			ProcessTouchHitResult(TouchId, Type, TouchLocation, HitResult, bHit);
 		}
 	}
 
 	return bResult;
 }
 
-void APlayerController::ProcessTouchHitResult(const FInputDeviceId DeviceId, uint32 Handle, ETouchType::Type Type, const FVector2D& TouchLocation, const FHitResult& HitResult, const bool bHit)
+void APlayerController::ProcessTouchHitResult(const FInputDeviceId DeviceId, uint32 Handle, ETouchType::Type Type, const FVector2D& TouchLocation,
+	const FHitResult& HitResult, const bool bHit)
 {
-	const ETouchIndex::Type FingerIndex = ETouchIndex::Type(Handle);
-	UPrimitiveComponent* PreviousComponent = CurrentTouchablePrimitives[Handle].Get();
+	ProcessTouchHitResult(FTouchId(DeviceId, static_cast<ETouchIndex::Type>(Handle)), Type, TouchLocation, HitResult, bHit);
+}
+
+void APlayerController::ProcessTouchHitResult(const FTouchId TouchId, const ETouchType::Type Type, const FVector2D& TouchLocation,
+	const FHitResult& HitResult, const bool bHit)
+{
+	UPrimitiveComponent* PreviousComponent = CurrentTouchablePrimitives[TouchId.GetIndex()].Get();
 	UPrimitiveComponent* CurrentComponent = (bHit ? HitResult.Component.Get() : NULL);
 
 	if (GetHUD())
@@ -2539,21 +2545,21 @@ void APlayerController::ProcessTouchHitResult(const FInputDeviceId DeviceId, uin
 		// Give it a begin touch
 		if (bEnableTouchEvents && (CurrentComponent != NULL))
 		{
-			CurrentComponent->DispatchOnInputTouchBegin(FingerIndex);
+			CurrentComponent->DispatchOnInputTouchBegin(TouchId, this, HitResult);
 		}
 
 		// Give a touch enter event
 		if (bEnableTouchOverEvents)
 		{
-			UPrimitiveComponent::DispatchTouchOverEvents(FingerIndex, PreviousComponent, CurrentComponent);
-			CurrentTouchablePrimitives[Handle] = CurrentComponent;
+			UPrimitiveComponent::DispatchTouchOverEvents(TouchId, this, HitResult, PreviousComponent, CurrentComponent);
+			CurrentTouchablePrimitives[TouchId.GetIndex()] = CurrentComponent;
 		}
 		break;
 	case ETouchType::Ended:
 		// Give it a touch exit
 		if (bEnableTouchEvents && (CurrentComponent != NULL))
 		{
-			CurrentComponent->DispatchOnInputTouchEnd(FingerIndex);
+			CurrentComponent->DispatchOnInputTouchEnd(TouchId, this, HitResult);
 		}
 
 		// Give it a end touch
@@ -2563,14 +2569,14 @@ void APlayerController::ProcessTouchHitResult(const FInputDeviceId DeviceId, uin
 			if ((PreviousComponent != CurrentComponent) && (PreviousComponent != NULL))
 			{
 				// First notify the old component that the touch left it to go to the current component
-				UPrimitiveComponent::DispatchTouchOverEvents(FingerIndex, PreviousComponent, CurrentComponent);
+				UPrimitiveComponent::DispatchTouchOverEvents(TouchId, this, HitResult, PreviousComponent, CurrentComponent);
 			}
 
 			// Now notify that the current component is being released and thus the touch is leaving it
 			PreviousComponent = CurrentComponent;
 			CurrentComponent = NULL;
-			UPrimitiveComponent::DispatchTouchOverEvents(FingerIndex, PreviousComponent, CurrentComponent);
-			CurrentTouchablePrimitives[Handle] = CurrentComponent;
+			UPrimitiveComponent::DispatchTouchOverEvents(TouchId, this, HitResult, PreviousComponent, CurrentComponent);
+			CurrentTouchablePrimitives[TouchId.GetIndex()] = CurrentComponent;
 		}
 		break;
 	default:
@@ -5478,7 +5484,8 @@ void APlayerController::TickPlayerInput(const float DeltaSeconds, const bool bGa
 				UPrimitiveComponent* PreviousComponent = CurrentTouchablePrimitives[TouchIndexInt].Get();
 				UPrimitiveComponent* CurrentComponent = (bHit ? HitResult.Component.Get() : NULL);
 
-				UPrimitiveComponent::DispatchTouchOverEvents(FingerIndex, PreviousComponent, CurrentComponent);
+				const FTouchId TouchId(PlayerInput->GetTouchStates()[FingerIndex].DeviceId, FingerIndex);
+				UPrimitiveComponent::DispatchTouchOverEvents(TouchId, this, HitResult, PreviousComponent, CurrentComponent);
 
 				CurrentTouchablePrimitives[TouchIndexInt] = CurrentComponent;
 			}
@@ -6113,13 +6120,16 @@ FVector APlayerController::GetInputVectorKeyState(const FKey Key) const
 
 void APlayerController::GetInputTouchState(ETouchIndex::Type FingerIndex, float& LocationX, float& LocationY, bool& bIsCurrentlyPressed) const
 {
+	using namespace UE;
+
 	if (PlayerInput)
 	{
 		if (FingerIndex < EKeys::NUM_TOUCH_KEYS)
 		{
-			LocationX = PlayerInput->Touches[FingerIndex].X;
-			LocationY = PlayerInput->Touches[FingerIndex].Y;
-			bIsCurrentlyPressed = PlayerInput->Touches[FingerIndex].Z != 0 ? true : false;
+			const FTouchState& State = PlayerInput->GetTouchStates()[FingerIndex];
+			LocationX = State.ViewportRelativeLocation.X;
+			LocationY = State.ViewportRelativeLocation.Y;
+			bIsCurrentlyPressed = State.Pressure != 0;
 		}
 		else
 		{
diff --git a/Engine/Source/Runtime/Engine/Private/Slate/SceneViewport.cpp b/Engine/Source/Runtime/Engine/Private/Slate/SceneViewport.cpp
--- a/Engine/Source/Runtime/Engine/Private/Slate/SceneViewport.cpp
+++ b/Engine/Source/Runtime/Engine/Private/Slate/SceneViewport.cpp
@@ -906,9 +906,9 @@ FReply FSceneViewport::OnTouchStarted( const FGeometry& MyGeometry, const FPoint
 
 		const FVector2D TouchPosition = CachedCursorPos;
 
-		const bool bEventConsumed = ViewportClient->InputTouch(this, TouchEvent.GetInputDeviceId(),
-			static_cast<ETouchIndex::Type>(TouchEvent.GetPointerIndex()), ETouchType::Began, TouchPosition, TouchEvent.GetTouchForce(),
-			TouchEvent.GetEventTimestamp());
+		const bool bEventConsumed = ViewportClient->InputTouch(this,
+			FTouchId(TouchEvent.GetInputDeviceId(), static_cast<ETouchIndex::Type>(TouchEvent.GetPointerIndex())), ETouchType::Began, TouchPosition,
+			TouchEvent.GetTouchForce(), TouchEvent.GetEventTimestamp());
 		if (bEventConsumed)
 		{
 			const bool bTemporaryCapture = ViewportClient->GetMouseCaptureMode() == EMouseCaptureMode::CaptureDuringMouseDown;
@@ -941,9 +941,9 @@ FReply FSceneViewport::OnTouchMoved( const FGeometry& MyGeometry, const FPointer
 		// Switch to the viewport clients world before processing input
 		FScopedConditionalWorldSwitcher WorldSwitcher( ViewportClient );
 
-		const bool bEventConsumed = ViewportClient->InputTouch(this, TouchEvent.GetInputDeviceId(),
-			static_cast<ETouchIndex::Type>(TouchEvent.GetPointerIndex()), ETouchType::Moved, CachedCursorPos, TouchEvent.GetTouchForce(),
-			TouchEvent.GetEventTimestamp());
+		const bool bEventConsumed = ViewportClient->InputTouch(this,
+			FTouchId(TouchEvent.GetInputDeviceId(), static_cast<ETouchIndex::Type>(TouchEvent.GetPointerIndex())), ETouchType::Moved, CachedCursorPos,
+			TouchEvent.GetTouchForce(), TouchEvent.GetEventTimestamp());
 		if (!bEventConsumed)
 		{
 			CurrentReplyState = FReply::Unhandled(); 
@@ -979,9 +979,9 @@ FReply FSceneViewport::OnTouchEnded( const FGeometry& MyGeometry, const FPointer
 		// Switch to the viewport clients world before processing input
 		FScopedConditionalWorldSwitcher WorldSwitcher( ViewportClient );
 
-		const bool bEventConsumed = ViewportClient->InputTouch(this, TouchEvent.GetInputDeviceId(),
-			static_cast<ETouchIndex::Type>(TouchEvent.GetPointerIndex()), ETouchType::Ended, CurCursorPos, /* Force */ 0.0f,
-			TouchEvent.GetEventTimestamp());
+		const bool bEventConsumed = ViewportClient->InputTouch(this,
+			FTouchId(TouchEvent.GetInputDeviceId(), static_cast<ETouchIndex::Type>(TouchEvent.GetPointerIndex())), ETouchType::Ended, CurCursorPos,
+			/* Force */ 0.0f, TouchEvent.GetEventTimestamp());
 		if (!bEventConsumed)
 		{
 			CurrentReplyState = FReply::Unhandled(); 
@@ -1011,9 +1011,9 @@ FReply FSceneViewport::OnTouchForceChanged(const FGeometry& MyGeometry, const FP
 
 		const FVector2D TouchPosition = MyGeometry.AbsoluteToLocal(TouchEvent.GetScreenSpacePosition()) * MyGeometry.Scale;
 
-		const bool bEventConsumed = ViewportClient->InputTouch(this, TouchEvent.GetInputDeviceId(),
-			static_cast<ETouchIndex::Type>(TouchEvent.GetPointerIndex()), ETouchType::ForceChanged, TouchPosition, TouchEvent.GetTouchForce(), 
-			TouchEvent.GetEventTimestamp());
+		const bool bEventConsumed = ViewportClient->InputTouch(this,
+			FTouchId(TouchEvent.GetInputDeviceId(), static_cast<ETouchIndex::Type>(TouchEvent.GetPointerIndex())), ETouchType::ForceChanged,
+			TouchPosition, TouchEvent.GetTouchForce(), TouchEvent.GetEventTimestamp());
 		if (!bEventConsumed)
 		{
 			CurrentReplyState = FReply::Unhandled();
@@ -1038,9 +1038,9 @@ FReply FSceneViewport::OnTouchFirstMove(const FGeometry& MyGeometry, const FPoin
 
 		const FVector2D TouchPosition = MyGeometry.AbsoluteToLocal(TouchEvent.GetScreenSpacePosition()) * MyGeometry.Scale;
 
-		const bool bEventConsumed = ViewportClient->InputTouch(this, TouchEvent.GetInputDeviceId(),
-			static_cast<ETouchIndex::Type>(TouchEvent.GetPointerIndex()), ETouchType::FirstMove, TouchPosition, TouchEvent.GetTouchForce(),
-			TouchEvent.GetEventTimestamp());
+		const bool bEventConsumed = ViewportClient->InputTouch(this,
+			FTouchId(TouchEvent.GetInputDeviceId(), static_cast<ETouchIndex::Type>(TouchEvent.GetPointerIndex())), ETouchType::FirstMove,
+			TouchPosition, TouchEvent.GetTouchForce(), TouchEvent.GetEventTimestamp());
 		if (!bEventConsumed)
 		{
 			CurrentReplyState = FReply::Unhandled();
diff --git a/Engine/Source/Runtime/Engine/Private/UserInterface/PlayerInput.cpp b/Engine/Source/Runtime/Engine/Private/UserInterface/PlayerInput.cpp
--- a/Engine/Source/Runtime/Engine/Private/UserInterface/PlayerInput.cpp
+++ b/Engine/Source/Runtime/Engine/Private/UserInterface/PlayerInput.cpp
@@ -489,12 +489,14 @@ bool UPlayerInput::InputTouch(
 {
 	checkf(TouchpadIndex == 0, TEXT("We currently assume one touchpad in UPlayerInput::InputTouch. If this triggers, add support for multiple pads"));
 
-	return InputTouch(DeviceId, static_cast<ETouchIndex::Type>(Handle), Type, TouchLocation, Force, Timestamp);
+	return InputTouch(FTouchId(DeviceId, static_cast<ETouchIndex::Type>(Handle)), Type, TouchLocation, Force, Timestamp);
 }
 
-bool UPlayerInput::InputTouch(const FInputDeviceId DeviceId, const ETouchIndex::Type Handle, const ETouchType::Type Type, const FVector2D& TouchLocation,
-	const float Force, const uint64 Timestamp)
+bool UPlayerInput::InputTouch(const FTouchId TouchId, const ETouchType::Type Type, const FVector2D& TouchLocation, const float Force,
+	const uint64 Timestamp)
 {
+	using namespace UE;
+
 	// if touch disabled from the command line, consume the input and do nothing
 	static bool bTouchDisabled = FParse::Param(FCommandLine::Get(), TEXT("DisableTouch"));
 	if (bTouchDisabled)
@@ -503,28 +505,36 @@ bool UPlayerInput::InputTouch(const FInputDeviceId DeviceId, const ETouchIndex::
 	}
 
 	// get the current state of each finger
-	// if the handle is out of bounds, we can't handle it
-	if (Handle >= EKeys::NUM_TOUCH_KEYS)
+	// if the index is out of bounds, we can't handle it
+	const ETouchIndex::Type TouchIndex = TouchId.GetIndex();
+	if (TouchIndex >= EKeys::NUM_TOUCH_KEYS)
 	{
 		return false;
 	}
 
+PRAGMA_DISABLE_DEPRECATION_WARNINGS
 	// update position
-	Touches[Handle].X = TouchLocation.X;
-	Touches[Handle].Y = TouchLocation.Y;
+	Touches[TouchIndex].X = TouchLocation.X;
+	Touches[TouchIndex].Y = TouchLocation.Y;
 
 	// update touched/untouched flag
 	// > 0 means that it's currently held down (anything besides an end message is held down)
-	Touches[Handle].Z = (Type == ETouchType::Ended) ? 0.0f : Force;
+	Touches[TouchIndex].Z = (Type == ETouchType::Ended) ? 0.0f : Force;
+PRAGMA_ENABLE_DEPRECATION_WARNINGS
+
+	TouchStates[TouchIndex].DeviceId = TouchId.GetDeviceId();
+	TouchStates[TouchIndex].ViewportRelativeLocation = FVector2f(TouchLocation);
+	TouchStates[TouchIndex].Pressure = (Type == ETouchType::Ended) ? 0.0f : Force;
 
 	// hook up KeyState for InputComponent
-	FKeyState& KeyState = KeyStateMap.FindOrAdd(EKeys::TouchKeys[Handle]);
+	FKeyState& KeyState = KeyStateMap.FindOrAdd(EKeys::TouchKeys[TouchIndex]);
+	FTouchState& TouchState = TouchStates[TouchIndex];
 	switch(Type)
 	{
 	case ETouchType::Began:
 		KeyState.EventAccumulator[IE_Pressed].Add(++EventCount);
 		// store current touch location paired with event id
-		TouchEventLocations.Add(EventCount, Touches[Handle]);
+		TouchEventLocations.Add(EventCount, FVector(FVector2D(TouchState.ViewportRelativeLocation), TouchState.Pressure));
 		if (KeyState.bDownPrevious == false)
 		{
 			UWorld* World = GetWorld();
@@ -537,7 +547,7 @@ bool UPlayerInput::InputTouch(const FInputDeviceId DeviceId, const ETouchIndex::
 			{
 				KeyState.EventAccumulator[IE_DoubleClick].Add(++EventCount);
 				// store current touch location paired with event id
-				TouchEventLocations.Add(EventCount, Touches[Handle]);
+				TouchEventLocations.Add(EventCount, FVector(FVector2D(TouchState.ViewportRelativeLocation), TouchState.Pressure));
 			}
 
 			// just went down
@@ -548,13 +558,13 @@ bool UPlayerInput::InputTouch(const FInputDeviceId DeviceId, const ETouchIndex::
 	case ETouchType::Ended:
 		KeyState.EventAccumulator[IE_Released].Add(++EventCount);
 		// store current touch location paired with event id
-		TouchEventLocations.Add(EventCount, Touches[Handle]);
+		TouchEventLocations.Add(EventCount, FVector(FVector2D(TouchState.ViewportRelativeLocation), TouchState.Pressure));
 		break;
 
 	default:
 		KeyState.EventAccumulator[IE_Repeat].Add(++EventCount);
 		// store current touch location paired with event id
-		TouchEventLocations.Add(EventCount, Touches[Handle]);
+		TouchEventLocations.Add(EventCount, FVector(FVector2D(TouchState.ViewportRelativeLocation), TouchState.Pressure));
 		break;
 	}
 
@@ -643,7 +653,8 @@ bool UPlayerInput::InputGesture(
 
 void UPlayerInput::UpdatePinchStartDistance()
 {
-	GestureRecognizer.SetAnchorDistanceSquared(FVector2D(Touches[ETouchIndex::Touch1]), FVector2D(Touches[ETouchIndex::Touch2]));
+	GestureRecognizer.SetAnchorDistanceSquared(FVector2D(TouchStates[ETouchIndex::Touch1].ViewportRelativeLocation),
+		FVector2D(TouchStates[ETouchIndex::Touch2].ViewportRelativeLocation));
 }
 
 bool UPlayerInput::GetAxisProperties(const FKey AxisKey, FInputAxisProperties& OutAxisProperties)
@@ -2347,7 +2358,7 @@ void UPlayerInput::Tick(float DeltaTime)
 	if (GetDefault<UInputSettings>()->bEnableGestureRecognizer)
 	{
 		SCOPE_CYCLE_COUNTER(STAT_PC_GestureRecognition);
-		GestureRecognizer.DetectGestures(Touches, this, DeltaTime);
+		GestureRecognizer.DetectGestures(TouchStates, this, DeltaTime);
 	}
 }
 
diff --git a/Engine/Source/Runtime/Engine/Public/GestureRecognizer.h b/Engine/Source/Runtime/Engine/Public/GestureRecognizer.h
--- a/Engine/Source/Runtime/Engine/Public/GestureRecognizer.h
+++ b/Engine/Source/Runtime/Engine/Public/GestureRecognizer.h
@@ -9,6 +9,12 @@
 #include "CoreMinimal.h"
 #include "InputCoreTypes.h"
 
+class UPlayerInput;
+namespace UE
+{
+	struct FTouchState;
+}
+
 class FGestureRecognizer
 {
 public:
@@ -18,14 +24,14 @@ public:
 	{};
 
 	/** Attempt to detect touch gestures */
-	void DetectGestures(const FVector (&Touches)[EKeys::NUM_TOUCH_KEYS], class UPlayerInput* PlayerInput, float DeltaTime);
+	void DetectGestures(TArrayView<const UE::FTouchState> TouchStates, UPlayerInput* PlayerInput, float DeltaTime);
 
 	/** Save the distance between the anchor points */
 	void SetAnchorDistanceSquared(const FVector2D FirstPoint, const FVector2D SecontPoint);
 
 protected:
 	/** Internal processing of gestures */
-	void HandleGesture(class UPlayerInput* PlayerInput, FKey Gesture, bool bStarted, bool bEnded);
+	void HandleGesture(UPlayerInput* PlayerInput, FKey Gesture, bool bStarted, bool bEnded);
 	
 	/** A mapping of a gesture to it's current value (how far swiped, pinch amount, etc) */
 	TMap<FKey, float> CurrentGestureValues;
diff --git a/Engine/Source/Runtime/Engine/Public/ViewportClient.h b/Engine/Source/Runtime/Engine/Public/ViewportClient.h
--- a/Engine/Source/Runtime/Engine/Public/ViewportClient.h
+++ b/Engine/Source/Runtime/Engine/Public/ViewportClient.h
@@ -78,7 +78,7 @@ public:
 	 * @param	Timestamp - The timestamp of when this input event was originally polled in the form of CPU cycles.
 	 * @return	True to consume the key event, false to pass it on.
 	 */
-	UE_DEPRECATED(5.8, "Use the overload that omits TouchpadIndex instead")
+	UE_DEPRECATED(5.8, "Use the overload that uses FTouchId instead")
 	virtual bool InputTouch(
 		FViewport* const Viewport,
 		const FInputDeviceId DeviceId,
@@ -89,7 +89,7 @@ public:
 		const uint32 TouchpadIndex,
 		const uint64 Timestamp) final
 	{
-		return InputTouch(Viewport, DeviceId, static_cast<ETouchIndex::Type>(Handle), Type, TouchLocation, Force, Timestamp);
+		return InputTouch(Viewport, FTouchId(DeviceId, static_cast<ETouchIndex::Type>(Handle)), Type, TouchLocation, Force, Timestamp);
 	}
 
 	/**
@@ -97,8 +97,7 @@ public:
 	 * If the viewport client uses the event, it should return true.  If true is returned, the viewport should not bubble the event.
 	 * 
 	 * @param	Viewport - The viewport that was touched.
-	 * @param	DeviceId - The input device that was touched.
-	 * @param	TouchId - Identifier unique to this ongoing touch (i.e. finger).
+	 * @param	TouchId - Identifier unique to this ongoing touch (i.e. device + finger).
 	 * @param	Type - What kind of touch event this is (see ETouchType).
 	 * @param	TouchLocation - Viewport-relative position of the touch.
 	 * @param	Force - How hard the touch is.
@@ -107,8 +106,7 @@ public:
 	 */
 	virtual bool InputTouch(
 		FViewport* const Viewport,
-		const FInputDeviceId DeviceId,
-		const ETouchIndex::Type TouchId,
+		const FTouchId TouchId,
 		const ETouchType::Type Type,
 		const FVector2D& TouchLocation,
 		const float Force,
diff --git a/Engine/Source/Runtime/InputCore/Classes/InputCoreTypes.h b/Engine/Source/Runtime/InputCore/Classes/InputCoreTypes.h
--- a/Engine/Source/Runtime/InputCore/Classes/InputCoreTypes.h
+++ b/Engine/Source/Runtime/InputCore/Classes/InputCoreTypes.h
@@ -275,6 +275,62 @@ namespace ETouchIndex
 	};
 }
 
+/**
+ * Identifies a specific touch contact on a specific input device.
+ * Combines device identity (FInputDeviceId) with finger index (ETouchIndex) so that multiple touchscreens can each have independent finger-0, etc.
+ */
+USTRUCT(BlueprintType, MinimalAPI)
+struct FTouchId
+{
+	GENERATED_BODY()
+
+	FTouchId() = default;
+
+	explicit FTouchId(const FInputDeviceId InDeviceId, const ETouchIndex::Type InTouchIndex)
+		: DeviceId(InDeviceId)
+		, TouchIndex(InTouchIndex)
+	{
+	}
+
+	FInputDeviceId GetDeviceId() const
+	{
+		return DeviceId;
+	}
+
+	ETouchIndex::Type GetIndex() const
+	{
+		return TouchIndex;
+	}
+
+	bool operator==(const FTouchId& Other) const = default;
+
+	bool operator!=(const FTouchId& Other) const = default;
+
+	bool IsValid() const
+	{
+		return DeviceId.IsValid();
+	}
+
+	friend uint32 GetTypeHash(const FTouchId& Id)
+	{
+		return HashCombine(GetTypeHash(Id.DeviceId), GetTypeHash(static_cast<int32>(Id.TouchIndex)));
+	}
+
+private:
+
+	/**
+	 * The device being touched.
+	 */
+	UPROPERTY(BlueprintReadOnly, Category = "Input|Touch", meta = (AllowPrivateAccess = "true"))
+	FInputDeviceId DeviceId = INPUTDEVICEID_NONE;
+
+	/**
+	 * The device-specific index that differentiates this touch from other simultaneous touch inputs on the same device.
+	 */
+	UPROPERTY(BlueprintReadOnly, Category = "Input|Touch", meta = (AllowPrivateAccess = "true"))
+	TEnumAsByte<ETouchIndex::Type> TouchIndex = ETouchIndex::Touch1;
+};
+
 UENUM()
 namespace EConsoleForGamepadLabels
 {
diff --git a/Templates/TP_Puzzle/Source/TP_Puzzle/TP_PuzzleBlock.cpp b/Templates/TP_Puzzle/Source/TP_Puzzle/TP_PuzzleBlock.cpp
--- a/Templates/TP_Puzzle/Source/TP_Puzzle/TP_PuzzleBlock.cpp
+++ b/Templates/TP_Puzzle/Source/TP_Puzzle/TP_PuzzleBlock.cpp
@@ -4,7 +4,9 @@
 #include "TP_PuzzleBlockGrid.h"
 #include "UObject/ConstructorHelpers.h"
 #include "Components/StaticMeshComponent.h"
+#include "Components/PrimitiveComponent.h"
 #include "Engine/StaticMesh.h"
+#include "Materials/Material.h"
 #include "Materials/MaterialInstance.h"
 
 ATP_PuzzleBlock::ATP_PuzzleBlock()
@@ -38,7 +40,7 @@ ATP_PuzzleBlock::ATP_PuzzleBlock()
 	BlockMesh->SetMaterial(0, ConstructorStatics.BlueMaterial.Get());
 	BlockMesh->SetupAttachment(DummyRoot);
 	BlockMesh->OnClicked.AddDynamic(this, &ATP_PuzzleBlock::BlockClicked);
-	BlockMesh->OnInputTouchBegin.AddDynamic(this, &ATP_PuzzleBlock::OnFingerPressedBlock);
+	BlockMesh->OnTouchBegan.AddDynamic(this, &ATP_PuzzleBlock::OnFingerPressedBlock);
 
 	// Save a pointer to the orange material
 	BaseMaterial = ConstructorStatics.BaseMaterial.Get();
@@ -52,7 +54,7 @@ void ATP_PuzzleBlock::BlockClicked(UPrimitiveComponent* ClickedComp, FKey Button
 }
 
 
-void ATP_PuzzleBlock::OnFingerPressedBlock(ETouchIndex::Type FingerIndex, UPrimitiveComponent* TouchedComponent)
+void ATP_PuzzleBlock::OnFingerPressedBlock(FTouchId, APlayerController*, const FHitResult&, UPrimitiveComponent*)
 {
 	HandleClicked();
 }
diff --git a/Templates/TP_Puzzle/Source/TP_Puzzle/TP_PuzzleBlock.h b/Templates/TP_Puzzle/Source/TP_Puzzle/TP_PuzzleBlock.h
--- a/Templates/TP_Puzzle/Source/TP_Puzzle/TP_PuzzleBlock.h
+++ b/Templates/TP_Puzzle/Source/TP_Puzzle/TP_PuzzleBlock.h
@@ -4,8 +4,12 @@
 
 #include "CoreMinimal.h"
 #include "GameFramework/Actor.h"
+#include "InputCoreTypes.h"
 #include "TP_PuzzleBlock.generated.h"
 
+class APlayerController;
+struct FHitResult;
+
 /** A block that can be clicked */
 UCLASS(minimalapi)
 class ATP_PuzzleBlock : public AActor
@@ -48,7 +52,7 @@ public:
 
 	/** Handle the block being touched  */
 	UFUNCTION()
-	void OnFingerPressedBlock(ETouchIndex::Type FingerIndex, UPrimitiveComponent* TouchedComponent);
+	void OnFingerPressedBlock(FTouchId TouchId, APlayerController* Controller, const FHitResult& HitResult, UPrimitiveComponent* TouchedComponent);
 
 	void HandleClicked();
 
