Finish backout.

#Synced-CL -1

diff --git a/Engine/Source/Runtime/Engine/Private/DataReplication.cpp b/Engine/Source/Runtime/Engine/Private/DataReplication.cpp
--- a/Engine/Source/Runtime/Engine/Private/DataReplication.cpp
+++ b/Engine/Source/Runtime/Engine/Private/DataReplication.cpp
@@ -1373,7 +1373,7 @@ bool FObjectReplicator::ReceivedRPC(FNetBitReader& Reader, TSet<FNetworkGUID>& U
 
 			if (UNLIKELY(Result == ERPCNotifyResult::BlockRPC))
 			{
-				UE_LOGF(LogRepTraffic, Log, "      Blocked RPC: %ls", *FunctionName.ToString());
+				UE_LOG(LogRepTraffic, Log, TEXT("      Blocked RPC: %s"), *FunctionName.ToString());
 
 				return true;
 			}
