Resolve remaining backout conflicts.

#Synced-CL -1

diff --git a/Engine/Source/Runtime/Engine/Private/Analytics/RPCDoSDetectionAnalytics.cpp b/Engine/Source/Runtime/Engine/Private/Analytics/RPCDoSDetectionAnalytics.cpp
--- a/Engine/Source/Runtime/Engine/Private/Analytics/RPCDoSDetectionAnalytics.cpp
+++ b/Engine/Source/Runtime/Engine/Private/Analytics/RPCDoSDetectionAnalytics.cpp
@@ -181,64 +181,64 @@ void FRPCDoSAnalyticsData::SendAnalytics()
 			}
 		}
 
-		UE_LOGF(LogNet, Log, "RPCDosDetection Analytics:");
+		UE_LOG(LogNet, Log, TEXT("RPCDosDetection Analytics:"));
 
-		UE_LOGF(LogNet, Log, " - MaxSeverityIndex: %i", MaxSeverityIndex);
-		UE_LOGF(LogNet, Log, " - MaxSeverityCategory: %ls", *MaxSeverityCategory);
-		UE_LOGF(LogNet, Log, " - MaxAnalyticsSeverityIndex: %i", MaxAnalyticsSeverityIndex);
-		UE_LOGF(LogNet, Log, " - MaxAnalyticsSeverityCategory: %ls", *MaxAnalyticsSeverityCategory);
-		UE_LOGF(LogNet, Log, " - RPCs: %i", RPCTrackingAnalytics.Num());
+		UE_LOG(LogNet, Log, TEXT(" - MaxSeverityIndex: %i"), MaxSeverityIndex);
+		UE_LOG(LogNet, Log, TEXT(" - MaxSeverityCategory: %s"), *MaxSeverityCategory);
+		UE_LOG(LogNet, Log, TEXT(" - MaxAnalyticsSeverityIndex: %i"), MaxAnalyticsSeverityIndex);
+		UE_LOG(LogNet, Log, TEXT(" - MaxAnalyticsSeverityCategory: %s"), *MaxAnalyticsSeverityCategory);
+		UE_LOG(LogNet, Log, TEXT(" - RPCs: %i"), RPCTrackingAnalytics.Num());
 
 		// NOTE: Game thread CPU must be in analytics, even if GTrackGameThreadCPUUsage == 0, as it's used for filtering.
 		for (int32 RPCIdx=0; RPCIdx<RPCTrackingAnalytics.Num(); RPCIdx++)
 		{
 			const FRPCAnalytics& CurRPC = *RPCTrackingAnalytics[RPCIdx].Get();
 
-			UE_LOGF(LogNet, Log, "  - RPC[%i]:", RPCIdx);
-			UE_LOGF(LogNet, Log, "   - Name: %ls", *CurRPC.RPCName.ToString());
-			UE_LOGF(LogNet, Log, "   - MaxCountPerSec: %i", CurRPC.MaxCountPerSec);
-			UE_LOGF(LogNet, Log, "   - MaxTimePerSec: %f", CurRPC.MaxTimePerSec);
-			UE_LOGF(LogNet, Log, "   - MaxTimeGameThreadCPU: %i", CurRPC.MaxTimeGameThreadCPU);
-			UE_LOGF(LogNet, Log, "   - MaxSinglePacketRPCTime: %f", CurRPC.MaxSinglePacketRPCTime);
-			UE_LOGF(LogNet, Log, "   - SinglePacketRPCCount: %i", CurRPC.SinglePacketRPCCount);
-			UE_LOGF(LogNet, Log, "   - SinglePacketGameThreadCPU: %i", CurRPC.SinglePacketGameThreadCPU);
-			UE_LOGF(LogNet, Log, "   - BlockedCount: %i", CurRPC.BlockedCount);
-			UE_LOGF(LogNet, Log, "   - PlayerIP: %ls", *CurRPC.PlayerIP);
-			UE_LOGF(LogNet, Log, "   - PlayerUID: %ls", *CurRPC.PlayerUID);
+			UE_LOG(LogNet, Log, TEXT("  - RPC[%i]:"), RPCIdx);
+			UE_LOG(LogNet, Log, TEXT("   - Name: %s"), *CurRPC.RPCName.ToString());
+			UE_LOG(LogNet, Log, TEXT("   - MaxCountPerSec: %i"), CurRPC.MaxCountPerSec);
+			UE_LOG(LogNet, Log, TEXT("   - MaxTimePerSec: %f"), CurRPC.MaxTimePerSec);
+			UE_LOG(LogNet, Log, TEXT("   - MaxTimeGameThreadCPU: %i"), CurRPC.MaxTimeGameThreadCPU);
+			UE_LOG(LogNet, Log, TEXT("   - MaxSinglePacketRPCTime: %f"), CurRPC.MaxSinglePacketRPCTime);
+			UE_LOG(LogNet, Log, TEXT("   - SinglePacketRPCCount: %i"), CurRPC.SinglePacketRPCCount);
+			UE_LOG(LogNet, Log, TEXT("   - SinglePacketGameThreadCPU: %i"), CurRPC.SinglePacketGameThreadCPU);
+			UE_LOG(LogNet, Log, TEXT("   - BlockedCount: %i"), CurRPC.BlockedCount);
+			UE_LOG(LogNet, Log, TEXT("   - PlayerIP: %s"), *CurRPC.PlayerIP);
+			UE_LOG(LogNet, Log, TEXT("   - PlayerUID: %s"), *CurRPC.PlayerUID);
 		}
 
-		UE_LOGF(LogNet, Log, " - FilteredRPCs: %i (bOverrideRPCThresholds: %i)", FilteredRPCs.Num(), (int32)bOverrideRPCThresholds);
+		UE_LOG(LogNet, Log, TEXT(" - FilteredRPCs: %i (bOverrideRPCThresholds: %i)"), FilteredRPCs.Num(), (int32)bOverrideRPCThresholds);
 
 		for (int32 RPCIdx=0; RPCIdx<FilteredRPCs.Num(); RPCIdx++)
 		{
 			const FRPCAnalytics& CurRPC = *FilteredRPCs[RPCIdx].Get();
 
-			UE_LOGF(LogNet, Log, "  - FilteredRPC[%i]:", RPCIdx);
-			UE_LOGF(LogNet, Log, "   - Name: (*) %ls", *CurRPC.RPCName.ToString());
-			UE_LOGF(LogNet, Log, "   - MaxCountPerSec: %i", CurRPC.MaxCountPerSec);
-			UE_LOGF(LogNet, Log, "   - MaxTimePerSec: %f", CurRPC.MaxTimePerSec);
-			UE_LOGF(LogNet, Log, "   - MaxTimeGameThreadCPU: %i", CurRPC.MaxTimeGameThreadCPU);
-			UE_LOGF(LogNet, Log, "   - MaxSinglePacketRPCTime: %f", CurRPC.MaxSinglePacketRPCTime);
-			UE_LOGF(LogNet, Log, "   - SinglePacketRPCCount: %i", CurRPC.SinglePacketRPCCount);
-			UE_LOGF(LogNet, Log, "   - SinglePacketGameThreadCPU: %i", CurRPC.SinglePacketGameThreadCPU);
-			UE_LOGF(LogNet, Log, "   - BlockedCount: %i", CurRPC.BlockedCount);
-			UE_LOGF(LogNet, Log, "   - PlayerIP: %ls", *CurRPC.PlayerIP);
-			UE_LOGF(LogNet, Log, "   - PlayerUID: %ls", *CurRPC.PlayerUID);
+			UE_LOG(LogNet, Log, TEXT("  - FilteredRPC[%i]:"), RPCIdx);
+			UE_LOG(LogNet, Log, TEXT("   - Name: (*) %s"), *CurRPC.RPCName.ToString());
+			UE_LOG(LogNet, Log, TEXT("   - MaxCountPerSec: %i"), CurRPC.MaxCountPerSec);
+			UE_LOG(LogNet, Log, TEXT("   - MaxTimePerSec: %f"), CurRPC.MaxTimePerSec);
+			UE_LOG(LogNet, Log, TEXT("   - MaxTimeGameThreadCPU: %i"), CurRPC.MaxTimeGameThreadCPU);
+			UE_LOG(LogNet, Log, TEXT("   - MaxSinglePacketRPCTime: %f"), CurRPC.MaxSinglePacketRPCTime);
+			UE_LOG(LogNet, Log, TEXT("   - SinglePacketRPCCount: %i"), CurRPC.SinglePacketRPCCount);
+			UE_LOG(LogNet, Log, TEXT("   - SinglePacketGameThreadCPU: %i"), CurRPC.SinglePacketGameThreadCPU);
+			UE_LOG(LogNet, Log, TEXT("   - BlockedCount: %i"), CurRPC.BlockedCount);
+			UE_LOG(LogNet, Log, TEXT("   - PlayerIP: %s"), *CurRPC.PlayerIP);
+			UE_LOG(LogNet, Log, TEXT("   - PlayerUID: %s"), *CurRPC.PlayerUID);
 		}
 
-		UE_LOGF(LogNet, Log, " - MaxPlayerSeverity: %i", MaxPlayerSeverity.Num());
+		UE_LOG(LogNet, Log, TEXT(" - MaxPlayerSeverity: %i"), MaxPlayerSeverity.Num());
 
 		for (int32 SevIdx=0; SevIdx<MaxPlayerSeverity.Num(); SevIdx++)
 		{
 			const FMaxRPCDoSEscalation& CurSev = MaxPlayerSeverity[SevIdx];
 
-			UE_LOGF(LogNet, Log, "  - MaxPlayerSeverity[%i]:", SevIdx);
-			UE_LOGF(LogNet, Log, "   - PlayerIP: %ls", *CurSev.PlayerIP);
-			UE_LOGF(LogNet, Log, "   - PlayerUID: %ls", *CurSev.PlayerUID);
-			UE_LOGF(LogNet, Log, "   - MaxSeverityIndex: %i", CurSev.MaxSeverityIndex);
-			UE_LOGF(LogNet, Log, "   - MaxSeverityCategory: %ls", *CurSev.MaxSeverityCategory);
-			UE_LOGF(LogNet, Log, "   - MaxAnalyticsSeverityIndex: %i", CurSev.MaxAnalyticsSeverityIndex);
-			UE_LOGF(LogNet, Log, "   - MaxAnalyticsSeverityCategory: %ls", *CurSev.MaxAnalyticsSeverityCategory);
+			UE_LOG(LogNet, Log, TEXT("  - MaxPlayerSeverity[%i]:"), SevIdx);
+			UE_LOG(LogNet, Log, TEXT("   - PlayerIP: %s"), *CurSev.PlayerIP);
+			UE_LOG(LogNet, Log, TEXT("   - PlayerUID: %s"), *CurSev.PlayerUID);
+			UE_LOG(LogNet, Log, TEXT("   - MaxSeverityIndex: %i"), CurSev.MaxSeverityIndex);
+			UE_LOG(LogNet, Log, TEXT("   - MaxSeverityCategory: %s"), *CurSev.MaxSeverityCategory);
+			UE_LOG(LogNet, Log, TEXT("   - MaxAnalyticsSeverityIndex: %i"), CurSev.MaxAnalyticsSeverityIndex);
+			UE_LOG(LogNet, Log, TEXT("   - MaxAnalyticsSeverityCategory: %s"), *CurSev.MaxAnalyticsSeverityCategory);
 		}
 
 
diff --git a/Engine/Source/Runtime/Engine/Private/Net/RPCDoSDetection.cpp b/Engine/Source/Runtime/Engine/Private/Net/RPCDoSDetection.cpp
--- a/Engine/Source/Runtime/Engine/Private/Net/RPCDoSDetection.cpp
+++ b/Engine/Source/Runtime/Engine/Private/Net/RPCDoSDetection.cpp
@@ -139,7 +139,7 @@ bool FRPCDoSState::HasHitQuota_Count(const FRPCDoSCounters(&PerPeriodHistory)[16
 	bool bReturnVal = EscalateQuotaRPCsPerFrame > 0 && InFrameCounter.RPCCounter >= EscalateQuotaRPCsPerFrame;
 
 #if RPC_QUOTA_DEBUG
-	UE_CLOGF(bReturnVal, LogNet, Log, "HasHitQuota_Count: Hit Frame Quota: RPCCounter: %i, Limit: %i",
+	UE_CLOG(bReturnVal, LogNet, Log, TEXT("HasHitQuota_Count: Hit Frame Quota: RPCCounter: %i, Limit: %i"),
 			InFrameCounter.RPCCounter, EscalateQuotaRPCsPerFrame);
 #endif
 
@@ -150,7 +150,7 @@ bool FRPCDoSState::HasHitQuota_Count(const FRPCDoSCounters(&PerPeriodHistory)[16
 		bReturnVal = PeriodCounter.RPCCounter + InFrameCounter.RPCCounter >= EscalateQuotaRPCsPerPeriod;
 
 #if RPC_QUOTA_DEBUG
-		UE_CLOGF(bReturnVal, LogNet, Log, "HasHitQuota_Count: Hit Period Quota: RPCsPerPeriodCounter: %i, Limit: %i",
+		UE_CLOG(bReturnVal, LogNet, Log, TEXT("HasHitQuota_Count: Hit Period Quota: RPCsPerPeriodCounter: %i, Limit: %i"),
 				(PeriodCounter.RPCCounter + InFrameCounter.RPCCounter), EscalateQuotaRPCsPerPeriod);
 #endif
 	}
@@ -165,8 +165,8 @@ bool FRPCDoSState::HasHitQuota_Time(const FRPCDoSCounters(&PerPeriodHistory)[16]
 #if RPC_QUOTA_DEBUG
 	bool bDebugReturnVal = EscalateTimeQuotaMSPerFrame > 0 && InFrameCounter.DebugAccumRPCTime >= EscalateTimeQuotaSecsPerFrame;
 
-	UE_CLOGF(bReturnVal, LogNet, Log,
-			"HasHitQuota_Time: Hit Frame Quota: AccumRPCTime: %f, DebugAccumRPCTime: %f, Limit: %f (%i ms)",
+	UE_CLOG(bReturnVal, LogNet, Log,
+			TEXT("HasHitQuota_Time: Hit Frame Quota: AccumRPCTime: %f, DebugAccumRPCTime: %f, Limit: %f (%i ms)"),
 			InFrameCounter.AccumRPCTime, InFrameCounter.DebugAccumRPCTime, EscalateTimeQuotaSecsPerFrame,
 			EscalateTimeQuotaMSPerFrame);
 
@@ -186,14 +186,15 @@ bool FRPCDoSState::HasHitQuota_Time(const FRPCDoSCounters(&PerPeriodHistory)[16]
 		bDebugReturnVal = (PeriodCounter.DebugAccumRPCTime + InFrameCounter.DebugAccumRPCTime) >=
 							EscalateTimeQuotaSecsPerPeriod;
 
-		UE_CLOGF(bReturnVal, LogNet, Log,
-				"HasHitQuota_Time: Hit Period Quota: AccumPerPeriodRPCTime: %f, DebugAccumPerPeriodRPCTime: %f, Limit: %f (%i ms)",
+		UE_CLOG(bReturnVal, LogNet, Log,
+				TEXT("HasHitQuota_Time: Hit Period Quota: AccumPerPeriodRPCTime: %f, DebugAccumPerPeriodRPCTime: %f, Limit: %f (%i ms)"),
 				(PeriodCounter.AccumRPCTime + InFrameCounter.AccumRPCTime),
 				(PeriodCounter.DebugAccumRPCTime + InFrameCounter.DebugAccumRPCTime),
 				EscalateTimeQuotaSecsPerPeriod, EscalateTimeQuotaMSPerPeriod);
 
-		UE_CLOGF(bReturnVal ^ bDebugReturnVal, LogNet, Warning,
-				"HasHitQuota_Time: Approximate vs Debug Period Quota mismatch: " "AccumPerPeriodRPCTime: %f, DebugAccumPerPeriodRPCTime: %f, Limit: %f (%i ms)",
+		UE_CLOG(bReturnVal ^ bDebugReturnVal, LogNet, Warning,
+				TEXT("HasHitQuota_Time: Approximate vs Debug Period Quota mismatch: ")
+				TEXT("AccumPerPeriodRPCTime: %f, DebugAccumPerPeriodRPCTime: %f, Limit: %f (%i ms)"),
 				(PeriodCounter.AccumRPCTime + InFrameCounter.AccumRPCTime),
 				(PeriodCounter.DebugAccumRPCTime + InFrameCounter.DebugAccumRPCTime),
 				EscalateTimeQuotaSecsPerPeriod, EscalateTimeQuotaMSPerPeriod);
@@ -253,12 +254,12 @@ bool FRPCDoSStateConfig::LoadStructConfig(const TCHAR* SectionName, const TCHAR*
 					}
 					else
 					{
-						UE_LOGF(LogNet, Error, "FRPCDoSStateConfig: Ini property '%ls' is not a config property.", *Var);
+						UE_LOG(LogNet, Error, TEXT("FRPCDoSStateConfig: Ini property '%s' is not a config property."), *Var);
 					}
 				}
 				else
 				{
-					UE_LOGF(LogNet, Error, "FRPCDoSStateConfig: Ini property '%ls' not found.", *Var);
+					UE_LOG(LogNet, Error, TEXT("FRPCDoSStateConfig: Ini property '%s' not found."), *Var);
 				}
 			}
 		}
@@ -277,7 +278,7 @@ void FRPCDoSStateConfig::ValidateConfig()
 
 			if (Value < BaseValue || Value > 16)
 			{
-				UE_LOGF(LogNet, Warning, "FRPCDoSStateConfig: %ls '%i' must be between %i and 16. Clamping.",
+				UE_LOG(LogNet, Warning, TEXT("FRPCDoSStateConfig: %s '%i' must be between %i and 16. Clamping."),
 						PropertyName, Value, BaseValue);
 
 				Value = FMath::Clamp((int32)Value, BaseValue, 16);
@@ -292,7 +293,7 @@ void FRPCDoSStateConfig::ValidateConfig()
 
 	if (AutoEscalateTime > 0 && AutoEscalateTime < CooloffTime)
 	{
-		UE_LOGF(LogNet, Warning, "FRPCDoSStateConfig: AutoEscalateTime must be larger than CooloffTime.");
+		UE_LOG(LogNet, Warning, TEXT("FRPCDoSStateConfig: AutoEscalateTime must be larger than CooloffTime."));
 	}
 }
 
@@ -455,12 +456,12 @@ void FRPCDoSDetection::InitConfig(FName NetDriverName)
 					}
 					else
 					{
-						UE_LOGF(LogNet, Warning, "RPC DoS Detection failed to load ini section: %ls", *CurSection);
+						UE_LOG(LogNet, Warning, TEXT("RPC DoS Detection failed to load ini section: %s"), *CurSection);
 					}
 				}
 				else
 				{
-					UE_LOGF(LogNet, Warning, "RPC DoS detection could not find ini section: %ls", *CurSection);
+					UE_LOG(LogNet, Warning, TEXT("RPC DoS detection could not find ini section: %s"), *CurSection);
 				}
 			}
 		}
@@ -473,7 +474,7 @@ void FRPCDoSDetection::InitConfig(FName NetDriverName)
 		}
 		else
 		{
-			UE_LOGF(LogNet, Warning, "RPC DoS detection enabled, but no DetectionSeverity states specified! Disabling.");
+			UE_LOG(LogNet, Warning, TEXT("RPC DoS detection enabled, but no DetectionSeverity states specified! Disabling."));
 
 			bRPCDoSDetection = false;
 		}
@@ -619,7 +620,7 @@ void FRPCDoSDetection::UpdateSeverity_Private(ERPCDoSSeverityUpdate Update, ERPC
 
 		if (bLogEscalate)
 		{
-			UE_LOGF(LogNet, Warning, "Updated RPC DoS detection severity for '%ls' from '%ls' to '%ls' (Reason: %ls)",
+			UE_LOG(LogNet, Warning, TEXT("Updated RPC DoS detection severity for '%s' from '%s' to '%s' (Reason: %s)"),
 					*GetPlayerAddress(), *OldState.SeverityCategory, *NewState.SeverityCategory, LexToString(Reason));
 		}
 
@@ -712,7 +713,7 @@ void FRPCDoSDetection::UpdateSeverity_Private(ERPCDoSSeverityUpdate Update, ERPC
 
 					if (bLogEscalate)
 					{
-						UE_LOGF(LogNet, Warning, "Updated analytics RPC DoS detection severity from '%ls' to '%ls'.",
+						UE_LOG(LogNet, Warning, TEXT("Updated analytics RPC DoS detection severity from '%s' to '%s'."),
 								*DetectionSeverity[OldWorstState].SeverityCategory, *NewState.SeverityCategory);
 
 						// Log RPC's which could not be timed individually (as they won't display in analytics)
@@ -730,7 +731,7 @@ void FRPCDoSDetection::UpdateSeverity_Private(ERPCDoSSeverityUpdate Update, ERPC
 								RPCGroupStr.Append(*CurRPC.ToString());
 							}
 
-							UE_LOGF(LogNet, Log, "RPC DoS escalation Untimed RPC's: Time: %f, RPC's: %ls",
+							UE_LOG(LogNet, Log, TEXT("RPC DoS escalation Untimed RPC's: Time: %f, RPC's: %s"),
 									RPCGroupTime, ToCStr(RPCGroupStr.ToString()))
 						}
 					}
@@ -742,7 +743,7 @@ void FRPCDoSDetection::UpdateSeverity_Private(ERPCDoSSeverityUpdate Update, ERPC
 
 			if (bKickPlayer && KickPlayerFunc)
 			{
-				UE_LOGF(LogNet, Warning, "RPC DoS detection kicking player IP '%ls', UID '%ls'.", *GetPlayerAddress(),
+				UE_LOG(LogNet, Warning, TEXT("RPC DoS detection kicking player IP '%s', UID '%s'."), *GetPlayerAddress(),
 						*GetPlayerUID());
 
 				KickPlayerFunc();
