Try to figure out missing SceneViewport during UGameEngine::OnGameWindowClosed

#Synced-CL 52463319

diff --git a/Engine/Source/Runtime/Engine/Private/GameEngine.cpp b/Engine/Source/Runtime/Engine/Private/GameEngine.cpp
--- a/Engine/Source/Runtime/Engine/Private/GameEngine.cpp
+++ b/Engine/Source/Runtime/Engine/Private/GameEngine.cpp
@@ -783,6 +783,16 @@ void UGameEngine::SwitchGameWindowToUseGameViewport()
 
 void UGameEngine::OnGameWindowClosed(const TSharedRef<SWindow>& WindowBeingClosed)
 {
+#if !UE_BUILD_SHIPPING
+#if PLATFORM_CPU_X86_FAMILY
+	if (FPlatformMisc::GetHostArchitecture() != TEXT("arm64"))
+#endif
+	{
+		// TEMP: Try to figure out missing SceneViewport during UGameEngine::OnGameWindowClosed
+		UE_LOG(LogEngine, Log, TEXT("UGameEngine::OnGameWindowClosed call stack:"));
+		FDebug::DumpStackTraceToLog(ELogVerbosity::Log);
+	}
+#endif
 	FSlateApplication::Get().UnregisterGameViewport();
 	if (ensure(SceneViewport.IsValid()))
 	{
diff --git a/Engine/Source/Runtime/Engine/Private/Slate/SceneViewport.cpp b/Engine/Source/Runtime/Engine/Private/Slate/SceneViewport.cpp
--- a/Engine/Source/Runtime/Engine/Private/Slate/SceneViewport.cpp
+++ b/Engine/Source/Runtime/Engine/Private/Slate/SceneViewport.cpp
@@ -79,6 +79,16 @@ FSceneViewport::FSceneViewport( FViewportClient* InViewportClient, TSharedPtr<SV
 
 FSceneViewport::~FSceneViewport()
 {
+#if !UE_BUILD_SHIPPING
+#if PLATFORM_CPU_X86_FAMILY
+	if (FPlatformMisc::GetHostArchitecture() != TEXT("arm64"))
+#endif
+	{
+		// TEMP: Try to figure out missing SceneViewport during UGameEngine::OnGameWindowClosed
+		UE_LOG(LogViewport, Log, TEXT("FSceneViewport::~FSceneViewport call stack:"));
+		FDebug::DumpStackTraceToLog(ELogVerbosity::Log);
+	}
+#endif
 	Destroy();
 	// Wait for resources to be deleted
 	FlushRenderingCommands();
