Add blueprint-exposed function to reset sync group information on Anim Graph

#Synced-CL 52530536

diff --git a/Engine/Source/Runtime/Engine/Classes/Animation/AnimInstance.h b/Engine/Source/Runtime/Engine/Classes/Animation/AnimInstance.h
--- a/Engine/Source/Runtime/Engine/Classes/Animation/AnimInstance.h
+++ b/Engine/Source/Runtime/Engine/Classes/Animation/AnimInstance.h
@@ -1301,6 +1301,9 @@ public:
 	ENGINE_API void UnlockAIResources(bool bUnlockMovement, bool UnlockAILogic);
 	//--- AI communication end ---//
 
+	UFUNCTION(BlueprintCallable, Category = "Animation|Synchronization")
+	ENGINE_API void ResetSync();
+
 	UFUNCTION(BlueprintCallable, Category = "Animation|Synchronization", meta=(BlueprintThreadSafe))
 	ENGINE_API bool GetTimeToClosestMarker(FName SyncGroup, FName MarkerName, float& OutMarkerTime) const;
 
diff --git a/Engine/Source/Runtime/Engine/Private/Animation/AnimInstance.cpp b/Engine/Source/Runtime/Engine/Private/Animation/AnimInstance.cpp
--- a/Engine/Source/Runtime/Engine/Private/Animation/AnimInstance.cpp
+++ b/Engine/Source/Runtime/Engine/Private/Animation/AnimInstance.cpp
@@ -4593,6 +4593,11 @@ void UAnimInstance::UnlockAIResources(bool bUnlockMovement, bool UnlockAILogic)
 	UE_LOG(LogAnimation, Error, TEXT("%s: UnlockAIResources is no longer supported. Please use UnlockAIResourcesWithAnimation instead."), *GetName());
 }
 
+void UAnimInstance::ResetSync()
+{
+	GetProxyOnAnyThread<FAnimInstanceProxy>().ResetSync();
+}
+
 bool UAnimInstance::GetTimeToClosestMarker(FName SyncGroup, FName MarkerName, float& OutMarkerTime) const
 {
 	return GetProxyOnAnyThread<FAnimInstanceProxy>().GetTimeToClosestMarker(SyncGroup, MarkerName, OutMarkerTime);
