[MIR] Added displacement tesselation support.

#Synced-CL 52505914

diff --git a/Engine/Source/Runtime/Engine/Private/Materials/MaterialIRModuleBuilder.cpp b/Engine/Source/Runtime/Engine/Private/Materials/MaterialIRModuleBuilder.cpp
--- a/Engine/Source/Runtime/Engine/Private/Materials/MaterialIRModuleBuilder.cpp
+++ b/Engine/Source/Runtime/Engine/Private/Materials/MaterialIRModuleBuilder.cpp
@@ -1276,7 +1276,7 @@ struct FMaterialIRModuleBuilderImpl
 		CompilationOutput.bUsesOpacityMask          = Module->IsMaterialPropertyUsed(MP_OpacityMask);
 		
 		// Check if displacement is actually used (disable if it was left at the default, invalid constant value of -1)
-		CompilationOutput.bUsesDisplacement 	 |= DoesPlatformSupportNanite(Module->GetShaderPlatform()) && Builder->MaterialInterface->IsTessellationEnabled() && MaterialPropertyHasNonZeroValue(MP_Displacement);
+		CompilationOutput.bUsesDisplacement 	 |= DoesPlatformSupportNanite(Module->GetShaderPlatform()) && Builder->MaterialInterface->IsTessellationEnabled() && Module->IsMaterialPropertyUsed(MP_Displacement);
 		CompilationOutput.bModifiesMeshPosition  |= CompilationOutput.bUsesPixelDepthOffset || CompilationOutput.bUsesWorldPositionOffset || CompilationOutput.bUsesDisplacement;
 		
 		// Run a final material compilation output validation and report any error reported.
diff --git a/Engine/Source/Runtime/Engine/Private/Materials/MaterialIRToHLSLTranslator.cpp b/Engine/Source/Runtime/Engine/Private/Materials/MaterialIRToHLSLTranslator.cpp
--- a/Engine/Source/Runtime/Engine/Private/Materials/MaterialIRToHLSLTranslator.cpp
+++ b/Engine/Source/Runtime/Engine/Private/Materials/MaterialIRToHLSLTranslator.cpp
@@ -1831,7 +1831,7 @@ struct FEntryPointCodeGen
 		// Assign a local to the call result
 		InstrToLocalIndex.Add(Call, NumLocals++);
 	}
-
+	
 	void LowerCallOutput(const MIR::FCallParameterOutput* CallOutput)
 	{
 		const MIR::FCall* Call = CallOutput->Call->As<MIR::FCall>();
@@ -1848,7 +1848,7 @@ struct FEntryPointCodeGen
 
 		Printer << TEXT('_') << ParamIndex;
 	}
-	
+
 	// Returns the name of the material uniform buffer for the current stage.
 	const TCHAR* GetMaterialUBName() const
 	{
@@ -2594,7 +2594,8 @@ static void GetShaderCompilerEnvironment(const FMaterial* Material, const FMater
 	OutEnvironment.SetDefine(TEXT("USES_TRANSFORM_VECTOR"), false);
 	OutEnvironment.SetDefine(TEXT("WANT_PIXEL_DEPTH_OFFSET"), CompilationOutput.bUsesPixelDepthOffset);
 	OutEnvironment.SetDefineAndCompileArgument(TEXT("USES_WORLD_POSITION_OFFSET"), (bool)CompilationOutput.bUsesWorldPositionOffset);
-	OutEnvironment.SetDefineAndCompileArgument(TEXT("USES_DISPLACEMENT"), false);
+	OutEnvironment.SetDefineAndCompileArgument(TEXT("USES_DISPLACEMENT"), (bool)CompilationOutput.bUsesDisplacement);
+	OutEnvironment.SetDefineAndCompileArgument(TEXT("SHADER_SUPPORTS_MATERIAL_VS_UB"), !CompilationOutput.bUsesDisplacement);
 	OutEnvironment.SetDefine(TEXT("USES_EMISSIVE_COLOR"), false);
 	OutEnvironment.SetDefine(TEXT("USES_DISTORTION"), Material->IsDistorted());
 	OutEnvironment.SetDefine(TEXT("MATERIAL_ENABLE_TRANSLUCENCY_FOGGING"), Material->ShouldApplyFogging());
