[SimpleConstructionScript/ALT]
- Make some data used by both GT and ALT thread-safe

#Synced-CL 52534307

diff --git a/Engine/Source/Runtime/Engine/Classes/Engine/SimpleConstructionScript.h b/Engine/Source/Runtime/Engine/Classes/Engine/SimpleConstructionScript.h
--- a/Engine/Source/Runtime/Engine/Classes/Engine/SimpleConstructionScript.h
+++ b/Engine/Source/Runtime/Engine/Classes/Engine/SimpleConstructionScript.h
@@ -6,6 +6,7 @@
 #include "CoreMinimal.h"
 #include "UObject/ObjectMacros.h"
 #include "UObject/Object.h"
+#include "Async/TransactionallySafeRecursiveMutex.h"
 #include "Misc/Guid.h"
 #include "GameFramework/Actor.h"
 #include "SimpleConstructionScript.generated.h"
@@ -120,6 +121,8 @@ class USimpleConstructionScript : public UObject
 	FStaticMeshComponentBulkReregisterContext* GetReregisterContext() const { return ReregisterContext; }
 
 private:
+	mutable UE::FTransactionallySafeRecursiveMutex Mutex;
+
 	/** Root nodes of the construction script */
 	UPROPERTY()
 	TArray<TObjectPtr<class USCS_Node>> RootNodes;
diff --git a/Engine/Source/Runtime/Engine/Private/SimpleConstructionScript.cpp b/Engine/Source/Runtime/Engine/Private/SimpleConstructionScript.cpp
--- a/Engine/Source/Runtime/Engine/Private/SimpleConstructionScript.cpp
+++ b/Engine/Source/Runtime/Engine/Private/SimpleConstructionScript.cpp
@@ -149,10 +149,33 @@ void USimpleConstructionScript::Serialize(FArchive& Ar)
 #endif // WITH_EDITORONLY_DATA
 }
 
+template <typename T>
+class TSimpleConstructionScriptMutexScope
+{
+private:
+	T& Mutex;
+public:
+	TSimpleConstructionScriptMutexScope(T& InMutex)
+		: Mutex(InMutex)
+	{
+		if (!Mutex.TryLock())
+		{
+			TRACE_CPUPROFILER_EVENT_SCOPE(SimpleConstructionScriptMutexContention);
+			Mutex.Lock();
+		}
+	}
+
+	~TSimpleConstructionScriptMutexScope()
+	{
+		Mutex.Unlock();
+	}
+};
+
 void USimpleConstructionScript::PreloadChain()
 {
 	ConditionalPreload();
 
+	TSimpleConstructionScriptMutexScope ScopeLock(Mutex);
 	for (USCS_Node* Node : RootNodes)
 	{
 		Node->PreloadChain();
@@ -181,7 +204,13 @@ void USimpleConstructionScript::PostLoad()
 		UE_LOG(LogBlueprint, Warning, TEXT("USimpleConstructionScript::PostLoad() '%s' cannot find its owner blueprint"), *GetPathName());
 		return;
 	}
+#endif // WITH_EDITOR
 
+	// Protect AllNodes/RootNodes/DefaultSceneRootNode against concurrent reads from ALT
+	// (via FindArchetype -> FindSCSNode/FindSCSNodeByGuid/PreloadChain).
+	TSimpleConstructionScriptMutexScope ScopeLock(Mutex);
+
+#if WITH_EDITOR
 	// This pass is not needed during reinstancing.
 	if (!GIsDuplicatingClassForReinstancing)
 	{
@@ -1020,6 +1049,8 @@ USCS_Node* USimpleConstructionScript::FindParentNode(USCS_Node* InNode) const
 
 USCS_Node* USimpleConstructionScript::FindSCSNode(const FName InName) const
 {
+	TSimpleConstructionScriptMutexScope ScopeLock(Mutex);
+
 	if (NameToSCSNodeMap.Num() > 0)
 	{
 		return NameToSCSNodeMap.FindRef(InName);
@@ -1037,6 +1068,8 @@ USCS_Node* USimpleConstructionScript::FindSCSNode(const FName InName) const
 
 USCS_Node* USimpleConstructionScript::FindSCSNodeByGuid(const FGuid Guid) const
 {
+	TSimpleConstructionScriptMutexScope ScopeLock(Mutex);
+
 	for (USCS_Node* SCSNode : GetAllNodes())
 	{
 		if (SCSNode && (SCSNode->VariableGuid == Guid))
@@ -1160,6 +1193,8 @@ USceneComponent* USimpleConstructionScript::GetSceneRootComponentTemplate(bool b
 void USimpleConstructionScript::ValidateSceneRootNodes()
 {
 #if WITH_EDITOR
+	TSimpleConstructionScriptMutexScope ScopeLock(Mutex);
+
 	UBlueprint* Blueprint = GetBlueprint();
 
 	if(DefaultSceneRootNode == nullptr)
