[MIR] Fixed bug when emitting custom output name during IR to HLSL generation.

#Synced-CL 52505914

diff --git a/Engine/Source/Runtime/Engine/Private/Materials/MaterialIRToHLSLTranslator.cpp b/Engine/Source/Runtime/Engine/Private/Materials/MaterialIRToHLSLTranslator.cpp
--- a/Engine/Source/Runtime/Engine/Private/Materials/MaterialIRToHLSLTranslator.cpp
+++ b/Engine/Source/Runtime/Engine/Private/Materials/MaterialIRToHLSLTranslator.cpp
@@ -1056,9 +1056,13 @@ struct FEntryPointCodeGen
 		{
 			Printer << TEXTVIEW("Parameters.CustomizedUVs[") << (int)(SetMaterialOutput->Property - MP_CustomizedUVs0) << TEXT(']');
 		}
+		else if (SetMaterialOutput->Property == MP_CustomOutput)
+		{
+			Printer << TEXTVIEW("Parameters.") << *SetMaterialOutput->Name;
+		}
 		else
 		{
-			Printer << (CurrentStage != MIR::Stage_Vertex && SetMaterialOutput->Property != MP_CustomOutput ? TEXTVIEW("PixelMaterialInputs.") : TEXTVIEW("Parameters."));
+			Printer << (CurrentStage != MIR::Stage_Vertex ? TEXTVIEW("PixelMaterialInputs.") : TEXTVIEW("Parameters."));
 			PrintMaterialPropertyName(Printer, SetMaterialOutput->Property);
 		}
 
