[NetPhysics] Register and unregister PhysicsReplicationLOD focal particle on pawn possession change.

#Synced-CL 52463959

diff --git a/Engine/Source/Runtime/Engine/Private/PhysicsEngine/NetworkPhysicsSettingsComponent.cpp b/Engine/Source/Runtime/Engine/Private/PhysicsEngine/NetworkPhysicsSettingsComponent.cpp
--- a/Engine/Source/Runtime/Engine/Private/PhysicsEngine/NetworkPhysicsSettingsComponent.cpp
+++ b/Engine/Source/Runtime/Engine/Private/PhysicsEngine/NetworkPhysicsSettingsComponent.cpp
@@ -16,6 +16,7 @@
 #include "PhysicsProxy/SingleParticlePhysicsProxy.h"
 #include "PhysicsProxy/ClusterUnionPhysicsProxy.h"
 #include "PhysicsProxy/GeometryCollectionPhysicsProxy.h"
+#include "GameFramework/Pawn.h"
 
 #include UE_INLINE_GENERATED_CPP_BY_NAME(NetworkPhysicsSettingsComponent)
 
@@ -125,6 +126,13 @@ void UNetworkPhysicsSettingsComponent::OnUnregister()
 		}
 	}
 
+	// Unbind from pawn controller change delegate and unregister focal particle
+	if (APawn* OwnerPawn = Cast<APawn>(GetOwner()))
+	{
+		OwnerPawn->ReceiveControllerChangedDelegate.RemoveDynamic(this, &ThisClass::OnControllerChanged);
+	}
+	UnregisterFromPhysicsReplicationLOD();
+
 	if (SettingsDataAsset)
 	{
 		SettingsDataAsset->MarkUninitialized();
@@ -167,6 +175,12 @@ void UNetworkPhysicsSettingsComponent::BeginPlay()
 		}
 	}
 	RegisterInPhysicsReplicationLOD();
+
+	// Bind to pawn controller change delegate to register/unregister focal particle on possession changes
+	if (APawn* OwnerPawn = Cast<APawn>(GetOwner()))
+	{
+		OwnerPawn->ReceiveControllerChangedDelegate.AddUniqueDynamic(this, &ThisClass::OnControllerChanged);
+	}
 }
 
 void UNetworkPhysicsSettingsComponent::OnComponentPhysicsStateChanged(UPrimitiveComponent* ChangedComponent, EComponentPhysicsStateChange StateChange)
@@ -207,6 +221,37 @@ void UNetworkPhysicsSettingsComponent::RegisterInPhysicsReplicationLOD()
 	}
 }
 
+void UNetworkPhysicsSettingsComponent::UnregisterFromPhysicsReplicationLOD()
+{
+	if (AActor* Owner = GetOwner())
+	{
+		Owner->UnregisterAsFocalPointInPhysicsReplicationLOD();
+	}
+}
+
+void UNetworkPhysicsSettingsComponent::OnControllerChanged(APawn* InPawn, AController* OldController, AController* NewController)
+{
+	if (!GetSettings().GeneralSettings.bFocalParticleInPhysicsReplicationLOD)
+	{
+		return;
+	}
+
+	AActor* Owner = GetOwner();
+	if (!Owner)
+	{
+		return;
+	}
+
+	if (NewController != nullptr && Owner->GetLocalRole() == ENetRole::ROLE_AutonomousProxy)
+	{
+		Owner->RegisterAsFocalPointInPhysicsReplicationLOD();
+	}
+	else
+	{
+		UnregisterFromPhysicsReplicationLOD();
+	}
+}
+
 UNetworkPhysicsSettingsComponent* UNetworkPhysicsSettingsComponent::GetSettingsForActor(AActor* Owner)
 {
 	UNetworkPhysicsSettingsComponent** Value = UNetworkPhysicsSettingsComponent::ObjectToSettings_External.Find(Owner);
diff --git a/Engine/Source/Runtime/Engine/Public/Physics/NetworkPhysicsSettingsComponent.h b/Engine/Source/Runtime/Engine/Public/Physics/NetworkPhysicsSettingsComponent.h
--- a/Engine/Source/Runtime/Engine/Public/Physics/NetworkPhysicsSettingsComponent.h
+++ b/Engine/Source/Runtime/Engine/Public/Physics/NetworkPhysicsSettingsComponent.h
@@ -813,6 +813,13 @@ private:
 	/** Perform logic to register this actor into physics replication LOD, depending on settings */
 	void RegisterInPhysicsReplicationLOD();
 
+	/** Unregister this actor from physics replication LOD */
+	void UnregisterFromPhysicsReplicationLOD();
+
+	/** Callback for pawn controller changes to register/unregister as focal particle */
+	UFUNCTION()
+	void OnControllerChanged(APawn* InPawn, AController* OldController, AController* NewController);
+
 public:
 	/** Get the settings, on the game thread */
 	const FNetworkPhysicsSettingsData& GetSettings() const { return SettingsDataAsset ? SettingsDataAsset->GetSettings_External() : SettingsNetworkPhysicsData_Default; }
