NewMaterialTranslator:

- Fixes a bug with multiple custom HLSL outputs.
- The second issue was Custom HLSL with multiple output nodes.  The logic to calculate the local variable index was incorrect when there were more than 1 custom HLSL output.  This resulted in the function output always getting the value of the return value of the custom HLSL function which was incorrect.
- To fix this we lookup the function call in InstrToLocalIndex which is a map from function call to how many locals (input or output) it has.  We subtract the number of output parameters (GetNumOutputParameters()), and then offset by the current output index to get the current output index.

#rnx
#Synced-CL 52508799

#Synced-CL 52508799

diff --git a/Engine/Source/Runtime/Engine/Private/Materials/MaterialIRToHLSLTranslator.cpp b/Engine/Source/Runtime/Engine/Private/Materials/MaterialIRToHLSLTranslator.cpp
--- a/Engine/Source/Runtime/Engine/Private/Materials/MaterialIRToHLSLTranslator.cpp
+++ b/Engine/Source/Runtime/Engine/Private/Materials/MaterialIRToHLSLTranslator.cpp
@@ -1839,16 +1839,9 @@ struct FEntryPointCodeGen
 	void LowerCallOutput(const MIR::FCallParameterOutput* CallOutput)
 	{
 		const MIR::FCall* Call = CallOutput->Call->As<MIR::FCall>();
-		int32 ParamIndex = InstrToLocalIndex[Call];
-
-		// Void functions don't use a local
-		if (!Call->Function->ReturnType.IsVoid())
-		{
-			ParamIndex -= 1;
-		}
-
+		const int32 NumOutputParams = Call->Function->GetNumOutputParameters();
 		// Compute the index of the local that stores the additional output parameter 
-		ParamIndex = ParamIndex - Call->Function->NumInputAndOutputParams + Call->Function->NumInputOnlyParams + CallOutput->Index;
+		const int32 ParamIndex = InstrToLocalIndex[Call] - NumOutputParams + CallOutput->Index;
 
 		Printer << TEXT('_') << ParamIndex;
 	}
