[Navigation] Fix crash in AddNavigationSystemToWorld during undo/redo on transient worlds

During Sequencer undo/redo, AWorldSettings::PostTransacted can fire on a
WorldSettings belonging to a transient world without a PersistentLevel.
Guard against this in InternalPostPropertyChanged and pass NavigationSystemConfig
directly to avoid redundant GetWorldSettings call.

#Synced-CL 52469485

diff --git a/Engine/Source/Runtime/Engine/Private/WorldSettings.cpp b/Engine/Source/Runtime/Engine/Private/WorldSettings.cpp
--- a/Engine/Source/Runtime/Engine/Private/WorldSettings.cpp
+++ b/Engine/Source/Runtime/Engine/Private/WorldSettings.cpp
@@ -969,12 +969,12 @@ void AWorldSettings::InternalPostPropertyChanged(FName PropertyName)
 	else if (PropertyName == GET_MEMBER_NAME_CHECKED(AWorldSettings, NavigationSystemConfig))
 	{
 		UWorld* World = GetWorld();
-		if (World)
+		if (World && World->PersistentLevel)
 		{
 			World->SetNavigationSystem(nullptr);
 			if (NavigationSystemConfig)
 			{
-				FNavigationSystem::AddNavigationSystemToWorld(*World, FNavigationSystemRunMode::EditorMode);
+				FNavigationSystem::AddNavigationSystemToWorld(*World, FNavigationSystemRunMode::EditorMode, NavigationSystemConfig);
 			}
 		}
 	}
