[Materials/MTSAVE]
- Fix a few thread-safety hazard when saving materials in parallel

#Synced-CL 52534307

diff --git a/Engine/Source/Runtime/Engine/Private/Materials/MaterialCachedData.cpp b/Engine/Source/Runtime/Engine/Private/Materials/MaterialCachedData.cpp
--- a/Engine/Source/Runtime/Engine/Private/Materials/MaterialCachedData.cpp
+++ b/Engine/Source/Runtime/Engine/Private/Materials/MaterialCachedData.cpp
@@ -932,11 +932,9 @@ void FMaterialCachedExpressionData::Validate(const UMaterialInterface& Material)
 		}
 		FMaterialLayersFunctions::Validate(MaterialLayers, EditorOnlyData->MaterialLayers);
 
+		TArray<FShaderCompilerError> Errors;
 		if (!FPlatformProperties::RequiresCookedData() && AllowShaderCompiling())
 		{
-			// Mute log errors created by GetShaderSourceFilePath during include path validation
-			LOG_SCOPE_VERBOSITY_OVERRIDE(LogShaders, ELogVerbosity::Fatal);
-
 			for (auto PathIt = EditorOnlyData->ExpressionIncludeFilePaths.CreateIterator(); PathIt; ++PathIt)
 			{
 				const FString& IncludeFilePath = *PathIt;
@@ -944,7 +942,12 @@ void FMaterialCachedExpressionData::Validate(const UMaterialInterface& Material)
 
 				if (!IncludeFilePath.IsEmpty())
 				{
-					FString ValidatedPath = GetShaderSourceFilePath(IncludeFilePath);
+					// Mute log errors by using the variant that gives use the errors instead of logging them.
+					FString ValidatedPath = GetShaderSourceFilePath(IncludeFilePath, &Errors);
+
+					// Keep the array small during the loop since we don't care about errors.
+					Errors.Reset();
+
 					if (!ValidatedPath.IsEmpty())
 					{
 						ValidatedPath = IFileManager::Get().ConvertToAbsolutePathForExternalAppForWrite(*ValidatedPath);
diff --git a/Engine/Source/Runtime/Engine/Private/Materials/MaterialInstance.cpp b/Engine/Source/Runtime/Engine/Private/Materials/MaterialInstance.cpp
--- a/Engine/Source/Runtime/Engine/Private/Materials/MaterialInstance.cpp
+++ b/Engine/Source/Runtime/Engine/Private/Materials/MaterialInstance.cpp
@@ -1495,13 +1495,13 @@ void UMaterialInstance::ValidateTextureOverrides(ERHIFeatureLevel::Type InFeatur
 
 void UMaterialInstance::ValidateTextureOverrides(EShaderPlatform InShaderPlatform) const
 {
-	if (!(IsInGameThread() || IsAsyncLoading()))
+	if (!(IsInGameThread() || IsAsyncLoading() || UE::IsSavingPackage()))
 	{
 		// Fatal to call getmaterial in a non-game thread or async loading
 		return;
 	}
 
-	const UMaterial* Material = GetMaterial();
+	const UMaterial* Material = GetMaterial_Concurrent();
 	const FMaterialResource* CurrentResource = Material->GetMaterialResource(InShaderPlatform);
 	const bool bShouldValidateVTUsage = UseVirtualTexturing(GMaxRHIShaderPlatform);
 
@@ -5469,7 +5469,7 @@ int32 UMaterialInstance::CompilePropertyEx( class FMaterialCompiler* Compiler, c
 bool UMaterialInstance::TryGetUsedTexturesInCookCache(TArray<TObjectPtr<UObject>>& OutTextures)
 {
 	OutTextures.Empty();
-	const UMaterial* BaseMaterial = GetMaterial();
+	const UMaterial* BaseMaterial = GetMaterial_Concurrent();
 
 	if (!BaseMaterial || BaseMaterial->IsDefaultMaterial())
 	{
diff --git a/Engine/Source/Runtime/Engine/Private/Materials/MaterialInterface.cpp b/Engine/Source/Runtime/Engine/Private/Materials/MaterialInterface.cpp
--- a/Engine/Source/Runtime/Engine/Private/Materials/MaterialInterface.cpp
+++ b/Engine/Source/Runtime/Engine/Private/Materials/MaterialInterface.cpp
@@ -227,7 +227,7 @@ void UMaterialInterface::Serialize(FArchive& Ar)
 		bIncludedInBaseGame = Ar.GetCookContext()->GetCookingDLC() == UE::Cook::ECookingDLC::No;
 
 		// Make sure we are not cooking materials with the new translator flag enabled as it is still experimental.
-		UMaterial* Material = GetMaterial();
+		const UMaterial* Material = GetMaterial_Concurrent();
 		if (Material && Material->IsUsingNewHLSLGenerator())
 		{
 			UE_LOG(LogMaterial, Display, TEXT("Material '%s' is marked to use the new material translator. Please refrain from submitting materials using the new translator as it is still experimental."), *GetName());
