- Add VirtualTexture type support to the Material IR system to enable virtual texture sampling in the new HLSL translator
- Introduce a distinct VirtualTexture type kind separate from RuntimeVirtualTexture, with fallback to Texture2D when virtual texturing is disabled on the target platform
- Fix VT stack sharing logic to merge feedback flags across shared stacks and handle both VT and RVT uniform types during page table read analysis

#Synced-CL 52545647

diff --git a/Engine/Source/Runtime/Engine/Private/Materials/MaterialIREmitter.cpp b/Engine/Source/Runtime/Engine/Private/Materials/MaterialIREmitter.cpp
--- a/Engine/Source/Runtime/Engine/Private/Materials/MaterialIREmitter.cpp
+++ b/Engine/Source/Runtime/Engine/Private/Materials/MaterialIREmitter.cpp
@@ -1180,7 +1180,14 @@ FValueRef FEmitter::TextureUniform(FName Name, UTexture* Texture, EMaterialSampl
 		return Poison();
 	}
 
-	FTextureUniform Proto = MakePrototype<FTextureUniform>(FType::FromMaterialValueType(Texture->GetMaterialType()));
+	// When virtual texturing is disabled we need to fall back to texture 2D.
+	FType TextureType = FType::FromMaterialValueType(Texture->GetMaterialType());
+	if (TextureType.IsVirtualTexture() && !UseVirtualTexturing(GetShaderPlatform()))
+	{
+		TextureType = FType::MakeTexture2D();
+	}
+
+	FTextureUniform Proto = MakePrototype<FTextureUniform>(TextureType);
 	Proto.Name = FScriptName{ Name };
 	Proto.Texture = Texture;
 	Proto.SamplerType = SamplerType;
@@ -3617,7 +3624,7 @@ FValueRef FEmitter::Nop(FValueRef Arg)
 	return EmitPrototype(*this, Prototype);
 }
 
-static EMaterialSamplerType MapSamplerTypeForTexture(FValueRef TextureValue, EMaterialSamplerType InSamplerType)
+static EMaterialSamplerType MapSamplerTypeForTexture(FValueRef TextureValue, EMaterialSamplerType InSamplerType, EShaderPlatform ShaderPlatform)
 {
 	// Can't sample with virtual texturing if input texture is not a virtual texture
 	if (InSamplerType == SAMPLERTYPE_MAX)
@@ -3631,7 +3638,9 @@ static EMaterialSamplerType MapSamplerTypeForTexture(FValueRef TextureValue, EMa
 			InSamplerType = SAMPLERTYPE_VirtualMasks;
 		}
 	}
-	return As<FVirtualTextureUniform>(TextureValue) ? PromoteSamplerTypeFromSTtoVT(InSamplerType) : DemoteSamplerTypeFromVTtoST(InSamplerType);
+
+	const bool bIsVirtualTexture = TextureValue->Type.IsRuntimeVirtualTexture() || TextureValue->Type.IsVirtualTexture();
+	return (bIsVirtualTexture && UseVirtualTexturing(ShaderPlatform)) ? PromoteSamplerTypeFromSTtoVT(InSamplerType) : DemoteSamplerTypeFromVTtoST(InSamplerType);
 }
 
 static FValueRef VTPageTableLoadFromSamplerSource(FEmitter& Em, FValueRef Texture, const FTextureSampleAttributes& BaseAttributes, FValueRef TexCoord,
@@ -3656,6 +3665,18 @@ static FValueRef VTPageTableLoadFromSamplerSource(FEmitter& Em, FValueRef Textur
 	Prototype.bEnableFeedback = BaseAttributes.bEnableFeedback;
 	Prototype.bIsAdaptive = BaseAttributes.bIsAdaptive;
 
+	float AspectRatio = 1.0f;
+	if (UTexture* TextureObj = Cast<UTexture>(Texture->GetTextureObject()))
+	{
+		const int32 SizeX = TextureObj->Source.GetSizeX();
+		const int32 SizeY = TextureObj->Source.GetSizeY();
+		if (SizeX > 0 && SizeY > 0)
+		{
+			AspectRatio = (float)SizeX / (float)SizeY;
+		}
+	}
+	Prototype.AspectRatio = AspectRatio;
+
 	return EmitPrototype(Em, Prototype);
 }
 
@@ -3690,7 +3711,7 @@ FValueRef FEmitter::TextureGather(FValueRef Texture, FValueRef TexCoord, ETextur
 		return Poison();
 	}
 
-	const EMaterialSamplerType SamplerType = MapSamplerTypeForTexture(Texture, BaseAttributes.SamplerType);
+	const EMaterialSamplerType SamplerType = MapSamplerTypeForTexture(Texture, BaseAttributes.SamplerType, GetShaderPlatform());
 
 	FTextureRead Prototype = MakePrototype<FTextureRead>(FType::MakeFloatVector(4));
 	Prototype.TextureUniform = Texture;
@@ -3717,7 +3738,7 @@ FValueRef FEmitter::TextureSample(FValueRef Texture, FValueRef TexCoord, bool bA
 		return Poison();
 	}
 
-	const EMaterialSamplerType SamplerType = MapSamplerTypeForTexture(Texture, BaseAttributes.SamplerType);
+	const EMaterialSamplerType SamplerType = MapSamplerTypeForTexture(Texture, BaseAttributes.SamplerType, GetShaderPlatform());
 
 	FTextureRead PrototypePixel = MakePrototype<FTextureRead>(FType::MakeFloatVector(4));
 	PrototypePixel.TextureUniform = Texture;
@@ -3771,7 +3792,7 @@ FValueRef FEmitter::TextureSampleLevel(FValueRef Texture, FValueRef TexCoord, FV
 		return Poison();
 	}
 
-	const EMaterialSamplerType SamplerType = MapSamplerTypeForTexture(Texture, BaseAttributes.SamplerType);
+	const EMaterialSamplerType SamplerType = MapSamplerTypeForTexture(Texture, BaseAttributes.SamplerType, GetShaderPlatform());
 
 	FTextureRead Prototype = MakePrototype<FTextureRead>(FType::MakeFloatVector(4));
 	Prototype.TextureUniform = Texture;
@@ -3809,7 +3830,7 @@ FValueRef FEmitter::TextureSampleBias(FValueRef Texture, FValueRef TexCoord, FVa
 		MipBias = Operator(BO_Add, MipBias, Builtin(EBuiltin::TextureMipBias));
 	}
 
-	const EMaterialSamplerType SamplerType = MapSamplerTypeForTexture(Texture, BaseAttributes.SamplerType);
+	const EMaterialSamplerType SamplerType = MapSamplerTypeForTexture(Texture, BaseAttributes.SamplerType, GetShaderPlatform());
 
 	FTextureRead PrototypePixel = MakePrototype<FTextureRead>(FType::MakeFloatVector(4));
 	PrototypePixel.TextureUniform = Texture;
@@ -3850,7 +3871,7 @@ FValueRef FEmitter::TextureSampleGrad(FValueRef Texture, FValueRef TexCoord, FVa
 		return Poison();
 	}
 
-	const EMaterialSamplerType SamplerType = MapSamplerTypeForTexture(Texture, BaseAttributes.SamplerType);
+	const EMaterialSamplerType SamplerType = MapSamplerTypeForTexture(Texture, BaseAttributes.SamplerType, GetShaderPlatform());
 
 	FTextureRead Prototype = MakePrototype<FTextureRead>(FType::MakeFloatVector(4));
 	Prototype.TextureUniform = Texture;
diff --git a/Engine/Source/Runtime/Engine/Private/Materials/MaterialIRModuleBuilder.cpp b/Engine/Source/Runtime/Engine/Private/Materials/MaterialIRModuleBuilder.cpp
--- a/Engine/Source/Runtime/Engine/Private/Materials/MaterialIRModuleBuilder.cpp
+++ b/Engine/Source/Runtime/Engine/Private/Materials/MaterialIRModuleBuilder.cpp
@@ -1261,14 +1261,15 @@ struct FMaterialIRModuleBuilderImpl
 		// NumUsedUVScalars
 		CompilationOutput.NumUsedUVScalars = Module->Statistics.NumInterpolatedTexCoords * 2;
 
-		// NumVirtualTextureFeedbackRequests
+		// NumVirtualTextureFeedbackRequests / EstimatedNumVirtualTextureLookups
 		int32 NumVirtualTextureFeedbackRequests = 0;
 		for (const FMaterialIRValueAnalyzer::FVTStackEntry& VTStack : ValueAnalyzer.VTStacks)
 		{
 			NumVirtualTextureFeedbackRequests += VTStack.bGenerateFeedback ? 1 : 0;
 		}
-		Module->CompilationOutput.NumVirtualTextureFeedbackRequests = NumVirtualTextureFeedbackRequests;
-		
+		CompilationOutput.NumVirtualTextureFeedbackRequests = NumVirtualTextureFeedbackRequests;
+		CompilationOutput.EstimatedNumVirtualTextureLookups = ValueAnalyzer.VTStacks.Num();
+
 		// Simple on/off flags depending on whether a non-default value is set to some material property
 		CompilationOutput.bUsesWorldPositionOffset  = MaterialPropertyHasNonZeroValue(MP_WorldPositionOffset);
 		CompilationOutput.bUsesPixelDepthOffset     = MaterialPropertyHasNonZeroValue(MP_PixelDepthOffset);
diff --git a/Engine/Source/Runtime/Engine/Private/Materials/MaterialIRToHLSLTranslator.cpp b/Engine/Source/Runtime/Engine/Private/Materials/MaterialIRToHLSLTranslator.cpp
--- a/Engine/Source/Runtime/Engine/Private/Materials/MaterialIRToHLSLTranslator.cpp
+++ b/Engine/Source/Runtime/Engine/Private/Materials/MaterialIRToHLSLTranslator.cpp
@@ -262,6 +262,7 @@ static FPrinter& operator<<(FPrinter& Printer, const MIR::FType& Type)
 			case MIR::ETypeKind::Texture2DArray: Printer << TEXTVIEW("Texture2DArray"); break;
 			case MIR::ETypeKind::TextureCubeArray: Printer << TEXTVIEW("TextureCubeArray"); break;
 			case MIR::ETypeKind::Texture3D: Printer << TEXTVIEW("Texture3D"); break;
+			case MIR::ETypeKind::VirtualTexture: Printer << TEXTVIEW("VirtualTexturePhysical"); break;
 			case MIR::ETypeKind::RuntimeVirtualTexture: Printer << TEXTVIEW("VirtualTexturePhysical"); break;
 			case MIR::ETypeKind::SubstrateData: Printer << TEXTVIEW("FSubstrateData"); break;
 			case MIR::ETypeKind::VTPageTableResult: Printer << TEXTVIEW("VTPageTableResult"); break;
diff --git a/Engine/Source/Runtime/Engine/Private/Materials/MaterialIRTypes.cpp b/Engine/Source/Runtime/Engine/Private/Materials/MaterialIRTypes.cpp
--- a/Engine/Source/Runtime/Engine/Private/Materials/MaterialIRTypes.cpp
+++ b/Engine/Source/Runtime/Engine/Private/Materials/MaterialIRTypes.cpp
@@ -23,6 +23,7 @@ const TCHAR* TypeKindToString(ETypeKind Kind)
 		CASE(Texture2DArray)
 		CASE(TextureCubeArray)
 		CASE(Texture3D)
+		CASE(VirtualTexture)
 		CASE(RuntimeVirtualTexture)
 		CASE(ParameterCollection)
 		CASE(SubstrateData)
@@ -104,6 +105,7 @@ FType FType::FromMaterialValueType(EMaterialValueType Type)
 		case MCT_Float2: return FType::MakeFloatVector(2);
 		case MCT_Float3: return FType::MakeFloatVector(3);
 		case MCT_Float4: return FType::MakeFloatVector(4);
+		case MCT_TextureVirtual: return FType::MakeVirtualTexture();
 		case MCT_Texture2D: return FType::MakeTexture2D();
 		case MCT_TextureCube: return FType::MakeTextureCube();
 		case MCT_Texture2DArray: return FType::MakeTexture2DArray();
@@ -225,6 +227,13 @@ FType FType::MakeTexture3D()
 	return Type;
 }
 
+FType FType::MakeVirtualTexture()
+{
+	FType Type{};
+	Type.Kind = ETypeKind::VirtualTexture;
+	return Type;
+}
+
 FType FType::MakeRuntimeVirtualTexture()
 {
 	FType Type{};
diff --git a/Engine/Source/Runtime/Engine/Private/Materials/MaterialIRValueAnalyzer.cpp b/Engine/Source/Runtime/Engine/Private/Materials/MaterialIRValueAnalyzer.cpp
--- a/Engine/Source/Runtime/Engine/Private/Materials/MaterialIRValueAnalyzer.cpp
+++ b/Engine/Source/Runtime/Engine/Private/Materials/MaterialIRValueAnalyzer.cpp
@@ -84,19 +84,28 @@ static void Extern_AnalyzeInStage(FMaterialIRValueAnalyzer& Analyzer, MIR::FExte
 	Extern->AnalyzeInStage(Context, Stage);
 }
 
-static int32 AcquireVTStackIndex(FMaterialIRValueAnalyzer& Analyzer, MIR::FVTPageTableRead& VTPageTableRead, bool bGenerateFeedback)
+static int32 AcquireVTStackIndex(FMaterialIRValueAnalyzer& Analyzer, MIR::FVTPageTableRead& VTPageTableRead, bool bGenerateFeedback, int32 TextureUniformIndex, int32 PreallocatedStackTextureIndex = INDEX_NONE)
 {
-	// Try to find matching VT stack entry
+	// Try to find matching VT stack entry, ignoring bGenerateFeedback to share stacks
 	for (int32 VTStackIndex = 0; VTStackIndex < Analyzer.VTStacks.Num(); ++VTStackIndex)
 	{
-		const FMaterialIRValueAnalyzer::FVTStackEntry& VTStack = Analyzer.VTStacks[VTStackIndex];
-		if (VTStack.TexCoord == VTPageTableRead.TexCoord &&
-			VTStack.bGenerateFeedback == bGenerateFeedback &&
+		FMaterialIRValueAnalyzer::FVTStackEntry& VTStack = Analyzer.VTStacks[VTStackIndex];
+
+		// Don't reuse a full stack unless this texture is already a layer in it
+		const FMaterialVirtualTextureStack& VTStackData = Analyzer.CompilationOutput->UniformExpressionSet.GetVTStack(VTStackIndex);
+		const bool bStackIsFull = VTStackData.AreLayersFull() && VTStackData.FindLayer(TextureUniformIndex) == INDEX_NONE;
+
+		if (!bStackIsFull &&
+			VTStack.TexCoord == VTPageTableRead.TexCoord &&
 			VTStack.AddressU == VTPageTableRead.AddressU &&
 			VTStack.AddressV == VTPageTableRead.AddressV &&
 			VTStack.MipValue == VTPageTableRead.MipValue &&
-			VTStack.MipValueMode == VTPageTableRead.MipValueMode)
+			VTStack.MipValueMode == VTPageTableRead.MipValueMode &&
+			VTStack.bIsAdaptive == VTPageTableRead.bIsAdaptive &&
+			VTStack.AspectRatio == VTPageTableRead.AspectRatio &&
+			VTStack.PreallocatedStackTextureIndex == PreallocatedStackTextureIndex)
 		{
+			VTStack.bGenerateFeedback |= bGenerateFeedback;
 			return VTStackIndex;
 		}
 	}
@@ -109,41 +118,65 @@ static int32 AcquireVTStackIndex(FMaterialIRValueAnalyzer& Analyzer, MIR::FVTPag
 	VTStack.AddressV = VTPageTableRead.AddressV;
 	VTStack.MipValue = VTPageTableRead.MipValue;
 	VTStack.MipValueMode = VTPageTableRead.MipValueMode;
+	VTStack.bIsAdaptive = VTPageTableRead.bIsAdaptive;
+	VTStack.AspectRatio = VTPageTableRead.AspectRatio;
+	VTStack.PreallocatedStackTextureIndex = PreallocatedStackTextureIndex;
 
-	return Analyzer.CompilationOutput->UniformExpressionSet.AddVTStack(INDEX_NONE);
+	return Analyzer.CompilationOutput->UniformExpressionSet.AddVTStack(PreallocatedStackTextureIndex);
 }
 
 static void VTPageTableRead_AnalyzeInStage(FMaterialIRValueAnalyzer& Analyzer, MIR::FVTPageTableRead* VTPageTableRead, MIR::EStage Stage)
 {
-	const MIR::FVirtualTextureUniform* VirtualTextureUniform = MIR::As<MIR::FVirtualTextureUniform>(VTPageTableRead->VirtualTextureUniform);
-	check(VirtualTextureUniform);
+	// The VirtualTextureUniform can be either an RVT (URuntimeVirtualTexture) or a VT (UTexture2D)
+	int32 TextureUniformIndex = INDEX_NONE;
+	int32 VTLayerIndex = INDEX_NONE;
+	uint16 VTPageTableIndex = 0;
+	int32 PreallocatedStackTextureIndex = INDEX_NONE;
+
+	if (const MIR::FVirtualTextureUniform* VTUniform = MIR::As<MIR::FVirtualTextureUniform>(VTPageTableRead->VirtualTextureUniform))
+	{
+		TextureUniformIndex = VTUniform->Analysis_UniformIndex;
+		VTLayerIndex = VTUniform->VTLayerIndex;
+		VTPageTableIndex = VTUniform->VTPageTableIndex;
+
+		// RVTs have a pre-baked VT layer index (set by the asset), which means they use preallocated VT stacks.
+		if (VTUniform->VTLayerIndex != INDEX_NONE)
+		{
+			PreallocatedStackTextureIndex = Analyzer.MaterialInterface->GetDefaultTextureIdx(VTUniform->VirtualTexture);
+			check(PreallocatedStackTextureIndex != INDEX_NONE);
+		}
+	}
+	else
+	{
+		const MIR::FTextureUniform* TexUniform = MIR::As<MIR::FTextureUniform>(VTPageTableRead->VirtualTextureUniform);
+		checkf(TexUniform && TexUniform->Type.IsVirtualTexture(), TEXT("VTPageTableRead's VirtualTextureUniform must be FVirtualTextureUniform or a VirtualTexture FTextureUniform"));
+		TextureUniformIndex = TexUniform->Analysis_UniformIndex;
+	}
 
-	const int32 TextureUniformIndex = VirtualTextureUniform->Analysis_UniformIndex;
 	check(TextureUniformIndex >= 0);
 
 	// Only support GPU feedback from pixel shader
 	const bool bGenerateFeedback = VTPageTableRead->bEnableFeedback && Stage == MIR::Stage_Pixel;
 
-	const int32 VTStackIndex = AcquireVTStackIndex(Analyzer, *VTPageTableRead, bGenerateFeedback);
+	const int32 VTStackIndex = AcquireVTStackIndex(Analyzer, *VTPageTableRead, bGenerateFeedback, TextureUniformIndex, PreallocatedStackTextureIndex);
 
 	VTPageTableRead->VTStackIndex = VTStackIndex;
 
 	// Check if VT layer is already known. Otherwise, acquire VT layer
-	int32 VTLayerIndex = VirtualTextureUniform->VTLayerIndex;
 	if (VTLayerIndex != INDEX_NONE)
 	{
 		// The layer index in the VT stack is already known, so fetch the page table from the texture object and assign it to the VT stack
-		VTPageTableRead->VTPageTableIndex = VirtualTextureUniform->VTPageTableIndex;
+		VTPageTableRead->VTPageTableIndex = VTPageTableIndex;
 		Analyzer.CompilationOutput->UniformExpressionSet.SetVTLayer(VTStackIndex, VTLayerIndex, TextureUniformIndex);
 	}
 	else
 	{
-		VTLayerIndex = Analyzer.CompilationOutput->UniformExpressionSet.GetVTStack(VTStackIndex).FindLayer(TextureUniformIndex);
-		if (VTLayerIndex == INDEX_NONE)
+		int32 AcquiredVTLayerIndex = Analyzer.CompilationOutput->UniformExpressionSet.GetVTStack(VTStackIndex).FindLayer(TextureUniformIndex);
+		if (AcquiredVTLayerIndex == INDEX_NONE)
 		{
-			VTLayerIndex = Analyzer.CompilationOutput->UniformExpressionSet.AddVTLayer(VTStackIndex, TextureUniformIndex);
+			AcquiredVTLayerIndex = Analyzer.CompilationOutput->UniformExpressionSet.AddVTLayer(VTStackIndex, TextureUniformIndex);
 		}
-		VTPageTableRead->VTPageTableIndex = VTLayerIndex;
+		VTPageTableRead->VTPageTableIndex = AcquiredVTLayerIndex;
 	}
 }
 
@@ -267,7 +300,13 @@ static int32 RegisterTextureUniform(FMaterialIRValueAnalyzer& Analyzer, FName Na
 
 static void TextureUniform_Analyze(FMaterialIRValueAnalyzer& Analyzer, MIR::FTextureUniform* TextureUniform)
 {
-	TextureUniform->Analysis_UniformIndex = RegisterTextureUniform(Analyzer, FName(TextureUniform->Name), TextureUniform->Texture, TextureUniform->Texture->GetMaterialType(), UINT8_MAX);
+	// We need to handle downgrading virtual texture to texture 2D when virtual texturing is disabled in the project setting
+	EMaterialValueType TextureValueType = TextureUniform->Texture->GetMaterialType();
+	if (TextureValueType == MCT_TextureVirtual && !TextureUniform->Type.IsVirtualTexture())
+	{
+		TextureValueType = MCT_Texture2D;
+	}
+	TextureUniform->Analysis_UniformIndex = RegisterTextureUniform(Analyzer, FName(TextureUniform->Name), TextureUniform->Texture, TextureValueType, UINT8_MAX);
 }
 
 static void VirtualTextureUniform_Analyze(FMaterialIRValueAnalyzer& Analyzer, MIR::FVirtualTextureUniform* VirtualTextureUniform)
diff --git a/Engine/Source/Runtime/Engine/Private/Materials/MaterialIRValueAnalyzer.h b/Engine/Source/Runtime/Engine/Private/Materials/MaterialIRValueAnalyzer.h
--- a/Engine/Source/Runtime/Engine/Private/Materials/MaterialIRValueAnalyzer.h
+++ b/Engine/Source/Runtime/Engine/Private/Materials/MaterialIRValueAnalyzer.h
@@ -83,6 +83,9 @@ struct FMaterialIRValueAnalyzer
 		TextureAddress AddressV{ TA_MAX };
 		MIR::FValue* MipValue{};
 		ETextureMipValueMode MipValueMode{ TMVM_None };
+		bool bIsAdaptive{};
+		float AspectRatio{ 1.0f };
+		int32 PreallocatedStackTextureIndex{ INDEX_NONE };
 	};
 
 	// VT stack entries during analysis. This is later converted to FMaterialVirtualTextureStack entries in FUniformExpressionSet.
diff --git a/Engine/Source/Runtime/Engine/Public/Materials/MaterialIR.h b/Engine/Source/Runtime/Engine/Public/Materials/MaterialIR.h
--- a/Engine/Source/Runtime/Engine/Public/Materials/MaterialIR.h
+++ b/Engine/Source/Runtime/Engine/Public/Materials/MaterialIR.h
@@ -906,6 +906,10 @@ struct FVTPageTableRead : TInstruction<VK_VTPageTableRead, 5>
 
 	// Determines whether adaptive or regular VT page table loads will be emitted, i.e. TextureLoadVirtualPageTable*() or TextureLoadVirtualPageTableAdaptive*().
 	bool bIsAdaptive : 1;
+
+	// Aspect ratio (width/height) of the source texture. Used to prevent merging VTs with different aspect ratios into the same stack.
+	// Defaults to 1.0f for RVTs and texture parameters with no default texture assigned.
+	float AspectRatio = 1.0f;
 };
 
 // Utility value for selecting a different value based on the execution stage.
diff --git a/Engine/Source/Runtime/Engine/Public/Materials/MaterialIRTypes.h b/Engine/Source/Runtime/Engine/Public/Materials/MaterialIRTypes.h
--- a/Engine/Source/Runtime/Engine/Public/Materials/MaterialIRTypes.h
+++ b/Engine/Source/Runtime/Engine/Public/Materials/MaterialIRTypes.h
@@ -26,7 +26,8 @@ enum class ETypeKind : uint8
 	Texture2DArray,			// A 2D texture array object.
 	TextureCubeArray,		// A cube texture array object.
 	Texture3D,				// A 3D (volume) texture object.
-	RuntimeVirtualTexture,	// A runtime virtual texture object.
+	VirtualTexture,			// A virtual texture object (UTexture - MCT_TextureVirtual).
+	RuntimeVirtualTexture,	// A runtime virtual texture object (URuntimeVirtualTexture).
 	ParameterCollection,	// Material parameter collection object.
 	SubstrateData,			// Substrate-specific data structure.
 	VTPageTableResult,		// A virtual texture page table lookup result.
@@ -208,6 +209,9 @@ struct FType
 	// Returns the Texture3D object type.
 	static FType MakeTexture3D();
 
+	// Returns the VirtualTexture object type (for MCT_TextureVirtual, backed by UTexture).
+	static FType MakeVirtualTexture();
+
 	// Returns the RuntimeVirtualTexture object type.
 	static FType MakeRuntimeVirtualTexture();
 	
@@ -271,8 +275,11 @@ struct FType
 	// Returns true if the given type is a matrix primitive.
 	bool IsMatrix() const { return Is(ETypeKind::Primitive) && Primitive.IsMatrix(); }
 
-	// Returns whether this type is any non-virtual texture object.
-	bool IsTexture() const { return Kind == ETypeKind::Texture2D || Kind == ETypeKind::TextureCube || Kind == ETypeKind::Texture2DArray || Kind == ETypeKind::TextureCubeArray || Kind == ETypeKind::Texture3D; }
+	// Returns whether this type is any texture object stored in FTextureUniform (includes virtual textures, excludes RuntimeVirtualTexture).
+	bool IsTexture() const { return Kind == ETypeKind::Texture2D || Kind == ETypeKind::TextureCube || Kind == ETypeKind::Texture2DArray || Kind == ETypeKind::TextureCubeArray || Kind == ETypeKind::Texture3D || Kind == ETypeKind::VirtualTexture; }
+
+	// Returns whether this type is a VirtualTexture object (MCT_TextureVirtual, backed by UTexture).
+	bool IsVirtualTexture() const { return Kind == ETypeKind::VirtualTexture; }
 
 	// Returns whether this type is a RuntimeVirtualTexture object.
 	bool IsRuntimeVirtualTexture() const { return Kind == ETypeKind::RuntimeVirtualTexture; }
