Fixed a few bugs in ASTP:
- Fixed pop when looping an animation.
- Fixed half precision sequences being enabled by default.
- Added SetManualPosition / SetBlendSpacePosition methods to modify position. Calling SetXData now sets all properties instead of skipping them when IsSameAnimation is true. The behavior is more predicable and defined.
- Use FGameTime::GetTimeSinceAppStart().GetWorldTimeSeconds() when no world exists.

#virtualized
#Synced-CL 52519763

#Synced-CL 52519763

diff --git a/Engine/Source/Editor/AnimGraph/Private/AnimSequenceTransformProviderFactories.cpp b/Engine/Source/Editor/AnimGraph/Private/AnimSequenceTransformProviderFactories.cpp
--- a/Engine/Source/Editor/AnimGraph/Private/AnimSequenceTransformProviderFactories.cpp
+++ b/Engine/Source/Editor/AnimGraph/Private/AnimSequenceTransformProviderFactories.cpp
@@ -26,3 +26,15 @@ UObject* UAnimSequenceTransformProviderSequenceListFactory::FactoryCreateNew(UCl
 {
 	return NewObject<UAnimSequenceTransformProviderSequenceList>(InParent, Class, Name, Flags);
 }
+
+UAnimSequenceTransformProviderBlendSpaceListFactory::UAnimSequenceTransformProviderBlendSpaceListFactory()
+{
+	SupportedClass = UAnimSequenceTransformProviderBlendSpaceList::StaticClass();
+	bCreateNew = true;
+	bEditAfterNew = true;
+}
+
+UObject* UAnimSequenceTransformProviderBlendSpaceListFactory::FactoryCreateNew(UClass* Class, UObject* InParent, FName Name, EObjectFlags Flags, UObject* Context, FFeedbackContext* Warn)
+{
+	return NewObject<UAnimSequenceTransformProviderBlendSpaceList>(InParent, Class, Name, Flags);
+}
diff --git a/Engine/Source/Editor/AnimGraph/Private/AnimSequenceTransformProviderFactories.h b/Engine/Source/Editor/AnimGraph/Private/AnimSequenceTransformProviderFactories.h
--- a/Engine/Source/Editor/AnimGraph/Private/AnimSequenceTransformProviderFactories.h
+++ b/Engine/Source/Editor/AnimGraph/Private/AnimSequenceTransformProviderFactories.h
@@ -24,3 +24,13 @@ public:
 	UAnimSequenceTransformProviderSequenceListFactory();
 	virtual UObject* FactoryCreateNew(UClass* Class, UObject* InParent, FName Name, EObjectFlags Flags, UObject* Context, FFeedbackContext* Warn) override;
 };
+
+UCLASS()
+class UAnimSequenceTransformProviderBlendSpaceListFactory : public UFactory
+{
+	GENERATED_BODY()
+
+public:
+	UAnimSequenceTransformProviderBlendSpaceListFactory();
+	virtual UObject* FactoryCreateNew(UClass* Class, UObject* InParent, FName Name, EObjectFlags Flags, UObject* Context, FFeedbackContext* Warn) override;
+};
diff --git a/Engine/Source/Runtime/Engine/Classes/Animation/AnimSequenceTransformProviderData.h b/Engine/Source/Runtime/Engine/Classes/Animation/AnimSequenceTransformProviderData.h
--- a/Engine/Source/Runtime/Engine/Classes/Animation/AnimSequenceTransformProviderData.h
+++ b/Engine/Source/Runtime/Engine/Classes/Animation/AnimSequenceTransformProviderData.h
@@ -39,6 +39,14 @@ enum class EAnimSequenceTrackLoopMode : uint8
 	Clamp
 };
 
+UENUM(BlueprintType)
+enum class EAnimSequenceTrackMode : uint8
+{
+	AutoPlay,
+	Manual,
+	BlendSpace,
+};
+
 UENUM(BlueprintType)
 enum class EAnimSequenceTransformProviderLayerBlendMode : uint8
 {
@@ -258,22 +266,18 @@ public:
 		return FVector2f(BlendSpace.BlendPosition[0].GetFloat(), BlendSpace.BlendPosition[1].GetFloat());
 	}
 
-	void UpdateBlendSpaceData(const FAnimSequenceTrackPackedData& Data, double WorldTime)
+	void SetBlendSpacePosition(FVector2f InPosition)
 	{
-		if (BlendSpace.PlayRate.Encoded != Data.BlendSpace.PlayRate.Encoded)
-		{
-			SetPlayRate(Data.BlendSpace.PlayRate.GetFloat(), WorldTime);
-		}
-		BlendSpace.BlendPosition[0] = Data.BlendSpace.BlendPosition[0];
-		BlendSpace.BlendPosition[1] = Data.BlendSpace.BlendPosition[1];
-		BlendSpace.Samples[0] = Data.BlendSpace.Samples[0];
-		BlendSpace.Samples[1] = Data.BlendSpace.Samples[1];
-		BlendSpace.Samples[2] = Data.BlendSpace.Samples[2];
-		BlendSpace.NumSamples = Data.BlendSpace.NumSamples;
-		LoopMode = Data.LoopMode;
-		BlendTime = Data.BlendTime;
+		BlendSpace.BlendPosition[0] = FFloat16(InPosition.X);
+		BlendSpace.BlendPosition[1] = FFloat16(InPosition.Y);
 	}
 
+	void SetManualPosition(float InPosition)
+	{
+		Manual.Position = InPosition;
+	}
+
+
 	uint32 GetNumSamples() const
 	{
 		return bBlendSpace ? BlendSpace.NumSamples : 1;
@@ -334,21 +338,6 @@ public:
 		return (WorldTime - Auto.ReferenceTimestamp) * Auto.PlayRate.GetFloat();
 	}
 
-	bool IsSameAnimation(const FAnimSequenceTrackPackedData& Other) const
-	{
-		if (bBlendSpace != Other.bBlendSpace)
-		{
-			return false;
-		}
-
-		if (bBlendSpace)
-		{
-			return BlendSpace.BlendSpaceIndex == Other.BlendSpace.BlendSpaceIndex;
-		}
-
-		return SequenceIndex == Other.SequenceIndex;
-	}
-
 	void SetLoopMode(EAnimSequenceTrackLoopMode NewLoopMode)
 	{
 		LoopMode = static_cast<uint32>(NewLoopMode);
@@ -534,7 +523,7 @@ public:
 
 	bool IsBlending() const
 	{
-		return !Target.IsSameAnimation(Source) || BlendWeight.GetFloat() > 0.0f;
+		return BlendWeight.GetFloat() > 0.0f;
 	}
 
 	float GetBlendWeight(bool bPrevious = false) const
@@ -1037,6 +1026,14 @@ public:
 	UFUNCTION(BlueprintCallable, Category = "TransformProviders|AnimSequence")
 	ENGINE_API bool SetBlendSpaceData(int32 TrackIndex, int32 LayerIndex, const FAnimSequenceTrackBlendSpaceData& Data);
 
+	// Updates the blend space 2D position on a track layer without resetting the animation.
+	UFUNCTION(BlueprintCallable, Category = "TransformProviders|AnimSequence")
+	ENGINE_API bool SetBlendSpacePosition(int32 TrackIndex, int32 LayerIndex, FVector2f Position);
+
+	// Updates the manual position on a track layer without resetting the animation.
+	UFUNCTION(BlueprintCallable, Category = "TransformProviders|AnimSequence")
+	ENGINE_API bool SetManualPosition(int32 TrackIndex, int32 LayerIndex, float Position);
+
 	// Sets the play rate on a track layer. Works for auto-play and blend space modes.
 	UFUNCTION(BlueprintCallable, Category = "TransformProviders|AnimSequence")
 	ENGINE_API bool SetPlayRate(int32 TrackIndex, int32 LayerIndex, float PlayRate);
@@ -1065,6 +1062,10 @@ public:
 	UFUNCTION(BlueprintCallable, Category = "TransformProviders|AnimSequence")
 	ENGINE_API int32 GetSequenceIndex(int32 TrackIndex, int32 LayerIndex) const;
 
+	// Returns the current playback mode for a track layer.
+	UFUNCTION(BlueprintCallable, Category = "TransformProviders|AnimSequence")
+	ENGINE_API EAnimSequenceTrackMode GetTrackMode(int32 TrackIndex, int32 LayerIndex) const;
+
 private:
 	void MarkOwnerRenderStateDirty();
 	void InitLayerPools(bool bMarkDirty = true);
diff --git a/Engine/Source/Runtime/Engine/Private/Animation/AnimSequenceTransformProviderData.cpp b/Engine/Source/Runtime/Engine/Private/Animation/AnimSequenceTransformProviderData.cpp
--- a/Engine/Source/Runtime/Engine/Private/Animation/AnimSequenceTransformProviderData.cpp
+++ b/Engine/Source/Runtime/Engine/Private/Animation/AnimSequenceTransformProviderData.cpp
@@ -186,42 +186,19 @@ void FAnimSequenceTrackRenderData::Patch(const FAnimSequenceTrackPackedData& Dat
 			Target.SetLoopMode(Data.GetLoopMode());
 		}
 	}
-	// Full struct replacement. Check if it's the same animation (position update) or a new one (transition).
 	else
 	{
-		if (Target.IsSameAnimation(Data))
+		if (Data.GetBlendTime() <= 0.0f || !bWasActive || !CVarAnimSequenceTransformProviderBlending.GetValueOnRenderThread())
 		{
-			if (Data.IsBlendSpace())
-			{
-				Target.UpdateBlendSpaceData(Data, WorldTime);
-				if (BlendWeight.GetFloat() <= 0.0f)
-				{
-					Source.UpdateBlendSpaceData(Data, WorldTime);
-				}
-			}
-			else
-			{
-				Target.SetData(Data);
-				if (BlendWeight.GetFloat() <= 0.0f)
-				{
-					Source.SetData(Data);
-				}
-			}
+			Source = Data;
 		}
-		else
+		else if (BlendWeight.GetFloat() > 0.0f)
 		{
-			if (Data.GetBlendTime() <= 0.0f || !bWasActive || !CVarAnimSequenceTransformProviderBlending.GetValueOnRenderThread())
-			{
-				Source = Data;
-			}
-			else if (BlendWeight.GetFloat() > 0.0f)
-			{
-				Source = Target;
-			}
-
-			BlendWeight = FFloat16(0.0f);
-			Target = Data;
+			Source = Target;
 		}
+
+		BlendWeight = FFloat16(0.0f);
+		Target = Data;
 	}
 
 	Source.ResetPatchFlags();
@@ -627,7 +604,7 @@ namespace AnimSequenceTrackTime
 			}
 			return GTimeCache.WorldTime;
 		}
-		return 0.0;
+		return FGameTime::GetTimeSinceAppStart().GetWorldTimeSeconds();
 	}
 }
 
@@ -1649,6 +1626,26 @@ int32 UAnimSequenceTransformProviderDataInstance::GetSequenceIndex(int32 TrackIn
 	return INDEX_NONE;
 }
 
+EAnimSequenceTrackMode UAnimSequenceTransformProviderDataInstance::GetTrackMode(int32 TrackIndex, int32 LayerIndex) const
+{
+	if (!IsValidTrackAndLayer(TrackIndex, LayerIndex))
+	{
+		return EAnimSequenceTrackMode::Manual;
+	}
+
+	const FAnimSequenceTrackPackedData& Data = TrackPools[LayerIndex].GetData(TrackIndex);
+
+	if (Data.IsBlendSpace())
+	{
+		return EAnimSequenceTrackMode::BlendSpace;
+	}
+	if (Data.IsAutoPlay())
+	{
+		return EAnimSequenceTrackMode::AutoPlay;
+	}
+	return EAnimSequenceTrackMode::Manual;
+}
+
 float UAnimSequenceTransformProviderDataInstance::GetPosition(int32 TrackIndex, int32 LayerIndex) const
 {
 	if (!IsValidTrackAndLayer(TrackIndex, LayerIndex))
@@ -1687,6 +1684,50 @@ bool UAnimSequenceTransformProviderDataInstance::SetBlendSpaceData(int32 TrackIn
 	return true;
 }
 
+bool UAnimSequenceTransformProviderDataInstance::SetBlendSpacePosition(int32 TrackIndex, int32 LayerIndex, FVector2f Position)
+{
+	if (!IsValidTrackAndLayer(TrackIndex, LayerIndex))
+	{
+		return false;
+	}
+
+	FAnimSequenceTrackPackedData& Data = TrackPools[LayerIndex].GetData(TrackIndex);
+	if (!Data.IsBlendSpace())
+	{
+		return false;
+	}
+
+	Data.SetBlendSpacePosition(Position);
+	TrackPools[LayerIndex].SetDirty(TrackIndex);
+	MarkOwnerRenderStateDirty();
+	return true;
+}
+
+bool UAnimSequenceTransformProviderDataInstance::SetManualPosition(int32 TrackIndex, int32 LayerIndex, float Position)
+{
+	if (!IsValidTrackAndLayer(TrackIndex, LayerIndex))
+	{
+		return false;
+	}
+
+	FAnimSequenceTrackPackedData& Data = TrackPools[LayerIndex].GetData(TrackIndex);
+	if (!Data.IsManual())
+	{
+		return false;
+	}
+
+	const int32 SeqIndex = Data.GetSequenceIndex();
+	if (SeqIndex == INDEX_NONE || !Sequences.IsValidIndex(SeqIndex))
+	{
+		return false;
+	}
+
+	Data.SetManualPosition(FAnimSequenceTrackPackedData::WrapPosition(Data.GetLoopMode(), Position, GetSequencePlayLength(SeqIndex)));
+	TrackPools[LayerIndex].SetDirty(TrackIndex);
+	MarkOwnerRenderStateDirty();
+	return true;
+}
+
 FVector2f UAnimSequenceTransformProviderDataInstance::GetBlendSpacePosition(int32 TrackIndex, int32 LayerIndex) const
 {
 	if (IsValidTrackAndLayer(TrackIndex, LayerIndex))
diff --git a/Engine/Source/Runtime/Renderer/Private/Skinning/AnimSequenceTransformProvider.cpp b/Engine/Source/Runtime/Renderer/Private/Skinning/AnimSequenceTransformProvider.cpp
--- a/Engine/Source/Runtime/Renderer/Private/Skinning/AnimSequenceTransformProvider.cpp
+++ b/Engine/Source/Runtime/Renderer/Private/Skinning/AnimSequenceTransformProvider.cpp
@@ -103,7 +103,7 @@ static TAutoConsoleVariable<int32> CVarAnimSequenceTransformProviderForceSampleR
 
 static TAutoConsoleVariable<bool> CVarAnimSequenceTransformProviderHalfPrecision(
 	TEXT("r.AnimSequenceTransformProvider.HalfPrecision"),
-	true,
+	false,
 	TEXT("Pack sequence transforms as 16-bit floats. Halves sequence buffer memory with minor precision loss.\n")
 	TEXT("Read-only: must be set at startup. Requires shader recompile in editor, cook for client builds."),
 	ECVF_ReadOnly
@@ -140,7 +140,7 @@ uint32 FAnimSequenceTransformProvider::FSequence::GetSampleStep() const
 	return 1u << FMath::Clamp<uint32>(SampleRateLevel, CurrentSampleRateLevel, NumSampleRateLevels - 1);
 }
 
-void FAnimSequenceTransformProvider::FSequence::GetKeyIndicesFromTime(int32& OutKey0, int32& OutKey1, float& OutAlpha, float Time) const
+void FAnimSequenceTransformProvider::FSequence::GetKeyIndicesFromTime(int32& OutKey0, int32& OutKey1, float& OutAlpha, float Time, EAnimSequenceTrackLoopMode LoopMode) const
 {
 	const auto AlignDown = [] (int32 Value, int32 Alignment)
 	{
@@ -151,9 +151,15 @@ void FAnimSequenceTransformProvider::FSequence::GetKeyIndicesFromTime(int32& Out
 	const float ExactFrame = Time * NumFrames / PlayLength;
 	const int32 LastAlignedFrame = AlignDown((int32)NumFrames - 1, Step);
 
-	OutKey0  = FMath::Min(AlignDown((int32)ExactFrame, Step), LastAlignedFrame);
-	OutKey1  = FMath::Min(OutKey0 + Step, LastAlignedFrame);
-	OutAlpha = OutKey1 > OutKey0 ? (ExactFrame - (float)OutKey0) / (float)Step : 0.0f;
+	OutKey0 = FMath::Min(AlignDown((int32)ExactFrame, Step), LastAlignedFrame);
+	OutKey1 = OutKey0 + Step;
+
+	if (OutKey1 > LastAlignedFrame)
+	{
+		OutKey1 = (LoopMode == EAnimSequenceTrackLoopMode::Loop) ? 0 : LastAlignedFrame;
+	}
+
+	OutAlpha = (OutKey1 != OutKey0) ? (ExactFrame - (float)OutKey0) / (float)Step : 0.0f;
 }
 
 class FEvaluateConcatenateScatterCS : public FGlobalShader
@@ -1039,7 +1045,7 @@ void FAnimSequenceTransformProvider::UploadBlockHeaders(
 
 				int32 KeyIndex0, KeyIndex1;
 				float Alpha;
-				TrackData.Sequence.GetKeyIndicesFromTime(KeyIndex0, KeyIndex1, Alpha, TrackData.Position);
+				TrackData.Sequence.GetKeyIndicesFromTime(KeyIndex0, KeyIndex1, Alpha, TrackData.Position, TrackData.LoopMode);
 				NumTotalFramesRequested += TrackData.Sequence.RequestFrame(KeyIndex0);
 				NumTotalFramesRequested += TrackData.Sequence.RequestFrame(KeyIndex1);
 
@@ -1066,7 +1072,7 @@ void FAnimSequenceTransformProvider::UploadBlockHeaders(
 
 			int32 KeyIndex0, KeyIndex1;
 			float Alpha;
-			TrackData.Sequence.GetKeyIndicesFromTime(KeyIndex0, KeyIndex1, Alpha, TrackData.Position);
+			TrackData.Sequence.GetKeyIndicesFromTime(KeyIndex0, KeyIndex1, Alpha, TrackData.Position, TrackData.LoopMode);
 			NumTotalFramesRequested += TrackData.Sequence.RequestFrame(KeyIndex0);
 			NumTotalFramesRequested += TrackData.Sequence.RequestFrame(KeyIndex1);
 
diff --git a/Engine/Source/Runtime/Renderer/Private/Skinning/AnimSequenceTransformProvider.h b/Engine/Source/Runtime/Renderer/Private/Skinning/AnimSequenceTransformProvider.h
--- a/Engine/Source/Runtime/Renderer/Private/Skinning/AnimSequenceTransformProvider.h
+++ b/Engine/Source/Runtime/Renderer/Private/Skinning/AnimSequenceTransformProvider.h
@@ -225,7 +225,7 @@ private:
 
 		uint32 GetSampleStep() const;
 
-		void GetKeyIndicesFromTime(int32& OutKey0, int32& OutKey1, float& OutAlpha, float Time) const;
+		void GetKeyIndicesFromTime(int32& OutKey0, int32& OutKey1, float& OutAlpha, float Time, EAnimSequenceTrackLoopMode LoopMode) const;
 	};
 
 	struct FMapping
@@ -251,6 +251,7 @@ private:
 		FSequence& Sequence;
 		uint16 SequenceIndex;
 		float Position;
+		EAnimSequenceTrackLoopMode LoopMode;
 	};
 
 	class FProxyUserData : public FInstancedSkinningSceneExtensionProxyUserData
@@ -288,6 +289,7 @@ private:
 				, .Sequence = Sequence
 				, .SequenceIndex = InSequenceIndex
 				, .Position = WrappedPosition
+				, .LoopMode = LoopMode
 			};
 		}
 
