Prevent DataDriven Cvars from being unregistered during GC from FSettingsEditorCustomization::CacheResetDefaultPropertyValues

#Synced-CL 52530536

diff --git a/Engine/Source/Runtime/Engine/Classes/DataDrivenCVars/DataDrivenCVars.h b/Engine/Source/Runtime/Engine/Classes/DataDrivenCVars/DataDrivenCVars.h
--- a/Engine/Source/Runtime/Engine/Classes/DataDrivenCVars/DataDrivenCVars.h
+++ b/Engine/Source/Runtime/Engine/Classes/DataDrivenCVars/DataDrivenCVars.h
@@ -73,6 +73,8 @@ struct FDataDrivenConsoleVariable
 	ENGINE_API void Register();
 	ENGINE_API void UnRegister(bool bUseShadowName = false);
 
+	bool bIsRegistered = false;
+
 	FString ShadowName;
 	FString ShadowToolTip;
 	FDataDrivenCVarType ShadowType;
diff --git a/Engine/Source/Runtime/Engine/Private/DataDrivenCVars/DataDrivenCVars.cpp b/Engine/Source/Runtime/Engine/Private/DataDrivenCVars/DataDrivenCVars.cpp
--- a/Engine/Source/Runtime/Engine/Private/DataDrivenCVars/DataDrivenCVars.cpp
+++ b/Engine/Source/Runtime/Engine/Private/DataDrivenCVars/DataDrivenCVars.cpp
@@ -17,7 +17,7 @@ FDataDrivenConsoleVariable::~FDataDrivenConsoleVariable()
 
 void FDataDrivenConsoleVariable::Register()
 {
-	if (!Name.IsEmpty())
+	if (!Name.IsEmpty() && !bIsRegistered)
 	{
 		IConsoleVariable* CVarToAdd = IConsoleManager::Get().FindConsoleVariable(*Name);
 		if (CVarToAdd == nullptr)
@@ -39,13 +39,14 @@ void FDataDrivenConsoleVariable::Register()
 		ShadowName = Name;
 		ShadowToolTip = ToolTip;
 		ShadowType = Type;
+		bIsRegistered = true;
 	}
 }
 
 void FDataDrivenConsoleVariable::UnRegister(bool bUseShadowName)
 {
 	const FString& NameToUnregister = bUseShadowName ? ShadowName : Name;
-	if (!NameToUnregister.IsEmpty())
+	if (!NameToUnregister.IsEmpty() && bIsRegistered)
 	{
 		IConsoleVariable* CVarToRemove = IConsoleManager::Get().FindConsoleVariable(*NameToUnregister);
 		if (CVarToRemove)
@@ -54,6 +55,8 @@ void FDataDrivenConsoleVariable::UnRegister(bool bUseShadowName)
 			CVarToRemove->SetOnChangedCallback(NullCallback);
 			IConsoleManager::Get().UnregisterConsoleObject(CVarToRemove, false);
 		}
+
+		bIsRegistered = false;
 	}
 }
 
@@ -101,6 +104,14 @@ void UDataDrivenConsoleVariableSettings::PostInitProperties()
 {
 	Super::PostInitProperties();
 
+	if (GetDefault<UDataDrivenConsoleVariableSettings>() != this)
+	{
+		// We are a temp copy so there's no need to re-register our cvars.
+		// i.e. FSettingsEditorCustomization::CacheResetDefaultPropertyValues makes a temp copy to check defaults.
+		return;
+	}
+
+
 	for (FDataDrivenConsoleVariable& CVar : CVarsArray)
 	{
 		CVar.Register();
