Fix MaterialIR preshader crash: Dot op dummy constant emitted after operator violates topological order.
Later on this would  cause a crash in PreshaderFixup() where a negative value would be generated for FreeComponentOffset().
#rnx
#Synced-CL 52535212

diff --git a/Engine/Source/Runtime/Engine/Public/Materials/MaterialIRPreshader.cpp b/Engine/Source/Runtime/Engine/Public/Materials/MaterialIRPreshader.cpp
--- a/Engine/Source/Runtime/Engine/Public/Materials/MaterialIRPreshader.cpp
+++ b/Engine/Source/Runtime/Engine/Public/Materials/MaterialIRPreshader.cpp
@@ -385,6 +385,32 @@ void FMaterialIRPreshader::EmitValue(MIR::FValue* Value)
 			EmitComposite(Operator->BArg->As<MIR::FComposite>());
 		}
 
+		// For Dot operations with a constant BArg, pre-emit the dummy constant BEFORE registering the operator itself.  
+		// EmitDummyConstant adds to ValueInfos, so it must come first to preserve topological ordering: every source must have a lower preshader index than its consumers.  
+		MIR::FValue* PreEmittedDotDummyBArg = nullptr;
+		if (Operator->Op == MIR::BO_Dot && Operator->BArg)
+		{
+			// Figure out which arg is the preshader and which is the constant, accounting for the commutative swap that will happen later.
+			MIR::FValue* DotPreshaderArg = nullptr;
+			MIR::FValue* DotConstantArg  = nullptr;
+			if (Operator->AArg->Analysis_PreshaderOffset && !Operator->BArg->Analysis_PreshaderOffset)
+			{
+				DotPreshaderArg = Operator->AArg;
+				DotConstantArg  = Operator->BArg;
+			}
+			else if (!Operator->AArg->Analysis_PreshaderOffset && Operator->BArg->Analysis_PreshaderOffset)
+			{
+				// AArg is constant and will be swapped to BArg position below.
+				DotPreshaderArg = Operator->BArg;
+				DotConstantArg  = Operator->AArg;
+			}
+			if (DotConstantArg)
+			{
+				EPreshader2Type DotPreshaderType = ToPreshader2Type(Operator->Type.AsPrimitive()->ScalarKind);
+				PreEmittedDotDummyBArg = EmitDummyConstant(DotConstantArg, DotPreshaderType, DotPreshaderArg->Type.AsPrimitive()->NumComponents());
+			}
+		}
+
 		// Generate preshader operator
 		int32 PreshaderIndex = AddValueInfo(Value);
 		FMaterialIRPreshader::FValueInfo& ValueInfo = ValueInfos[PreshaderIndex];
@@ -472,10 +498,10 @@ void FMaterialIRPreshader::EmitValue(MIR::FValue* Value)
 				{
 					if (Operator->Op == MIR::BO_Dot)
 					{
-						// Allocate a dummy constant, since the constant can't fit in the destination, which is a scalar.
-						// Because Dot is commutative, we know AArg is not a constant (if it were constant, the arguments
-						// would have been swapped), and so we don't need this special case above where AArg is handled.
-						BArg = EmitDummyConstant(BArg, Preshader2Type, AArg->Type.AsPrimitive()->NumComponents());
+						// The dummy constant should have been pre-emitted before this operator to preserve topological ordering in ValueInfos.  
+						// Use the already-allocated value rather than emitting a new one.
+						check(PreEmittedDotDummyBArg != nullptr);
+						BArg = PreEmittedDotDummyBArg;
 
 						// Allow the temporary memory for the dummy constant to be reclaimed by tracking its usage.
 						AddValueUse(BArg, Value);
@@ -1063,6 +1089,10 @@ void FMaterialIRPreshader::PreshaderFixup(UE::Shader::FPreshaderData& PreshaderO
 					continue;
 				}
 
+				// All ValueInfos must appear topologically after values they use in the array, and so the used values should already have had their Remap set
+				// in a previous loop iteration. Hitting this check implies ValueInfos were added out of dependency order in prior code.
+				check(Remap[UseIndex]);
+
 				// Temporary buffer starts at offset "UniformBufferWords".
 				if ((ValueInfos[UseIndex].LastUsedBy == PreshaderIndex) && ((uint16)(Remap[UseIndex] - 1) >= UniformBufferWords))		// Subtract one!
 				{
