[SkeletalMesh] Fix deadlock during reimport by using ExecuteOnGameThread for vertex attributes changed notification

FFunctionGraphTask::CreateAndDispatchWhenReady can execute during task graph pumping
(e.g. FFrameEndSync::Sync) while the skeletal mesh property lock is still held. The
vertex attributes changed callback then tries to access locked properties, causing a
deadlock. ExecuteOnGameThread is a safer dispatch method that avoids this.

#Synced-CL 52437958

diff --git a/Engine/Source/Runtime/Engine/Private/SkeletalMesh.cpp b/Engine/Source/Runtime/Engine/Private/SkeletalMesh.cpp
--- a/Engine/Source/Runtime/Engine/Private/SkeletalMesh.cpp
+++ b/Engine/Source/Runtime/Engine/Private/SkeletalMesh.cpp
@@ -6,6 +6,7 @@
 
 #include "Engine/SkeletalMesh.h"
 
+#include "Containers/Ticker.h"
 #include "DataDrivenShaderPlatformInfo.h"
 #include "Engine/Blueprint.h"
 #include "Engine/SkeletalMeshSampling.h"
@@ -3093,15 +3094,14 @@ bool USkeletalMesh::CommitMeshDescription(
 			{
 #if WITH_EDITOR
 				// Notify UI and other systems of the change
-				// Dispatch it on the game thread for thread-safety as this can be called on a worker thread
-				FFunctionGraphTask::CreateAndDispatchWhenReady(
+				ExecuteOnGameThread(UE_SOURCE_LOCATION,
 					[WeakSkelMesh = TWeakObjectPtr<USkeletalMesh>(this)]()
 					{
 						if (USkeletalMesh* SkeletalMesh = WeakSkelMesh.Get())
 						{
 							SkeletalMesh->GetOnVertexAttributesArrayChanged().Broadcast();
 						}
-					}, TStatId(), nullptr, ENamedThreads::GameThread);	
+					});	
 #endif			
 			}
 			
