Fix static analysis error.
Add check() to indicate known non-null values.

#rnx
#Synced-CL 52572089

diff --git a/Engine/Source/Runtime/Engine/Public/Materials/MaterialIRPreshader.cpp b/Engine/Source/Runtime/Engine/Public/Materials/MaterialIRPreshader.cpp
--- a/Engine/Source/Runtime/Engine/Public/Materials/MaterialIRPreshader.cpp
+++ b/Engine/Source/Runtime/Engine/Public/Materials/MaterialIRPreshader.cpp
@@ -374,6 +374,8 @@ void FMaterialIRPreshader::EmitValue(MIR::FValue* Value)
 	}
 	else if (MIR::FOperator* Operator = Value->As<MIR::FOperator>())
 	{
+		check(Operator->AArg);
+
 		// Check if any composite children of this operator need to be expanded.  We defer this so we don't generate opcodes
 		// if it's not actually used (common for trivial tail swizzles that have been dead stripped).
 		if (!Operator->AArg->Analysis_PreshaderOffset && Operator->AArg->HasSubgraphProperties(MIR::EGraphProperties::HasParameter) && Operator->AArg->IsA(MIR::VK_Composite))
