Additional header fix.

#Synced-CL 52531771

diff --git a/Engine/Source/Runtime/Engine/Private/Net/RPCDoSDetectionAnalyticsImpl.h b/Engine/Source/Runtime/Engine/Private/Net/RPCDoSDetectionAnalyticsImpl.h
--- a/Engine/Source/Runtime/Engine/Private/Net/RPCDoSDetectionAnalyticsImpl.h
+++ b/Engine/Source/Runtime/Engine/Private/Net/RPCDoSDetectionAnalyticsImpl.h
@@ -1,5 +1,7 @@
 // Copyright Epic Games, Inc. All Rights Reserved.
 
+#pragma once
+
 // Includes
 #include "Net/Core/Analytics/RPCDoSDetectionAnalytics.h"
 
