- Fix potential out-of-bounds crash in GetConnectedOutput by validating the output index before accessing the expression outputs array

#Synced-CL 52578968

diff --git a/Engine/Source/Runtime/Engine/Private/Materials/MaterialShared.cpp b/Engine/Source/Runtime/Engine/Private/Materials/MaterialShared.cpp
--- a/Engine/Source/Runtime/Engine/Private/Materials/MaterialShared.cpp
+++ b/Engine/Source/Runtime/Engine/Private/Materials/MaterialShared.cpp
@@ -425,7 +425,12 @@ FExpressionInput FExpressionInput::GetTracedInput() const
 
 FExpressionOutput* FExpressionInput::GetConnectedOutput() const
 {
-	return IsConnected() ? &Expression->GetOutputs()[OutputIndex] : nullptr;
+	if (IsConnected())
+	{
+		TArray<FExpressionOutput>& ExpressionOutputs = Expression->GetOutputs();
+		return ExpressionOutputs.IsValidIndex(OutputIndex) ? &ExpressionOutputs[OutputIndex] : nullptr;
+	}
+	return nullptr;
 }
 
 #endif // WITH_EDITOR
