Viewport: Use managed pointer for reference to Client

[RN] Viewport now keeps its assigned Viewport Client from being deleted (prevent use-after-free)
#Synced-CL 52565131

diff --git a/Engine/Plugins/Compositing/Composure/Source/ComposureLayersEditor/Private/CompElementManager.cpp b/Engine/Plugins/Compositing/Composure/Source/ComposureLayersEditor/Private/CompElementManager.cpp
--- a/Engine/Plugins/Compositing/Composure/Source/ComposureLayersEditor/Private/CompElementManager.cpp
+++ b/Engine/Plugins/Compositing/Composure/Source/ComposureLayersEditor/Private/CompElementManager.cpp
@@ -1,6 +1,7 @@
 // Copyright Epic Games, Inc. All Rights Reserved.
 
 #include "CompElementManager.h"
+#include "Templates/SharedPointer.h"
 #include "CompositingElement.h"
 #include "Engine/Level.h"
 #include "Selection.h"
@@ -50,7 +51,7 @@ void FCompElementManager::Initialize()
 
 	if (FSlateApplication::IsInitialized())
 	{
-		EditorCompositingViewport = MakeShareable(new FCompositingViewportClient(ElementsContainer));
+		EditorCompositingViewport = FCompositingViewport::Create(ElementsContainer);
 	}
 
 	FEditorDelegates::MapChange.AddRaw(this, &FCompElementManager::OnEditorMapChange);
@@ -511,12 +512,12 @@ void FCompElementManager::RequestRedraw()
 {
 	if (!EditorCompositingViewport.IsValid() && FSlateApplication::IsInitialized())
 	{
-		EditorCompositingViewport = MakeShareable(new FCompositingViewportClient(ElementsContainer));
+		EditorCompositingViewport = FCompositingViewport::Create(ElementsContainer);
 	}
 
 	if (EditorCompositingViewport.IsValid())
 	{
-		EditorCompositingViewport->RedrawRequested(EditorCompositingViewport->Viewport);
+		EditorCompositingViewport->RequestRedraw();
 	}
 }
 
diff --git a/Engine/Plugins/Compositing/Composure/Source/ComposureLayersEditor/Private/CompElementManager.h b/Engine/Plugins/Compositing/Composure/Source/ComposureLayersEditor/Private/CompElementManager.h
--- a/Engine/Plugins/Compositing/Composure/Source/ComposureLayersEditor/Private/CompElementManager.h
+++ b/Engine/Plugins/Compositing/Composure/Source/ComposureLayersEditor/Private/CompElementManager.h
@@ -10,7 +10,7 @@ class UEditorCompElementContainer;
 class AActor;
 class UPackage;
 class UEditorEngine;
-class FCompositingViewportClient;
+class FCompositingViewport;
 class ULevel;
 
 /** Management class for interfacing with editor compositing element objects. */
@@ -102,7 +102,7 @@ private:
 	/** UObject in charge of tracking all editor element actors. Separated as a UObject to more easily facilitate undo/redo actions. */
 	TObjectPtr<UEditorCompElementContainer> ElementsContainer;
 	/** Hidden editor viewport, in charge on enqueuing compositing render commands. */
-	TSharedPtr<FCompositingViewportClient> EditorCompositingViewport;
+	TSharedPtr<FCompositingViewport> EditorCompositingViewport;
 	
 	/** 
 	 * List which tracks elements queued for delete - used to circumvent broken 
diff --git a/Engine/Plugins/Compositing/Composure/Source/ComposureLayersEditor/Private/CompositingViewportClient.cpp b/Engine/Plugins/Compositing/Composure/Source/ComposureLayersEditor/Private/CompositingViewportClient.cpp
--- a/Engine/Plugins/Compositing/Composure/Source/ComposureLayersEditor/Private/CompositingViewportClient.cpp
+++ b/Engine/Plugins/Compositing/Composure/Source/ComposureLayersEditor/Private/CompositingViewportClient.cpp
@@ -1,7 +1,6 @@
 // Copyright Epic Games, Inc. All Rights Reserved.
 
 #include "CompositingViewportClient.h"
-#include "DummyViewport.h"
 #include "CompositingElement.h"
 #include "Editor/UnrealEdTypes.h"
 #include "EditorCompElementContainer.h"
@@ -24,38 +23,53 @@ DEFINE_LOG_CATEGORY_STATIC(LogComposureCompositingEditor, Log, All)
 /* FCompositingViewport
  *****************************************************************************/
 
-class FCompositingViewport : public FDummyViewport
+FCompositingViewport::FCompositingViewport(const TSharedRef<FCompositingViewportClient>& InViewportClient)
+	: FDummyViewport(InViewportClient)
 {
-public:
-	FCompositingViewport(FCompositingViewportClient* InViewportClient)
-		: FDummyViewport(InViewportClient)
-	{
-		// Need a non-zero size to call into FEditorViewportClient::Draw()
-		SizeX = 1920;
-		SizeY = 1080;
-	}
-};
+	// Need a non-zero size to call into FEditorViewportClient::Draw()
+	SizeX = 1920;
+	SizeY = 1080;
+}
+
+TSharedRef<FCompositingViewport> FCompositingViewport::Create(TWeakObjectPtr<UEditorCompElementContainer>&& CompElements)
+{
+	TSharedRef<FCompositingViewportClient> Client = MakeShared<FCompositingViewportClient>(MoveTemp(CompElements));
+	TSharedRef<FCompositingViewport> Result = MakeShareable(new FCompositingViewport(Client));
+	Client->Viewport = &Result.Get();
+	return Result;
+}
+
+FCompositingViewportClient& FCompositingViewport::GetCompositingClient() const
+{
+	return static_cast<FCompositingViewportClient&>(GetClientChecked());
+}
+
+bool FCompositingViewport::IsDrawing() const
+{
+	return GetCompositingClient().IsDrawing();
+}
+
+void FCompositingViewport::RequestRedraw()
+{
+	GetCompositingClient().RedrawRequested(this);
+}
 
 /* FCompositingViewportClient
  *****************************************************************************/
 
-FCompositingViewportClient::FCompositingViewportClient(TWeakObjectPtr<UEditorCompElementContainer> CompElements)
+FCompositingViewportClient::FCompositingViewportClient(TWeakObjectPtr<UEditorCompElementContainer>&& CompElements)
 	: FEditorViewportClient(nullptr)
-	, ElementsContainerPtr(CompElements)
+	, ElementsContainerPtr(MoveTemp(CompElements))
 {
 	SetViewModes(VMI_Unlit, VMI_Unlit);
 	SetViewportType(LVT_OrthoFreelook);
 
 	//SetRealtime(true);
 	VisibilityDelegate.BindRaw(this, &FCompositingViewportClient::InternalIsVisible);
-
-	CompositingViewport = MakeShareable(new FCompositingViewport(this));
-	Viewport = CompositingViewport.Get();
 }
 
 FCompositingViewportClient::~FCompositingViewportClient()
 {
-	CompositingViewport.Reset();
 	Viewport = nullptr;
 }
 
diff --git a/Engine/Plugins/Compositing/Composure/Source/ComposureLayersEditor/Private/CompositingViewportClient.h b/Engine/Plugins/Compositing/Composure/Source/ComposureLayersEditor/Private/CompositingViewportClient.h
--- a/Engine/Plugins/Compositing/Composure/Source/ComposureLayersEditor/Private/CompositingViewportClient.h
+++ b/Engine/Plugins/Compositing/Composure/Source/ComposureLayersEditor/Private/CompositingViewportClient.h
@@ -2,20 +2,20 @@
 
 #pragma once
 
+#include "DummyViewport.h"
 #include "EditorViewportClient.h"
 #include "TickableEditorObject.h"
 
 class FCanvas;
-class FCompositingViewport;
 class UEditorCompElementContainer;
 
-class FCompositingViewportClient : public FEditorViewportClient, public FTickableEditorObject
+class FCompositingViewportClient final : public FEditorViewportClient, public FTickableEditorObject
 {
 public:
-	 FCompositingViewportClient(TWeakObjectPtr<UEditorCompElementContainer> CompElements);
-	 virtual ~FCompositingViewportClient();
+	FCompositingViewportClient(TWeakObjectPtr<UEditorCompElementContainer>&& CompElements);
+	virtual ~FCompositingViewportClient();
 
-	 bool IsDrawing() const;
+	bool IsDrawing() const;
 
 	//~ Begin FViewportClient interface
 	virtual void Draw(const FSceneView* View, FPrimitiveDrawInterface* PDI) override;
@@ -37,7 +37,17 @@ public:
 private:
 	bool InternalIsVisible() const;
 
-	TSharedPtr<FCompositingViewport> CompositingViewport;
 	TWeakObjectPtr<UEditorCompElementContainer> ElementsContainerPtr;
 	bool bIsDrawing = false;
 };
+
+class FCompositingViewport final : public FDummyViewport
+{
+	FCompositingViewport(const TSharedRef<FCompositingViewportClient>& InViewportClient);
+public:
+	static TSharedRef<FCompositingViewport> Create(TWeakObjectPtr<UEditorCompElementContainer>&& CompElements);
+
+	FCompositingViewportClient& GetCompositingClient() const;
+	bool IsDrawing() const;
+	void RequestRedraw();
+};
diff --git a/Engine/Plugins/FX/Cascade/Source/Cascade/Private/SCascadeEmitterCanvas.cpp b/Engine/Plugins/FX/Cascade/Source/Cascade/Private/SCascadeEmitterCanvas.cpp
--- a/Engine/Plugins/FX/Cascade/Source/Cascade/Private/SCascadeEmitterCanvas.cpp
+++ b/Engine/Plugins/FX/Cascade/Source/Cascade/Private/SCascadeEmitterCanvas.cpp
@@ -74,7 +74,7 @@ void SCascadeEmitterCanvas::Construct(const FArguments& InArgs)
 
 	ViewportClient->VisibilityDelegate.BindSP(this, &SCascadeEmitterCanvas::IsVisible);
 
-	Viewport = MakeShareable(new FSceneViewport(ViewportClient.Get(), ViewportWidget));
+	Viewport = MakeShared<FSceneViewport>(ViewportClient, ViewportWidget);
 	ViewportClient->Viewport = Viewport.Get();
 
 	// The viewport widget needs an interface so it knows what should render
diff --git a/Engine/Plugins/Media/MediaFrameworkUtilities/Source/MediaFrameworkUtilities/Private/Profile/MediaProfilePlaybackManager.cpp b/Engine/Plugins/Media/MediaFrameworkUtilities/Source/MediaFrameworkUtilities/Private/Profile/MediaProfilePlaybackManager.cpp
--- a/Engine/Plugins/Media/MediaFrameworkUtilities/Source/MediaFrameworkUtilities/Private/Profile/MediaProfilePlaybackManager.cpp
+++ b/Engine/Plugins/Media/MediaFrameworkUtilities/Source/MediaFrameworkUtilities/Private/Profile/MediaProfilePlaybackManager.cpp
@@ -165,8 +165,8 @@ namespace MediaProfilePlaybackManager
 			.IsEnabled(FSlateApplication::Get().GetNormalExecutionAttribute())
 			.EnableGammaCorrection(false)
 			.EnableBlending(false);
-		
-		ManagedViewport->Viewport = MakeShared<FSceneViewport>(ViewportClient.Get(), ManagedViewport->ViewportWidget);
+
+		ManagedViewport->Viewport = MakeShared<FSceneViewport>(ViewportClient, ManagedViewport->ViewportWidget);
 		ManagedViewport->ViewportWidget->SetViewportInterface(ManagedViewport->Viewport.ToSharedRef());
 
 		ViewportClient->Viewport = ManagedViewport->Viewport.Get();
@@ -174,7 +174,7 @@ namespace MediaProfilePlaybackManager
 		TSharedPtr<FMediaCaptureGameViewportClient> ViewportClient = MakeShared<FMediaCaptureGameViewportClient>();
 
 		ManagedViewport->ViewportClient = ViewportClient;
-		ManagedViewport->Viewport = MakeShared<FSceneViewport>(ViewportClient.Get(), nullptr);
+		ManagedViewport->Viewport = MakeShared<FSceneViewport>(ViewportClient, nullptr);
 #endif
 
 		return ManagedViewport;
diff --git a/Engine/Plugins/Media/MediaFrameworkUtilities/Source/MediaFrameworkUtilitiesEditor/Private/CaptureTab/SMediaFrameworkCaptureOutputWidget.cpp b/Engine/Plugins/Media/MediaFrameworkUtilities/Source/MediaFrameworkUtilitiesEditor/Private/CaptureTab/SMediaFrameworkCaptureOutputWidget.cpp
--- a/Engine/Plugins/Media/MediaFrameworkUtilities/Source/MediaFrameworkUtilitiesEditor/Private/CaptureTab/SMediaFrameworkCaptureOutputWidget.cpp
+++ b/Engine/Plugins/Media/MediaFrameworkUtilities/Source/MediaFrameworkUtilitiesEditor/Private/CaptureTab/SMediaFrameworkCaptureOutputWidget.cpp
@@ -305,7 +305,7 @@ void SMediaFrameworkCaptureCameraViewportWidget::Construct(const FArguments& InA
 		.EnableGammaCorrection(false)		// Scene rendering handles gamma correction
 		.EnableBlending(false);
 
-	SceneViewport = MakeShareable(new FSceneViewport(LevelViewportClient.Get(), ViewportWidget));
+	SceneViewport = MakeShared<FSceneViewport>(LevelViewportClient, ViewportWidget);
 	{
 		LevelViewportClient->Viewport = SceneViewport.Get();
 		ViewportWidget->SetViewportInterface(SceneViewport.ToSharedRef());
diff --git a/Engine/Plugins/Mutable/Source/CustomizableObjectEditor/Private/MuCOE/SCustomizableObjectEditorViewport.cpp b/Engine/Plugins/Mutable/Source/CustomizableObjectEditor/Private/MuCOE/SCustomizableObjectEditorViewport.cpp
--- a/Engine/Plugins/Mutable/Source/CustomizableObjectEditor/Private/MuCOE/SCustomizableObjectEditorViewport.cpp
+++ b/Engine/Plugins/Mutable/Source/CustomizableObjectEditor/Private/MuCOE/SCustomizableObjectEditorViewport.cpp
@@ -68,7 +68,7 @@ TSharedRef<FEditorViewportClient> SCustomizableObjectEditorViewport::MakeEditorV
 	LevelViewportClient->SetViewLocation(EditorViewportDefs::DefaultPerspectiveViewLocation);
 	LevelViewportClient->SetViewRotation(EditorViewportDefs::DefaultPerspectiveViewRotation);
 
-	SceneViewport = MakeShareable(new FSceneViewport(LevelViewportClient.Get(), ViewportWidget));
+	SceneViewport = MakeShared<FSceneViewport>(LevelViewportClient, ViewportWidget);
 
 	return LevelViewportClient.ToSharedRef();
 }
diff --git a/Engine/Plugins/Runtime/nDisplay/Source/DisplayCluster/Private/Game/EngineClasses/Basics/DisplayClusterViewportClient.cpp b/Engine/Plugins/Runtime/nDisplay/Source/DisplayCluster/Private/Game/EngineClasses/Basics/DisplayClusterViewportClient.cpp
--- a/Engine/Plugins/Runtime/nDisplay/Source/DisplayCluster/Private/Game/EngineClasses/Basics/DisplayClusterViewportClient.cpp
+++ b/Engine/Plugins/Runtime/nDisplay/Source/DisplayCluster/Private/Game/EngineClasses/Basics/DisplayClusterViewportClient.cpp
@@ -157,8 +157,8 @@ static UCanvas* GetCanvasByName(FName CanvasName)
 class FDisplayClusterSceneViewport : public FSceneViewport
 {
 public:
-	FDisplayClusterSceneViewport(FViewportClient* InViewportClient, TSharedPtr<SViewport> InViewportWidget)
-		: FSceneViewport(InViewportClient, InViewportWidget)
+	FDisplayClusterSceneViewport(TStrongPtrVariant<FViewportClient>&& InViewportClient, const TSharedPtr<SViewport>& InViewportWidget)
+		: FSceneViewport(MoveTemp(InViewportClient), InViewportWidget)
 	{}
 
 	~FDisplayClusterSceneViewport()
diff --git a/Engine/Plugins/VirtualProduction/CameraCalibration/Source/CameraCalibrationEditor/Private/AssetEditor/Viewport/SSimulcamEditorViewport.cpp b/Engine/Plugins/VirtualProduction/CameraCalibration/Source/CameraCalibrationEditor/Private/AssetEditor/Viewport/SSimulcamEditorViewport.cpp
--- a/Engine/Plugins/VirtualProduction/CameraCalibration/Source/CameraCalibrationEditor/Private/AssetEditor/Viewport/SSimulcamEditorViewport.cpp
+++ b/Engine/Plugins/VirtualProduction/CameraCalibration/Source/CameraCalibrationEditor/Private/AssetEditor/Viewport/SSimulcamEditorViewport.cpp
@@ -69,7 +69,7 @@ void SSimulcamEditorViewport::Construct(const FArguments& InArgs, const TSharedR
 	];
 
 	ViewportClient = MakeShared<FSimulcamEditorViewportClient>(InSimulcamViewport, bWithZoom, bWithPan);
-	Viewport = MakeShared<FSceneViewport>(ViewportClient.Get(), ViewportWidget);
+	Viewport = MakeShared<FSceneViewport>(ViewportClient, ViewportWidget);
 	ViewportClient->Viewport = Viewport.Get();
 	// The viewport widget needs an interface so it knows what should render
 	ViewportWidget->SetViewportInterface(Viewport.ToSharedRef());
diff --git a/Engine/Plugins/VirtualProduction/DMX/DMXPixelMapping/Source/DMXPixelMappingEditor/Private/Viewports/DMXPixelMappingSceneViewport.h b/Engine/Plugins/VirtualProduction/DMX/DMXPixelMapping/Source/DMXPixelMappingEditor/Private/Viewports/DMXPixelMappingSceneViewport.h
--- a/Engine/Plugins/VirtualProduction/DMX/DMXPixelMapping/Source/DMXPixelMappingEditor/Private/Viewports/DMXPixelMappingSceneViewport.h
+++ b/Engine/Plugins/VirtualProduction/DMX/DMXPixelMapping/Source/DMXPixelMappingEditor/Private/Viewports/DMXPixelMappingSceneViewport.h
@@ -11,9 +11,10 @@ class FDMXPixelMappingSceneViewport :
 	public FSceneViewport
 {
 public:
-	FDMXPixelMappingSceneViewport(FViewportClient* InViewportClient, TSharedPtr<SViewport> InViewportWidget)
+	FDMXPixelMappingSceneViewport(const TSharedPtr<FViewportClient>& InViewportClient, const TSharedPtr<SViewport>& InViewportWidget)
 		: FSceneViewport(InViewportClient, InViewportWidget)
-	{}
+	{
+	}
 
 protected:
 	virtual FReply OnMouseButtonDown(const FGeometry& InGeometry, const FPointerEvent& MouseEvent) override
diff --git a/Engine/Plugins/VirtualProduction/DMX/DMXPixelMapping/Source/DMXPixelMappingEditor/Private/Widgets/SDMXPixelMappingPreviewViewport.cpp b/Engine/Plugins/VirtualProduction/DMX/DMXPixelMapping/Source/DMXPixelMappingEditor/Private/Widgets/SDMXPixelMappingPreviewViewport.cpp
--- a/Engine/Plugins/VirtualProduction/DMX/DMXPixelMapping/Source/DMXPixelMappingEditor/Private/Widgets/SDMXPixelMappingPreviewViewport.cpp
+++ b/Engine/Plugins/VirtualProduction/DMX/DMXPixelMapping/Source/DMXPixelMappingEditor/Private/Widgets/SDMXPixelMappingPreviewViewport.cpp
@@ -32,7 +32,7 @@ void SDMXPixelMappingPreviewViewport::Construct(const FArguments& InArgs, TWeakP
 		];
 
 	ViewportClient = MakeShared<FDMXPixelMappingPreviewViewportClient>(InWeakToolkit, SharedThis(this));
-	Viewport = MakeShared<FDMXPixelMappingSceneViewport>(ViewportClient.Get(), ViewportWidget);
+	Viewport = MakeShared<FDMXPixelMappingSceneViewport>(ViewportClient, ViewportWidget);
 
 	ViewportWidget->SetViewportInterface(Viewport.ToSharedRef());
 }
diff --git a/Engine/Plugins/VirtualProduction/DMX/DMXPixelMapping/Source/DMXPixelMappingEditor/Private/Widgets/SDMXPixelMappingSourceTextureViewport.cpp b/Engine/Plugins/VirtualProduction/DMX/DMXPixelMapping/Source/DMXPixelMappingEditor/Private/Widgets/SDMXPixelMappingSourceTextureViewport.cpp
--- a/Engine/Plugins/VirtualProduction/DMX/DMXPixelMapping/Source/DMXPixelMappingEditor/Private/Widgets/SDMXPixelMappingSourceTextureViewport.cpp
+++ b/Engine/Plugins/VirtualProduction/DMX/DMXPixelMapping/Source/DMXPixelMappingEditor/Private/Widgets/SDMXPixelMappingSourceTextureViewport.cpp
@@ -35,7 +35,7 @@ void SDMXPixelMappingSourceTextureViewport::Construct(const FArguments& InArgs,
 		];
 
 	ViewportClient = MakeShared<FDMXPixelMappingSourceTextureViewportClient>(InToolkit, SharedThis(this));
-	Viewport = MakeShared<FDMXPixelMappingSceneViewport>(ViewportClient.Get(), ViewportWidget);
+	Viewport = MakeShared<FDMXPixelMappingSceneViewport>(ViewportClient, ViewportWidget);
 
 	ViewportWidget->SetViewportInterface(Viewport.ToSharedRef());
 }
diff --git a/Engine/Source/Editor/DetailCustomizations/Private/AnimMontageSegmentDetails.cpp b/Engine/Source/Editor/DetailCustomizations/Private/AnimMontageSegmentDetails.cpp
--- a/Engine/Source/Editor/DetailCustomizations/Private/AnimMontageSegmentDetails.cpp
+++ b/Engine/Source/Editor/DetailCustomizations/Private/AnimMontageSegmentDetails.cpp
@@ -638,7 +638,7 @@ void SAnimationSegmentViewport::Construct(const FArguments& InArgs)
 	LevelViewportClient->SetViewLocation( EditorViewportDefs::DefaultPerspectiveViewLocation );
 	LevelViewportClient->SetViewRotation( EditorViewportDefs::DefaultPerspectiveViewRotation );	
 
-	SceneViewport = MakeShareable( new FSceneViewport( LevelViewportClient.Get(), ViewportWidget) );
+	SceneViewport = MakeShared<FSceneViewport>(LevelViewportClient, ViewportWidget);
 	LevelViewportClient->Viewport = SceneViewport.Get();
 	LevelViewportClient->SetRealtime( true );
 	LevelViewportClient->VisibilityDelegate.BindSP( this, &SAnimationSegmentViewport::IsVisible );
diff --git a/Engine/Source/Editor/DetailCustomizations/Private/AnimSequenceDetails.cpp b/Engine/Source/Editor/DetailCustomizations/Private/AnimSequenceDetails.cpp
--- a/Engine/Source/Editor/DetailCustomizations/Private/AnimSequenceDetails.cpp
+++ b/Engine/Source/Editor/DetailCustomizations/Private/AnimSequenceDetails.cpp
@@ -642,7 +642,7 @@ void SAnimationRefPoseViewport::Construct(const FArguments& InArgs)
 	LevelViewportClient->SetViewLocation( EditorViewportDefs::DefaultPerspectiveViewLocation );
 	LevelViewportClient->SetViewRotation( EditorViewportDefs::DefaultPerspectiveViewRotation );
 
-	SceneViewport = MakeShareable( new FSceneViewport( LevelViewportClient.Get(), ViewportWidget ) );
+	SceneViewport = MakeShared<FSceneViewport>(LevelViewportClient, ViewportWidget);
 	LevelViewportClient->Viewport = SceneViewport.Get();
 	LevelViewportClient->SetRealtime( true );
 	LevelViewportClient->VisibilityDelegate.BindSP( this, &SAnimationRefPoseViewport::IsVisible );
diff --git a/Engine/Source/Editor/DistCurveEditor/Private/SCurveEditorViewport.cpp b/Engine/Source/Editor/DistCurveEditor/Private/SCurveEditorViewport.cpp
--- a/Engine/Source/Editor/DistCurveEditor/Private/SCurveEditorViewport.cpp
+++ b/Engine/Source/Editor/DistCurveEditor/Private/SCurveEditorViewport.cpp
@@ -62,7 +62,7 @@ void SCurveEditorViewport::Construct(const FArguments& InArgs)
 	
 	ViewportClient = MakeShareable(new FCurveEditorViewportClient(CurveEditorPtr, SharedThis(this)));
 
-	Viewport = MakeShareable(new FSceneViewport(ViewportClient.Get(), ViewportWidget));
+	Viewport = MakeShared<FSceneViewport>(ViewportClient, ViewportWidget);
 
 	// The viewport widget needs an interface so it knows what should render
 	ViewportWidget->SetViewportInterface( Viewport.ToSharedRef() );
diff --git a/Engine/Source/Editor/FontEditor/Private/SFontEditorViewport.cpp b/Engine/Source/Editor/FontEditor/Private/SFontEditorViewport.cpp
--- a/Engine/Source/Editor/FontEditor/Private/SFontEditorViewport.cpp
+++ b/Engine/Source/Editor/FontEditor/Private/SFontEditorViewport.cpp
@@ -835,7 +835,7 @@ void SFontEditorViewport::Construct(const FArguments& InArgs)
 
 	ViewportClient = MakeShareable(new FFontEditorViewportClient(SharedThis(this)));
 
-	Viewport = MakeShareable(new FSceneViewport(ViewportClient.Get(), ViewportWidget));
+	Viewport = MakeShared<FSceneViewport>(ViewportClient, ViewportWidget);
 
 	// The viewport widget needs an interface so it knows what should render
 	ViewportWidget->SetViewportInterface( Viewport.ToSharedRef() );
diff --git a/Engine/Source/Editor/LandscapeEditor/Private/LandscapeEditorDetailCustomization_AlphaBrush.cpp b/Engine/Source/Editor/LandscapeEditor/Private/LandscapeEditorDetailCustomization_AlphaBrush.cpp
--- a/Engine/Source/Editor/LandscapeEditor/Private/LandscapeEditorDetailCustomization_AlphaBrush.cpp
+++ b/Engine/Source/Editor/LandscapeEditor/Private/LandscapeEditorDetailCustomization_AlphaBrush.cpp
@@ -98,7 +98,7 @@ public:
 
 		ViewportClient = MakeShareable(new FTextureMaskThumbnailViewportClient(SharedThis(this)));
 
-		Viewport = MakeShareable(new FSceneViewport(ViewportClient.Get(), ViewportWidget));
+		Viewport = MakeShared<FSceneViewport>(ViewportClient, ViewportWidget);
 
 		// The viewport widget needs an interface so it knows what should render
 		ViewportWidget->SetViewportInterface( Viewport.ToSharedRef() );
diff --git a/Engine/Source/Editor/LevelEditor/Private/SLevelViewport.cpp b/Engine/Source/Editor/LevelEditor/Private/SLevelViewport.cpp
--- a/Engine/Source/Editor/LevelEditor/Private/SLevelViewport.cpp
+++ b/Engine/Source/Editor/LevelEditor/Private/SLevelViewport.cpp
@@ -4571,7 +4571,7 @@ void SLevelViewport::PreviewActors( const TArray< AActor* >& InActorsToPreview,
 			{
 				TSharedPtr<SViewport> ActorPreviewViewportWidget = ActorPreviewWidget->GetViewportWidget();
 
-				ActorPreviewSceneViewport = MakeShareable( new FSceneViewport( ActorPreviewLevelViewportClient.Get(), ActorPreviewViewportWidget) );
+				ActorPreviewSceneViewport = MakeShared<FSceneViewport>(ActorPreviewLevelViewportClient, ActorPreviewViewportWidget);
 				{
 					ActorPreviewLevelViewportClient->Viewport = ActorPreviewSceneViewport.Get();
 					if (ensure(ActorPreviewViewportWidget.IsValid()))
diff --git a/Engine/Source/Editor/Persona/Private/SAnimationSequenceBrowser.cpp b/Engine/Source/Editor/Persona/Private/SAnimationSequenceBrowser.cpp
--- a/Engine/Source/Editor/Persona/Private/SAnimationSequenceBrowser.cpp
+++ b/Engine/Source/Editor/Persona/Private/SAnimationSequenceBrowser.cpp
@@ -1726,7 +1726,7 @@ void SAnimationSequenceBrowser::CreateAssetTooltipResources()
 		.ViewportSize(FVector2D(128, 128));
 
 	ViewportClient = MakeShareable(new FAnimationAssetViewportClient(PreviewScene));
-	SceneViewport = MakeShareable(new FSceneViewport(ViewportClient.Get(), ViewportWidget));
+	SceneViewport = MakeShared<FSceneViewport>(ViewportClient, ViewportWidget);
 	PreviewComponent = NewObject<UDebugSkelMeshComponent>();
 
 	// Client options
diff --git a/Engine/Source/Editor/TextureEditor/Private/Widgets/STextureEditorViewport.cpp b/Engine/Source/Editor/TextureEditor/Private/Widgets/STextureEditorViewport.cpp
--- a/Engine/Source/Editor/TextureEditor/Private/Widgets/STextureEditorViewport.cpp
+++ b/Engine/Source/Editor/TextureEditor/Private/Widgets/STextureEditorViewport.cpp
@@ -65,7 +65,7 @@ void STextureEditorViewport::Construct( const FArguments& InArgs, const TSharedR
 	
 	ViewportClient = MakeShareable(new FTextureEditorViewportClient(ToolkitPtr, SharedThis(this)));
 
-	Viewport = MakeShareable(new FSceneViewport(ViewportClient.Get(), ViewportWidget));
+	Viewport = MakeShared<FSceneViewport>(ViewportClient, ViewportWidget);
 
 	// The viewport widget needs an interface so it knows what should render
 	ViewportWidget->SetViewportInterface( Viewport.ToSharedRef() );
diff --git a/Engine/Source/Editor/UnrealEd/Private/SEditorViewport.cpp b/Engine/Source/Editor/UnrealEd/Private/SEditorViewport.cpp
--- a/Engine/Source/Editor/UnrealEd/Private/SEditorViewport.cpp
+++ b/Engine/Source/Editor/UnrealEd/Private/SEditorViewport.cpp
@@ -100,7 +100,7 @@ void SEditorViewport::Construct( const FArguments& InArgs )
 		Client->VisibilityDelegate.BindSP(this, &SEditorViewport::IsVisible);
 	}
 
-	SceneViewport = MakeShareable( new FSceneViewport( Client.Get(), ViewportWidget ) );
+	SceneViewport = MakeShared<FSceneViewport>(Client, ViewportWidget);
 	Client->Viewport = SceneViewport.Get();
 	ViewportWidget->SetViewportInterface(SceneViewport.ToSharedRef());
 
diff --git a/Engine/Source/Runtime/Engine/Private/Slate/SceneViewport.cpp b/Engine/Source/Runtime/Engine/Private/Slate/SceneViewport.cpp
--- a/Engine/Source/Runtime/Engine/Private/Slate/SceneViewport.cpp
+++ b/Engine/Source/Runtime/Engine/Private/Slate/SceneViewport.cpp
@@ -33,8 +33,28 @@ DEFINE_LOG_CATEGORY(LogViewport);
 
 extern EWindowMode::Type GetWindowModeType(EWindowMode::Type WindowMode);
 
-FSceneViewport::FSceneViewport( FViewportClient* InViewportClient, TSharedPtr<SViewport> InViewportWidget )
-	: FViewport( InViewportClient )
+FSceneViewport::FSceneViewport(const TStrongPtrVariant<FViewportClient>& InClient, const TSharedPtr<SViewport>& Widget)
+	: FSceneViewport(InClient, InClient.Get(), Widget)
+{
+}
+
+FSceneViewport::FSceneViewport(const TSharedPtr<FViewportClient>& Client, const TSharedPtr<SViewport>& Widget)
+	: FSceneViewport(Client, Client.Get(), Widget)
+{
+}
+
+FSceneViewport::FSceneViewport(FViewportClient* const InClient, const TSharedPtr<SViewport>& Widget)
+	: FSceneViewport(nullptr, InClient, Widget)
+{
+}
+
+FSceneViewport::FSceneViewport(UScriptViewportClient* const InClient, const TSharedPtr<SViewport>& Widget)
+	: FSceneViewport(InClient, InClient, Widget)
+{
+}
+
+FSceneViewport::FSceneViewport(const TStrongPtrVariant<FViewportClient>& InClient, FViewportClient* const RawClient, const TSharedPtr<SViewport>& InViewportWidget)
+	: FViewport(InClient, RawClient)
 	, CurrentReplyState( FReply::Unhandled() )
 	, CachedCursorPos(-1, -1)
 	, PreCaptureCursorPos(-1, -1)
@@ -71,9 +91,11 @@ FSceneViewport::FSceneViewport( FViewportClient* InViewportClient, TSharedPtr<SV
 	bIsSlateViewport = true;
 	RenderThreadSlateTexture = new FSlateRenderTargetRHI(nullptr, 0, 0);
 
-	if (InViewportClient)
+PRAGMA_DISABLE_DEPRECATION_WARNINGS
+	if (FViewportClient* ClientPtr = ViewportClient)
+PRAGMA_ENABLE_DEPRECATION_WARNINGS
 	{
-		bShouldCaptureMouseOnActivate = InViewportClient->CaptureMouseOnLaunch();
+		bShouldCaptureMouseOnActivate = ClientPtr->CaptureMouseOnLaunch();
 	}
 }
 
@@ -178,8 +200,11 @@ bool FSceneViewport::KeyState( FKey Key ) const
 
 void FSceneViewport::Destroy()
 {
-	ViewportClient = NULL;
-	
+	Client.Reset();
+PRAGMA_DISABLE_DEPRECATION_WARNINGS
+	ViewportClient = nullptr;
+PRAGMA_ENABLE_DEPRECATION_WARNINGS
+
 	UpdateViewportRHI( true, 0, 0, EWindowMode::Windowed, PF_Unknown );	
 }
 
@@ -289,53 +314,59 @@ void FSceneViewport::UpdateModifierKeys( const FPointerEvent& InMouseEvent )
 
 void FSceneViewport::ApplyModifierKeys(const FModifierKeysState& InKeysState, const uint64 Timestamp)
 {
-	if (ViewportClient && GetSizeXY() != FIntPoint::ZeroValue)
+PRAGMA_DISABLE_DEPRECATION_WARNINGS
+	FViewportClient* ClientPtr = ViewportClient;
+PRAGMA_ENABLE_DEPRECATION_WARNINGS
+	if (ClientPtr && GetSizeXY() != FIntPoint::ZeroValue)
 	{
 		const FInputDeviceId DefaultInputDevice = IPlatformInputDeviceMapper::Get().GetDefaultInputDevice();
-		
+
 		// Switch to the viewport clients world before processing input
-		FScopedConditionalWorldSwitcher WorldSwitcher(ViewportClient);
+		FScopedConditionalWorldSwitcher WorldSwitcher(ClientPtr);
 
 		if (InKeysState.IsLeftAltDown())
 		{
-			ViewportClient->InputKey(FInputKeyEventArgs(this, DefaultInputDevice, EKeys::LeftAlt, IE_Pressed, Timestamp));
+			ClientPtr->InputKey(FInputKeyEventArgs(this, DefaultInputDevice, EKeys::LeftAlt, IE_Pressed, Timestamp));
 		}
 		if (InKeysState.IsRightAltDown())
 		{
-			ViewportClient->InputKey(FInputKeyEventArgs(this, DefaultInputDevice, EKeys::RightAlt, IE_Pressed, Timestamp));
+			ClientPtr->InputKey(FInputKeyEventArgs(this, DefaultInputDevice, EKeys::RightAlt, IE_Pressed, Timestamp));
 		}
 		if (InKeysState.IsLeftControlDown())
 		{
-			ViewportClient->InputKey(FInputKeyEventArgs(this, DefaultInputDevice, EKeys::LeftControl, IE_Pressed, Timestamp));
+			ClientPtr->InputKey(FInputKeyEventArgs(this, DefaultInputDevice, EKeys::LeftControl, IE_Pressed, Timestamp));
 		}
 		if (InKeysState.IsRightControlDown())
 		{
-			ViewportClient->InputKey(FInputKeyEventArgs(this, DefaultInputDevice, EKeys::RightControl, IE_Pressed, Timestamp));
+			ClientPtr->InputKey(FInputKeyEventArgs(this, DefaultInputDevice, EKeys::RightControl, IE_Pressed, Timestamp));
 		}
 		if (InKeysState.IsLeftShiftDown())
 		{
-			ViewportClient->InputKey(FInputKeyEventArgs(this, DefaultInputDevice, EKeys::LeftShift, IE_Pressed, Timestamp));
+			ClientPtr->InputKey(FInputKeyEventArgs(this, DefaultInputDevice, EKeys::LeftShift, IE_Pressed, Timestamp));
 		}
 		if (InKeysState.IsRightShiftDown())
 		{
-			ViewportClient->InputKey(FInputKeyEventArgs(this, DefaultInputDevice, EKeys::RightShift, IE_Pressed, Timestamp));
+			ClientPtr->InputKey(FInputKeyEventArgs(this, DefaultInputDevice, EKeys::RightShift, IE_Pressed, Timestamp));
 		}
 	}
 }
 
 void FSceneViewport::ProcessAccumulatedPointerInput()
 {
-	if( !ViewportClient )
+PRAGMA_DISABLE_DEPRECATION_WARNINGS
+	FViewportClient* ClientPtr = ViewportClient;
+PRAGMA_ENABLE_DEPRECATION_WARNINGS
+	if (!ClientPtr)
 	{
 		return;
 	}
 
 	// Switch to the viewport clients world before processing input
-	FScopedConditionalWorldSwitcher WorldSwitcher( ViewportClient );
+	FScopedConditionalWorldSwitcher WorldSwitcher(ClientPtr);
 
 	const bool bViewportHasCapture = ViewportWidget.IsValid() && ViewportWidget.Pin()->HasMouseCapture();
 
-	ViewportClient->ProcessAccumulatedPointerInput(this);
+	ClientPtr->ProcessAccumulatedPointerInput(this);
 
 	if (NumMouseSamplesX > 0 || NumMouseSamplesY > 0)
 	{
@@ -345,24 +376,24 @@ void FSceneViewport::ProcessAccumulatedPointerInput()
 		const uint64 MoseRecentMouseTimestamp = 0u;
 		
 		FInputKeyEventArgs MouseXArgs (this, DefaultInputDevice, EKeys::MouseX, MouseDelta.X, DeltaTime, NumMouseSamplesX, MoseRecentMouseTimestamp);
-		ViewportClient->InputAxis(MouseXArgs);
+		ClientPtr->InputAxis(MouseXArgs);
 
 		FInputKeyEventArgs MouseYArgs (this, DefaultInputDevice, EKeys::MouseY, MouseDelta.Y, DeltaTime, NumMouseSamplesY, MoseRecentMouseTimestamp);
-		ViewportClient->InputAxis(MouseYArgs);
+		ClientPtr->InputAxis(MouseYArgs);
 	}
 
 	if ( bCursorHiddenDueToCapture )
 	{
-		switch ( ViewportClient->GetMouseCaptureMode() )
+		switch (ClientPtr->GetMouseCaptureMode())
 		{
 		case EMouseCaptureMode::NoCapture:
 		case EMouseCaptureMode::CaptureDuringMouseDown:
 		case EMouseCaptureMode::CaptureDuringRightMouseDown:
 			if ( !bViewportHasCapture )
 			{
-				bool bShouldMouseBeVisible = ViewportClient->GetCursor(this, GetMouseX(), GetMouseY()) != EMouseCursor::None;
+				bool bShouldMouseBeVisible = ClientPtr->GetCursor(this, GetMouseX(), GetMouseY()) != EMouseCursor::None;
 
-				UWorld* World = ViewportClient->GetWorld();
+				UWorld* World = ClientPtr->GetWorld();
 				if ( World && World->IsGameWorld() && World->GetGameInstance() )
 				{
 					APlayerController* PC = World->GetGameInstance()->GetFirstLocalPlayerController();
@@ -438,7 +469,9 @@ static void ComputeSceneViewportHDRMetaData(FHDRMetaData& OutHDRMetaData, const
 void FSceneViewport::OnDrawViewport( const FGeometry& AllottedGeometry, const FSlateRect& MyCullingRect, FSlateWindowElementList& OutDrawElements, int32 LayerId, const FWidgetStyle& InWidgetStyle, bool bParentEnabled )
 {
 	// Switch to the viewport clients world before resizing
-	FScopedConditionalWorldSwitcher WorldSwitcher( ViewportClient );
+PRAGMA_DISABLE_DEPRECATION_WARNINGS
+	FScopedConditionalWorldSwitcher WorldSwitcher(ViewportClient);
+PRAGMA_ENABLE_DEPRECATION_WARNINGS
 
 	/** Check to see if the viewport should be resized */
 	if (!bForceViewportSize)
@@ -539,11 +572,14 @@ FCursorReply FSceneViewport::OnCursorQuery( const FGeometry& MyGeometry, const F
 
 	// If the cursor should be hidden, use EMouseCursor::None,
 	// only when in the foreground, or we'll hide the mouse in the window/program above us.
-	if( ViewportClient && GetSizeXY() != FIntPoint::ZeroValue  )
+PRAGMA_DISABLE_DEPRECATION_WARNINGS
+	FViewportClient* ClientPtr = ViewportClient;
+PRAGMA_ENABLE_DEPRECATION_WARNINGS
+	if (ClientPtr && GetSizeXY() != FIntPoint::ZeroValue)
 	{
 		const int32 MouseX = GetMouseX();
 		const int32 MouseY = GetMouseY();
-		MouseCursorToUse = ViewportClient->GetCursor(this, MouseX, MouseY);
+		MouseCursorToUse = ClientPtr->GetCursor(this, MouseX, MouseY);
 	}
 
 	// In game mode we may be using a borderless window, which needs OnCursorQuery call to handle window resize cursors
@@ -566,9 +602,12 @@ FCursorReply FSceneViewport::OnCursorQuery( const FGeometry& MyGeometry, const F
 
 TOptional<TSharedRef<SWidget>> FSceneViewport::OnMapCursor(const FCursorReply& CursorReply)
 {
-	if (ViewportClient && GetSizeXY() != FIntPoint::ZeroValue)
+PRAGMA_DISABLE_DEPRECATION_WARNINGS
+	FViewportClient* ClientPtr = ViewportClient;
+PRAGMA_ENABLE_DEPRECATION_WARNINGS
+	if (ClientPtr && GetSizeXY() != FIntPoint::ZeroValue)
 	{
-		return ViewportClient->MapCursor(this, CursorReply);
+		return ClientPtr->MapCursor(this, CursorReply);
 	}
 	return ISlateViewport::OnMapCursor(CursorReply);
 }
@@ -596,8 +635,11 @@ FReply FSceneViewport::OnMouseButtonDown( const FGeometry& InGeometry, const FPo
 	}
 
 	// Switch to the viewport clients world before processing input
-	FScopedConditionalWorldSwitcher WorldSwitcher(ViewportClient);
-	if (ViewportClient && GetSizeXY() != FIntPoint::ZeroValue)
+PRAGMA_DISABLE_DEPRECATION_WARNINGS
+	FViewportClient* ClientPtr = ViewportClient;
+PRAGMA_ENABLE_DEPRECATION_WARNINGS
+	FScopedConditionalWorldSwitcher WorldSwitcher(ClientPtr);
+	if (ClientPtr && GetSizeXY() != FIntPoint::ZeroValue)
 	{
 		// If we're obtaining focus, we have to copy the modifier key states prior to processing this mouse button event, as this is the only point at which the mouse down
 		// event is processed when focus initially changes and the modifier keys need to be in-place to detect any unique drag-like events.
@@ -607,19 +649,20 @@ FReply FSceneViewport::OnMouseButtonDown( const FGeometry& InGeometry, const FPo
 			ApplyModifierKeys(KeysState, InMouseEvent.GetEventTimestamp());
 		}
 
-		const bool bTemporaryCapture = 
-			ViewportClient->GetMouseCaptureMode() == EMouseCaptureMode::CaptureDuringMouseDown ||
-			(ViewportClient->GetMouseCaptureMode() == EMouseCaptureMode::CaptureDuringRightMouseDown && InMouseEvent.GetEffectingButton() == EKeys::RightMouseButton);
+		const bool bTemporaryCapture =
+			ClientPtr->GetMouseCaptureMode() == EMouseCaptureMode::CaptureDuringMouseDown ||
+			(ClientPtr->GetMouseCaptureMode() == EMouseCaptureMode::CaptureDuringRightMouseDown && InMouseEvent.GetEffectingButton() == EKeys::RightMouseButton);
 
 		// Process primary input if we aren't currently a game viewport, we already have capture, or we are permanent capture that doesn't consume the mouse down.
-		const bool bProcessInputPrimary = !IsCurrentlyGameViewport() || HasMouseCapture() || (ViewportClient->GetMouseCaptureMode() == EMouseCaptureMode::CapturePermanently_IncludingInitialMouseDown);
+		const bool bProcessInputPrimary = !IsCurrentlyGameViewport() || HasMouseCapture()
+			|| ClientPtr->GetMouseCaptureMode() == EMouseCaptureMode::CapturePermanently_IncludingInitialMouseDown;
 
 		const bool bAnyMenuWasVisible = FSlateApplication::Get().AnyMenusVisible();
 
 		// Process the mouse event
 		if (bTemporaryCapture || bProcessInputPrimary)
 		{
-			if (!ViewportClient->InputKey(FInputKeyEventArgs(this, InMouseEvent.GetInputDeviceId(), InMouseEvent.GetEffectingButton(), IE_Pressed, 1.0f, InMouseEvent.IsTouchEvent(), InMouseEvent.GetEventTimestamp())))
+			if (!ClientPtr->InputKey(FInputKeyEventArgs(this, InMouseEvent.GetInputDeviceId(), InMouseEvent.GetEffectingButton(), IE_Pressed, 1.0f, InMouseEvent.IsTouchEvent(), InMouseEvent.GetEventTimestamp())))
 			{
 				CurrentReplyState = FReply::Unhandled();
 			}
@@ -629,10 +672,10 @@ FReply FSceneViewport::OnMouseButtonDown( const FGeometry& InGeometry, const FPo
 		const bool bNewMenuWasOpened = !bAnyMenuWasVisible && FSlateApplication::Get().AnyMenusVisible();
 
 		const bool bPermanentCapture =
-			(ViewportClient->GetMouseCaptureMode() == EMouseCaptureMode::CapturePermanently) ||
-			(ViewportClient->GetMouseCaptureMode() == EMouseCaptureMode::CapturePermanently_IncludingInitialMouseDown);
+			ClientPtr->GetMouseCaptureMode() == EMouseCaptureMode::CapturePermanently ||
+			ClientPtr->GetMouseCaptureMode() == EMouseCaptureMode::CapturePermanently_IncludingInitialMouseDown;
 
-		if (FSlateApplication::Get().IsActive() && !ViewportClient->IgnoreInput() &&
+		if (FSlateApplication::Get().IsActive() && !ClientPtr->IgnoreInput() &&
 			!bNewMenuWasOpened && // We should not focus the viewport if a menu was opened as it would close the menu
 			(bPermanentCapture || bTemporaryCapture))
 		{
@@ -656,12 +699,13 @@ FReply FSceneViewport::AcquireFocusAndCapture(FIntPoint MousePosition, EFocusCau
 	// Mouse down should focus viewport for user input
 	ReplyState.SetUserFocus(ViewportWidgetRef, FocusCause);
 
-	UWorld* World = ViewportClient->GetWorld();
+	FViewportClient& ViewportClientRef = GetClientChecked();
+	UWorld* World = ViewportClientRef.GetWorld();
 	if (World && World->IsGameWorld() && World->GetGameInstance() && (World->GetGameInstance()->GetFirstLocalPlayerController() || World->IsPlayInEditor()))
 	{
 		ReplyState.CaptureMouse(ViewportWidgetRef);
 
-		if (ViewportClient->LockDuringCapture())
+		if (ViewportClientRef.LockDuringCapture())
 		{
 			ReplyState.LockMouseToWidget(ViewportWidgetRef);
 		}
@@ -669,7 +713,7 @@ FReply FSceneViewport::AcquireFocusAndCapture(FIntPoint MousePosition, EFocusCau
 		APlayerController* PC = World->GetGameInstance()->GetFirstLocalPlayerController();
 		const bool bShouldShowMouseCursor = PC && PC->ShouldShowMouseCursor();
 
-		if ( ViewportClient->HideCursorDuringCapture() && bShouldShowMouseCursor )
+		if (ViewportClientRef.HideCursorDuringCapture() && bShouldShowMouseCursor)
 		{
 			bCursorHiddenDueToCapture = true;
 			MousePosBeforeHiddenDueToCapture = MousePosition;
@@ -709,7 +753,7 @@ FReply FSceneViewport::AcquireFocusAndCapture(FIntPoint MousePosition, EFocusCau
 bool FSceneViewport::IsCurrentlyGameViewport()
 {
 	// Either were game code only or were are currently play in editor.
-	return (FApp::IsGame() && !GIsEditor && ViewportClient == GEngine->GameViewport) || IsPlayInEditorViewport();
+	return (FApp::IsGame() && !GIsEditor && GetClient() == GEngine->GameViewport) || IsPlayInEditorViewport();
 }
 
 FReply FSceneViewport::OnMouseButtonUp( const FGeometry& InGeometry, const FPointerEvent& InMouseEvent )
@@ -725,30 +769,33 @@ FReply FSceneViewport::OnMouseButtonUp( const FGeometry& InGeometry, const FPoin
 	
 
 	// Switch to the viewport clients world before processing input
-	FScopedConditionalWorldSwitcher WorldSwitcher( ViewportClient );
+PRAGMA_DISABLE_DEPRECATION_WARNINGS
+	FViewportClient* ClientPtr = ViewportClient;
+PRAGMA_ENABLE_DEPRECATION_WARNINGS
+	FScopedConditionalWorldSwitcher WorldSwitcher(ClientPtr);
 	bool bCursorVisible = true;
 	bool bReleaseMouseCapture = true;
-	
-	if( ViewportClient && GetSizeXY() != FIntPoint::ZeroValue )
+
+	if (ClientPtr && GetSizeXY() != FIntPoint::ZeroValue )
 	{
-		if (!ViewportClient->InputKey(FInputKeyEventArgs(this, InMouseEvent.GetInputDeviceId(), InMouseEvent.GetEffectingButton(), IE_Released, 1.0f, InMouseEvent.IsTouchEvent(), InMouseEvent.GetEventTimestamp())))
+		if (!ClientPtr->InputKey(FInputKeyEventArgs(this, InMouseEvent.GetInputDeviceId(), InMouseEvent.GetEffectingButton(), IE_Released, 1.0f, InMouseEvent.IsTouchEvent(), InMouseEvent.GetEventTimestamp())))
 		{
-			CurrentReplyState = FReply::Unhandled(); 
+			CurrentReplyState = FReply::Unhandled();
 		}
 
-		bCursorVisible = ViewportClient->GetCursor(this, GetMouseX(), GetMouseY()) != EMouseCursor::None;
+		bCursorVisible = ClientPtr->GetCursor(this, GetMouseX(), GetMouseY()) != EMouseCursor::None;
 
 		if (bCursorVisible)
 		{
 			bReleaseMouseCapture = true;
 			UE_LOGF(LogViewport, Verbose, "Releasing Mouse Capture; Cursor is visible");
 		}
-		else if (ViewportClient->GetMouseCaptureMode() == EMouseCaptureMode::CaptureDuringMouseDown)
+		else if (ClientPtr->GetMouseCaptureMode() == EMouseCaptureMode::CaptureDuringMouseDown)
 		{
 			bReleaseMouseCapture = true;
 			UE_LOGF(LogViewport, Verbose, "Releasing Mouse Capture; EMouseCaptureMode::CaptureDuringMouseDown - Mouse Button Released");
 		}
-		else if (ViewportClient->GetMouseCaptureMode() == EMouseCaptureMode::CaptureDuringRightMouseDown && InMouseEvent.GetEffectingButton() == EKeys::RightMouseButton)
+		else if (ClientPtr->GetMouseCaptureMode() == EMouseCaptureMode::CaptureDuringRightMouseDown && InMouseEvent.GetEffectingButton() == EKeys::RightMouseButton)
 		{
 			bReleaseMouseCapture = true;
 			UE_LOGF(LogViewport, Verbose, "Releasing Mouse Capture; EMouseCaptureMode::CaptureDuringRightMouseDown - Right Mouse Button Released");
@@ -774,7 +821,7 @@ FReply FSceneViewport::OnMouseButtonUp( const FGeometry& InGeometry, const FPoin
 
 			CurrentReplyState.ReleaseMouseCapture();
 
-			if ( bCursorVisible && !ViewportClient->ShouldAlwaysLockMouse() )
+			if (bCursorVisible && !ClientPtr->ShouldAlwaysLockMouse())
 			{
 				CurrentReplyState.ReleaseMouseLock();
 			}
@@ -787,15 +834,17 @@ FReply FSceneViewport::OnMouseButtonUp( const FGeometry& InGeometry, const FPoin
 void FSceneViewport::OnMouseEnter( const FGeometry& MyGeometry, const FPointerEvent& MouseEvent )
 {
 	UpdateCachedCursorPos( MyGeometry, MouseEvent );
-	ViewportClient->MouseEnter( this, GetMouseX(), GetMouseY() );
+	GetClientChecked().MouseEnter(this, GetMouseX(), GetMouseY());
 }
 
 void FSceneViewport::OnMouseLeave( const FPointerEvent& MouseEvent )
 {
-	if( ViewportClient )
+PRAGMA_DISABLE_DEPRECATION_WARNINGS
+	if (FViewportClient* ClientPtr = ViewportClient)
+PRAGMA_ENABLE_DEPRECATION_WARNINGS
 	{
-		ViewportClient->MouseLeave( this );
-	
+		ClientPtr->MouseLeave(this);
+
 		if (IsCurrentlyGameViewport())
 		{
 			CachedCursorPos = FIntPoint(-1, -1);
@@ -814,18 +863,21 @@ FReply FSceneViewport::OnMouseMove( const FGeometry& InGeometry, const FPointerE
 	UpdateCachedCursorPos(InGeometry, InMouseEvent, bIsMoveEvent);
 
 	const bool bViewportHasCapture = ViewportWidget.IsValid() && ViewportWidget.Pin()->HasMouseCapture();
-	if ( ViewportClient && GetSizeXY() != FIntPoint::ZeroValue )
+PRAGMA_DISABLE_DEPRECATION_WARNINGS
+	FViewportClient* ClientPtr = ViewportClient;
+PRAGMA_ENABLE_DEPRECATION_WARNINGS
+	if (ClientPtr && GetSizeXY() != FIntPoint::ZeroValue)
 	{
 		// Switch to the viewport clients world before processing input
-		FScopedConditionalWorldSwitcher WorldSwitcher(ViewportClient);
+		FScopedConditionalWorldSwitcher WorldSwitcher(ClientPtr);
 
 		if ( bViewportHasCapture )
 		{
-			ViewportClient->CapturedMouseMove(this, GetMouseX(), GetMouseY());
+			ClientPtr->CapturedMouseMove(this, GetMouseX(), GetMouseY());
 		}
 		else
 		{
-			ViewportClient->MouseMove(this, GetMouseX(), GetMouseY());
+			ClientPtr->MouseMove(this, GetMouseX(), GetMouseY());
 		}
 
 		if ( bViewportHasCapture )
@@ -859,18 +911,21 @@ FReply FSceneViewport::OnMouseWheel( const FGeometry& InGeometry, const FPointer
 	UpdateCachedGeometry(InGeometry);
 	UpdateCachedCursorPos( InGeometry, InMouseEvent );
 
-	if( ViewportClient  && GetSizeXY() != FIntPoint::ZeroValue  )
+PRAGMA_DISABLE_DEPRECATION_WARNINGS
+	FViewportClient* ClientPtr = ViewportClient;
+PRAGMA_ENABLE_DEPRECATION_WARNINGS
+	if (ClientPtr  && GetSizeXY() != FIntPoint::ZeroValue )
 	{
 		// Switch to the viewport clients world before processing input
-		FScopedConditionalWorldSwitcher WorldSwitcher( ViewportClient );
+		FScopedConditionalWorldSwitcher WorldSwitcher(ClientPtr);
 
-		// The viewport client accepts two different keys depending on the direction of scroll.  
+		// The viewport client accepts two different keys depending on the direction of scroll.
 		FKey const ViewportClientKey = InMouseEvent.GetWheelDelta() < 0 ? EKeys::MouseScrollDown : EKeys::MouseScrollUp;
 
 		// Pressed and released should be sent
-		ViewportClient->InputKey(FInputKeyEventArgs(this, InMouseEvent.GetInputDeviceId(), ViewportClientKey, IE_Pressed, 1.0f, InMouseEvent.IsTouchEvent(), InMouseEvent.GetEventTimestamp()));
-		ViewportClient->InputKey(FInputKeyEventArgs(this, InMouseEvent.GetInputDeviceId(), ViewportClientKey, IE_Released, 1.0f, InMouseEvent.IsTouchEvent(), InMouseEvent.GetEventTimestamp()));
-		ViewportClient->InputAxis(FInputKeyEventArgs(this, InMouseEvent.GetInputDeviceId(), EKeys::MouseWheelAxis, InMouseEvent.GetWheelDelta(), FApp::GetDeltaTime(), 1, InMouseEvent.GetEventTimestamp()));
+		ClientPtr->InputKey(FInputKeyEventArgs(this, InMouseEvent.GetInputDeviceId(), ViewportClientKey, IE_Pressed, 1.0f, InMouseEvent.IsTouchEvent(), InMouseEvent.GetEventTimestamp()));
+		ClientPtr->InputKey(FInputKeyEventArgs(this, InMouseEvent.GetInputDeviceId(), ViewportClientKey, IE_Released, 1.0f, InMouseEvent.IsTouchEvent(), InMouseEvent.GetEventTimestamp()));
+		ClientPtr->InputAxis(FInputKeyEventArgs(this, InMouseEvent.GetInputDeviceId(), EKeys::MouseWheelAxis, InMouseEvent.GetWheelDelta(), FApp::GetDeltaTime(), 1, InMouseEvent.GetEventTimestamp()));
 	}
 	return CurrentReplyState;
 }
@@ -890,14 +945,17 @@ FReply FSceneViewport::OnMouseButtonDoubleClick( const FGeometry& InGeometry, co
 	UpdateCachedGeometry(InGeometry);
 	UpdateCachedCursorPos(InGeometry, InMouseEvent);
 
-	if( ViewportClient && GetSizeXY() != FIntPoint::ZeroValue  )
+PRAGMA_DISABLE_DEPRECATION_WARNINGS
+	FViewportClient* ClientPtr = ViewportClient;
+PRAGMA_ENABLE_DEPRECATION_WARNINGS
+	if (ClientPtr && GetSizeXY() != FIntPoint::ZeroValue)
 	{
 		// Switch to the viewport clients world before processing input
-		FScopedConditionalWorldSwitcher WorldSwitcher( ViewportClient );
+		FScopedConditionalWorldSwitcher WorldSwitcher(ClientPtr);
 
-		if (!ViewportClient->InputKey(FInputKeyEventArgs(this, InMouseEvent.GetInputDeviceId(), InMouseEvent.GetEffectingButton(), IE_DoubleClick, 1.0f, InMouseEvent.IsTouchEvent(), InMouseEvent.GetEventTimestamp())))
+		if (!ClientPtr->InputKey(FInputKeyEventArgs(this, InMouseEvent.GetInputDeviceId(), InMouseEvent.GetEffectingButton(), IE_DoubleClick, 1.0f, InMouseEvent.IsTouchEvent(), InMouseEvent.GetEventTimestamp())))
 		{
-			CurrentReplyState = FReply::Unhandled(); 
+			CurrentReplyState = FReply::Unhandled();
 		}
 	}
 	return CurrentReplyState;
@@ -911,21 +969,22 @@ FReply FSceneViewport::OnTouchStarted( const FGeometry& MyGeometry, const FPoint
 
 	UpdateCachedGeometry(MyGeometry);
 	UpdateCachedCursorPos(MyGeometry, TouchEvent);
-	
-	if( ViewportClient )
+
+PRAGMA_DISABLE_DEPRECATION_WARNINGS
+	if (FViewportClient* ClientPtr = ViewportClient)
+PRAGMA_ENABLE_DEPRECATION_WARNINGS
 	{
 		// Switch to the viewport clients world before processing input
-		FScopedConditionalWorldSwitcher WorldSwitcher( ViewportClient );
+		FScopedConditionalWorldSwitcher WorldSwitcher(ClientPtr);
 
 		const FVector2D TouchPosition = CachedCursorPos;
 
-		const bool bEventConsumed = ViewportClient->InputTouch(this,
+		const bool bEventConsumed = ClientPtr->InputTouch(this,
 			FTouchId(TouchEvent.GetInputDeviceId(), static_cast<ETouchIndex::Type>(TouchEvent.GetPointerIndex())), ETouchType::Began, TouchPosition,
 			TouchEvent.GetTouchForce(), TouchEvent.GetEventTimestamp());
 		if (bEventConsumed)
 		{
-			const bool bTemporaryCapture = ViewportClient->GetMouseCaptureMode() == EMouseCaptureMode::CaptureDuringMouseDown;
-			if (bTemporaryCapture)
+			if (ClientPtr->GetMouseCaptureMode() == EMouseCaptureMode::CaptureDuringMouseDown)
 			{
 				CurrentReplyState = AcquireFocusAndCapture(FIntPoint(TouchEvent.GetScreenSpacePosition().X, TouchEvent.GetScreenSpacePosition().Y), EFocusCause::Mouse);
 			}
@@ -948,13 +1007,14 @@ FReply FSceneViewport::OnTouchMoved( const FGeometry& MyGeometry, const FPointer
 	
 	constexpr bool bIsMoveEvent = true;
 	UpdateCachedCursorPos(MyGeometry, TouchEvent, bIsMoveEvent);
-
-	if( ViewportClient )
+PRAGMA_DISABLE_DEPRECATION_WARNINGS
+	if (FViewportClient* ClientPtr = ViewportClient)
+PRAGMA_ENABLE_DEPRECATION_WARNINGS
 	{
 		// Switch to the viewport clients world before processing input
-		FScopedConditionalWorldSwitcher WorldSwitcher( ViewportClient );
+		FScopedConditionalWorldSwitcher WorldSwitcher(ClientPtr);
 
-		const bool bEventConsumed = ViewportClient->InputTouch(this,
+		const bool bEventConsumed = ClientPtr->InputTouch(this,
 			FTouchId(TouchEvent.GetInputDeviceId(), static_cast<ETouchIndex::Type>(TouchEvent.GetPointerIndex())), ETouchType::Moved, CachedCursorPos,
 			TouchEvent.GetTouchForce(), TouchEvent.GetEventTimestamp());
 		if (!bEventConsumed)
@@ -986,21 +1046,22 @@ FReply FSceneViewport::OnTouchEnded( const FGeometry& MyGeometry, const FPointer
 		CachedCursorPos = FIntPoint(-1, -1);
 	}
 
-
-	if( ViewportClient )
+PRAGMA_DISABLE_DEPRECATION_WARNINGS
+	if (FViewportClient* ClientPtr = ViewportClient)
+PRAGMA_ENABLE_DEPRECATION_WARNINGS
 	{
 		// Switch to the viewport clients world before processing input
-		FScopedConditionalWorldSwitcher WorldSwitcher( ViewportClient );
+		FScopedConditionalWorldSwitcher WorldSwitcher(ClientPtr);
 
-		const bool bEventConsumed = ViewportClient->InputTouch(this,
+		const bool bEventConsumed = ClientPtr->InputTouch(this,
 			FTouchId(TouchEvent.GetInputDeviceId(), static_cast<ETouchIndex::Type>(TouchEvent.GetPointerIndex())), ETouchType::Ended, CurCursorPos,
 			/* Force */ 0.0f, TouchEvent.GetEventTimestamp());
 		if (!bEventConsumed)
 		{
-			CurrentReplyState = FReply::Unhandled(); 
+			CurrentReplyState = FReply::Unhandled();
 		}
 
-		if (ViewportClient->GetMouseCaptureMode() == EMouseCaptureMode::CaptureDuringMouseDown)
+		if (ClientPtr->GetMouseCaptureMode() == EMouseCaptureMode::CaptureDuringMouseDown)
 		{
 			CurrentReplyState.ReleaseMouseCapture();
 		}
@@ -1017,14 +1078,16 @@ FReply FSceneViewport::OnTouchForceChanged(const FGeometry& MyGeometry, const FP
 	UpdateCachedCursorPos(MyGeometry, TouchEvent);
 	UpdateCachedGeometry(MyGeometry);
 
-	if (ViewportClient)
+PRAGMA_DISABLE_DEPRECATION_WARNINGS
+	if (FViewportClient* ClientPtr = ViewportClient)
+PRAGMA_ENABLE_DEPRECATION_WARNINGS
 	{
 		// Switch to the viewport clients world before processing input
-		FScopedConditionalWorldSwitcher WorldSwitcher(ViewportClient);
+		FScopedConditionalWorldSwitcher WorldSwitcher(ClientPtr);
 
 		const FVector2D TouchPosition = MyGeometry.AbsoluteToLocal(TouchEvent.GetScreenSpacePosition()) * MyGeometry.Scale;
 
-		const bool bEventConsumed = ViewportClient->InputTouch(this,
+		const bool bEventConsumed = ClientPtr->InputTouch(this,
 			FTouchId(TouchEvent.GetInputDeviceId(), static_cast<ETouchIndex::Type>(TouchEvent.GetPointerIndex())), ETouchType::ForceChanged,
 			TouchPosition, TouchEvent.GetTouchForce(), TouchEvent.GetEventTimestamp());
 		if (!bEventConsumed)
@@ -1044,14 +1107,16 @@ FReply FSceneViewport::OnTouchFirstMove(const FGeometry& MyGeometry, const FPoin
 	UpdateCachedCursorPos(MyGeometry, TouchEvent);
 	UpdateCachedGeometry(MyGeometry);
 
-	if (ViewportClient)
+PRAGMA_DISABLE_DEPRECATION_WARNINGS
+	if (FViewportClient* ClientPtr = ViewportClient)
+PRAGMA_ENABLE_DEPRECATION_WARNINGS
 	{
 		// Switch to the viewport clients world before processing input
-		FScopedConditionalWorldSwitcher WorldSwitcher(ViewportClient);
+		FScopedConditionalWorldSwitcher WorldSwitcher(ClientPtr);
 
 		const FVector2D TouchPosition = MyGeometry.AbsoluteToLocal(TouchEvent.GetScreenSpacePosition()) * MyGeometry.Scale;
 
-		const bool bEventConsumed = ViewportClient->InputTouch(this,
+		const bool bEventConsumed = ClientPtr->InputTouch(this,
 			FTouchId(TouchEvent.GetInputDeviceId(), static_cast<ETouchIndex::Type>(TouchEvent.GetPointerIndex())), ETouchType::FirstMove,
 			TouchPosition, TouchEvent.GetTouchForce(), TouchEvent.GetEventTimestamp());
 		if (!bEventConsumed)
@@ -1071,17 +1136,19 @@ FReply FSceneViewport::OnTouchGesture( const FGeometry& MyGeometry, const FPoint
 	UpdateCachedGeometry(MyGeometry);
 	UpdateCachedCursorPos( MyGeometry, GestureEvent );
 
-	if( ViewportClient )
+PRAGMA_DISABLE_DEPRECATION_WARNINGS
+	if (FViewportClient* ClientPtr = ViewportClient)
+PRAGMA_ENABLE_DEPRECATION_WARNINGS
 	{
 		// Switch to the viewport clients world before processing input
-		FScopedConditionalWorldSwitcher WorldSwitcher( ViewportClient );
+		FScopedConditionalWorldSwitcher WorldSwitcher(ClientPtr);
 
 		if (GestureEvent.GetGestureType() != EGestureEvent::LongPress)
 		{
 			FSlateApplication::Get().SetKeyboardFocus(ViewportWidget.Pin());
 		}
 
-		if(!ViewportClient->InputGesture(this, GestureEvent.GetInputDeviceId(), GestureEvent.GetGestureType(), GestureEvent.GetGesturePhase(), GestureEvent.GetGestureDelta(), GestureEvent.IsDirectionInvertedFromDevice(), GestureEvent.GetEventTimestamp()))
+		if(!ClientPtr->InputGesture(this, GestureEvent.GetInputDeviceId(), GestureEvent.GetGestureType(), GestureEvent.GetGesturePhase(), GestureEvent.GetGestureDelta(), GestureEvent.IsDirectionInvertedFromDevice(), GestureEvent.GetEventTimestamp()))
 		{
 			CurrentReplyState = FReply::Unhandled();
 		}
@@ -1093,14 +1160,16 @@ FReply FSceneViewport::OnTouchGesture( const FGeometry& MyGeometry, const FPoint
 FReply FSceneViewport::OnMotionDetected( const FGeometry& MyGeometry, const FMotionEvent& MotionEvent )
 {
 	// Start a new reply state
-	CurrentReplyState = FReply::Handled(); 
+	CurrentReplyState = FReply::Handled();
 
-	if(ViewportClient) 
+PRAGMA_DISABLE_DEPRECATION_WARNINGS
+	if (FViewportClient* ClientPtr = ViewportClient)
+PRAGMA_ENABLE_DEPRECATION_WARNINGS
 	{
 		// Switch to the viewport clients world before processing input
-		FScopedConditionalWorldSwitcher WorldSwitcher( ViewportClient );
+		FScopedConditionalWorldSwitcher WorldSwitcher(ClientPtr);
 
-		if(!ViewportClient->InputMotion( this, MotionEvent.GetInputDeviceId(), MotionEvent.GetTilt(), MotionEvent.GetRotationRate(), MotionEvent.GetGravity(), MotionEvent.GetAcceleration(), MotionEvent.GetEventTimestamp()))
+		if (!ClientPtr->InputMotion( this, MotionEvent.GetInputDeviceId(), MotionEvent.GetTilt(), MotionEvent.GetRotationRate(), MotionEvent.GetGravity(), MotionEvent.GetAcceleration(), MotionEvent.GetEventTimestamp()))
 		{
 			CurrentReplyState = FReply::Unhandled(); 
 		}
@@ -1111,9 +1180,11 @@ FReply FSceneViewport::OnMotionDetected( const FGeometry& MyGeometry, const FMot
 
 FPopupMethodReply FSceneViewport::OnQueryPopupMethod() const
 {
-	if (ViewportClient != nullptr)
+PRAGMA_DISABLE_DEPRECATION_WARNINGS
+	if (FViewportClient* ClientPtr = ViewportClient)
+PRAGMA_ENABLE_DEPRECATION_WARNINGS
 	{
-		return ViewportClient->OnQueryPopupMethod();
+		return ClientPtr->OnQueryPopupMethod();
 	}
 	else
 	{
@@ -1123,21 +1194,25 @@ FPopupMethodReply FSceneViewport::OnQueryPopupMethod() const
 
 bool FSceneViewport::HandleNavigation(const uint32 InUserIndex, TSharedPtr<SWidget> InDestination)
 {
-	if (ViewportClient != nullptr)
+PRAGMA_DISABLE_DEPRECATION_WARNINGS
+	if (FViewportClient* ClientPtr = ViewportClient)
+PRAGMA_ENABLE_DEPRECATION_WARNINGS
 	{
-		return ViewportClient->HandleNavigation(InUserIndex, InDestination);
+		return ClientPtr->HandleNavigation(InUserIndex, InDestination);
 	}
 	return false;
 }
 
 TOptional<bool> FSceneViewport::OnQueryShowFocus(const EFocusCause InFocusCause) const
 {
-	if (ViewportClient)
+PRAGMA_DISABLE_DEPRECATION_WARNINGS
+	if (FViewportClient* ClientPtr = ViewportClient)
+PRAGMA_ENABLE_DEPRECATION_WARNINGS
 	{
 		// Switch to the viewport clients world before processing input
-		FScopedConditionalWorldSwitcher WorldSwitcher(ViewportClient);
+		FScopedConditionalWorldSwitcher WorldSwitcher(ClientPtr);
 
-		return ViewportClient->QueryShowFocus(InFocusCause);
+		return ClientPtr->QueryShowFocus(InFocusCause);
 	}
 
 	return TOptional<bool>();
@@ -1159,12 +1234,16 @@ FReply FSceneViewport::OnKeyDown( const FGeometry& InGeometry, const FKeyEvent&
 		KeyStateMap.Add(Key, true);
 
 		//@todo Slate Viewports: FWindowsViewport checks for Alt+Enter or F11 and toggles fullscreen.  Unknown if fullscreen via this method will be needed for slate viewports. 
-		if (ViewportClient && GetSizeXY() != FIntPoint::ZeroValue)
+PRAGMA_DISABLE_DEPRECATION_WARNINGS
+		FViewportClient* ClientPtr = ViewportClient;
+PRAGMA_ENABLE_DEPRECATION_WARNINGS
+		if (ClientPtr && GetSizeXY() != FIntPoint::ZeroValue)
 		{
 			// Switch to the viewport clients world before processing input
-			FScopedConditionalWorldSwitcher WorldSwitcher(ViewportClient);
+			FScopedConditionalWorldSwitcher WorldSwitcher(ClientPtr);
 
-			if (!ViewportClient->InputKey(FInputKeyEventArgs(this, InKeyEvent.GetInputDeviceId(), Key, InKeyEvent.IsRepeat() ? IE_Repeat : IE_Pressed, 1.0f, false, InKeyEvent.GetEventTimestamp())))
+			if (!ClientPtr->InputKey(FInputKeyEventArgs(
+				this, InKeyEvent.GetInputDeviceId(), Key, InKeyEvent.IsRepeat() ? IE_Repeat : IE_Pressed, 1.0f, false, InKeyEvent.GetEventTimestamp())))
 			{
 				CurrentReplyState = FReply::Unhandled();
 			}
@@ -1187,12 +1266,16 @@ FReply FSceneViewport::OnKeyUp( const FGeometry& InGeometry, const FKeyEvent& In
 	{
 		KeyStateMap.Add(Key, false);
 
-		if (ViewportClient && GetSizeXY() != FIntPoint::ZeroValue)
+PRAGMA_DISABLE_DEPRECATION_WARNINGS
+		FViewportClient* ClientPtr = ViewportClient;
+PRAGMA_ENABLE_DEPRECATION_WARNINGS
+		if (ClientPtr && GetSizeXY() != FIntPoint::ZeroValue)
 		{
 			// Switch to the viewport clients world before processing input
-			FScopedConditionalWorldSwitcher WorldSwitcher(ViewportClient);
+			FScopedConditionalWorldSwitcher WorldSwitcher(ClientPtr);
 
-			if (!ViewportClient->InputKey(FInputKeyEventArgs(this, InKeyEvent.GetInputDeviceId(), Key, IE_Released, 1.0f, false, InKeyEvent.GetEventTimestamp())))
+			if (!ClientPtr->InputKey(FInputKeyEventArgs(
+				this, InKeyEvent.GetInputDeviceId(), Key, IE_Released, 1.0f, false, InKeyEvent.GetEventTimestamp())))
 			{
 				CurrentReplyState = FReply::Unhandled();
 			}
@@ -1215,11 +1298,12 @@ FReply FSceneViewport::OnAnalogValueChanged(const FGeometry& MyGeometry, const F
 	if (Key.IsValid())
 	{
 		KeyStateMap.Add(Key, true);
-
-		if (ViewportClient)
+PRAGMA_DISABLE_DEPRECATION_WARNINGS
+		if (FViewportClient* ClientPtr = ViewportClient)
+PRAGMA_ENABLE_DEPRECATION_WARNINGS
 		{
 			// Switch to the viewport clients world before processing input
-			FScopedConditionalWorldSwitcher WorldSwitcher(ViewportClient);
+			FScopedConditionalWorldSwitcher WorldSwitcher(ClientPtr);
 
 			FInputKeyEventArgs Args(
 				this,
@@ -1230,7 +1314,7 @@ FReply FSceneViewport::OnAnalogValueChanged(const FGeometry& MyGeometry, const F
 				1,
 				InAnalogInputEvent.GetEventTimestamp());
 
-			if (!ViewportClient->InputAxis(Args))
+			if (!ClientPtr->InputAxis(Args))
 			{
 				CurrentReplyState = FReply::Unhandled();
 			}
@@ -1249,12 +1333,15 @@ FReply FSceneViewport::OnKeyChar( const FGeometry& InGeometry, const FCharacterE
 	// Start a new reply state
 	CurrentReplyState = FReply::Handled(); 
 
-	if( ViewportClient && GetSizeXY() != FIntPoint::ZeroValue  )
+PRAGMA_DISABLE_DEPRECATION_WARNINGS
+	FViewportClient* ClientPtr = ViewportClient;
+PRAGMA_ENABLE_DEPRECATION_WARNINGS
+	if (ClientPtr && GetSizeXY() != FIntPoint::ZeroValue)
 	{
 		// Switch to the viewport clients world before processing input
-		FScopedConditionalWorldSwitcher WorldSwitcher( ViewportClient );
+		FScopedConditionalWorldSwitcher WorldSwitcher(ClientPtr);
 
-		if (!ViewportClient->InputChar(this, InCharacterEvent.GetUserIndex(), InCharacterEvent.GetCharacter()))
+		if (!ClientPtr->InputChar(this, InCharacterEvent.GetUserIndex(), InCharacterEvent.GetCharacter()))
 		{
 			CurrentReplyState = FReply::Unhandled();
 		}
@@ -1268,10 +1355,13 @@ FReply FSceneViewport::OnFocusReceived(const FFocusEvent& InFocusEvent)
 
 	if ( InFocusEvent.GetUser() == FSlateApplication::Get().GetUserIndexForKeyboard() )
 	{
-		if ( ViewportClient != nullptr )
+PRAGMA_DISABLE_DEPRECATION_WARNINGS
+		FViewportClient* ClientPtr = ViewportClient;
+PRAGMA_ENABLE_DEPRECATION_WARNINGS
+		if (ClientPtr)
 		{
-			FScopedConditionalWorldSwitcher WorldSwitcher(ViewportClient);
-			ViewportClient->ReceivedFocus(this);
+			FScopedConditionalWorldSwitcher WorldSwitcher(ClientPtr);
+			ClientPtr->ReceivedFocus(this);
 		}
 
 		// Update key state mappings so that the the viewport modifier states are valid upon focus.
@@ -1290,14 +1380,13 @@ FReply FSceneViewport::OnFocusReceived(const FFocusEvent& InFocusEvent)
 		{
 			FSlateApplication& SlateApp = FSlateApplication::Get();
 
-			const bool bPermanentCapture = (!GIsEditor || InFocusEvent.GetCause() == EFocusCause::Mouse)	&&
-										   (ViewportClient != nullptr)										&&
-					(( ViewportClient->GetMouseCaptureMode() == EMouseCaptureMode::CapturePermanently ) ||
-					 ( ViewportClient->GetMouseCaptureMode() == EMouseCaptureMode::CapturePermanently_IncludingInitialMouseDown )
-					);
+			const bool bPermanentCapture = (!GIsEditor || InFocusEvent.GetCause() == EFocusCause::Mouse)
+				&& ClientPtr != nullptr
+				&& (ClientPtr->GetMouseCaptureMode() == EMouseCaptureMode::CapturePermanently
+					|| ClientPtr->GetMouseCaptureMode() == EMouseCaptureMode::CapturePermanently_IncludingInitialMouseDown);
 
 			// if bPermanentCapture is true, then ViewportClient != nullptr, so it's ok to dereference it.  But the permanent capture must be tested first.
-			if ( SlateApp.IsActive() && bPermanentCapture && !ViewportClient->IgnoreInput() )
+			if (SlateApp.IsActive() && bPermanentCapture && !ClientPtr->IgnoreInput())
 			{
 				TSharedRef<SViewport> ViewportWidgetRef = ViewportWidget.Pin().ToSharedRef();
 
@@ -1323,25 +1412,32 @@ void FSceneViewport::OnFocusLost( const FFocusEvent& InFocusEvent )
 	bShouldCaptureMouseOnActivate = false;
 	bCursorHiddenDueToCapture = false;
 	KeyStateMap.Empty();
-	if ( ViewportClient != nullptr )
+PRAGMA_DISABLE_DEPRECATION_WARNINGS
+	if (FViewportClient* ClientPtr = ViewportClient)
+PRAGMA_ENABLE_DEPRECATION_WARNINGS
 	{
-		FScopedConditionalWorldSwitcher WorldSwitcher(ViewportClient);
-		ViewportClient->LostFocus(this);
+		FScopedConditionalWorldSwitcher WorldSwitcher(ClientPtr);
+		ClientPtr->LostFocus(this);
 	}
 }
 
 void FSceneViewport::OnViewportClosed()
 {
-	if( ViewportClient )
+PRAGMA_DISABLE_DEPRECATION_WARNINGS
+	if (FViewportClient* ClientPtr = ViewportClient)
+PRAGMA_ENABLE_DEPRECATION_WARNINGS
 	{
-		FScopedConditionalWorldSwitcher WorldSwitcher( ViewportClient );
-		ViewportClient->CloseRequested( this );
+		FScopedConditionalWorldSwitcher WorldSwitcher(ClientPtr);
+		ClientPtr->CloseRequested(this);
 	}
 }
 
 FReply FSceneViewport::OnRequestWindowClose()
 {
-	return (ViewportClient && !ViewportClient->WindowCloseRequested()) ? FReply::Handled() : FReply::Unhandled();
+PRAGMA_DISABLE_DEPRECATION_WARNINGS
+	FViewportClient* ClientPtr = ViewportClient;
+PRAGMA_ENABLE_DEPRECATION_WARNINGS
+	return ClientPtr && !ClientPtr->WindowCloseRequested() ? FReply::Handled() : FReply::Unhandled();
 }
 
 TWeakPtr<SWidget> FSceneViewport::GetWidget()
@@ -1351,13 +1447,15 @@ TWeakPtr<SWidget> FSceneViewport::GetWidget()
 
 FReply FSceneViewport::OnViewportActivated(const FWindowActivateEvent& InActivateEvent)
 {
-	if (ViewportClient != nullptr)
+PRAGMA_DISABLE_DEPRECATION_WARNINGS
+	if (FViewportClient* ClientPtr = ViewportClient)
+PRAGMA_ENABLE_DEPRECATION_WARNINGS
 	{
-		FScopedConditionalWorldSwitcher WorldSwitcher(ViewportClient);
-		ViewportClient->Activated(this, InActivateEvent);
-		
+		FScopedConditionalWorldSwitcher WorldSwitcher(ClientPtr);
+		ClientPtr->Activated(this, InActivateEvent);
+
 		// Determine if we're in permanent capture mode.  This cannot be cached as part of bShouldCaptureMouseOnActivate because it could change between window activate and deactivate
-		const bool bPermanentCapture = ViewportClient->IsInPermanentCapture();
+		const bool bPermanentCapture = ClientPtr->IsInPermanentCapture();
 
 		// If we are activating and had Mouse Capture on deactivate then we should get focus again
 		// It's important to note in the case of:
@@ -1437,10 +1535,12 @@ void FSceneViewport::OnViewportDeactivated(const FWindowActivateEvent& InActivat
 	bShouldCaptureMouseOnActivate = !GIsEditor  && (bShouldCaptureMouseOnActivate || HasMouseCapture());
 
 	KeyStateMap.Empty();
-	if (ViewportClient != nullptr)
+PRAGMA_DISABLE_DEPRECATION_WARNINGS
+	if (FViewportClient* ClientPtr = ViewportClient)
+PRAGMA_ENABLE_DEPRECATION_WARNINGS
 	{
-		FScopedConditionalWorldSwitcher WorldSwitcher(ViewportClient);
-		ViewportClient->Deactivated(this, InActivateEvent);
+		FScopedConditionalWorldSwitcher WorldSwitcher(ClientPtr);
+		ClientPtr->Deactivated(this, InActivateEvent);
 	}
 }
 
@@ -1840,13 +1940,15 @@ void FSceneViewport::ResizeViewport(uint32 NewSizeX, uint32 NewSizeY, EWindowMod
 		UpdateViewportRHI(false, NewSizeX, NewSizeY, NewWindowMode, PF_Unknown);
 		FCoreDelegates::OnSafeFrameChangedEvent.Broadcast();
 
-		if (ViewportClient)
+PRAGMA_DISABLE_DEPRECATION_WARNINGS
+		if (FViewportClient* ClientPtr = ViewportClient)
+PRAGMA_ENABLE_DEPRECATION_WARNINGS
 		{
 			// Invalidate, then redraw immediately so the user isn't left looking at an empty black viewport
 			// as they continue to resize the window.
 			Invalidate();
 
-			if ( ViewportClient->GetWorld() )
+			if (ClientPtr->GetWorld())
 			{
 				Draw();
 			}
@@ -1865,17 +1967,21 @@ void FSceneViewport::ResizeViewport(uint32 NewSizeX, uint32 NewSizeY, EWindowMod
 void FSceneViewport::InvalidateDisplay()
 {
 	// Dirty the viewport.  It will be redrawn next time the editor ticks.
-	if( ViewportClient != NULL )
+PRAGMA_DISABLE_DEPRECATION_WARNINGS
+	if (FViewportClient* ClientPtr = ViewportClient)
+PRAGMA_ENABLE_DEPRECATION_WARNINGS
 	{
-		ViewportClient->RedrawRequested( this );
+		ClientPtr->RedrawRequested(this);
 	}
 }
 
 void FSceneViewport::DeferInvalidateHitProxy()
 {
-	if( ViewportClient != NULL )
+PRAGMA_DISABLE_DEPRECATION_WARNINGS
+	if (FViewportClient* ClientPtr = ViewportClient)
+PRAGMA_ENABLE_DEPRECATION_WARNINGS
 	{
-		ViewportClient->RequestInvalidateHitProxy( this );
+		ClientPtr->RequestInvalidateHitProxy(this);
 	}
 }
 
@@ -2111,7 +2217,8 @@ void FSceneViewport::EnqueueBeginRenderFrame(const bool bShouldPresent)
 		ViewportProvider = FSlateApplication::Get().GetRenderer()->GetViewportProvider(*Window);
 	}
 
-	DebugCanvasDrawer->InitDebugCanvas(GetClient(), GetClient()->GetWorld());
+	FViewportClient& ViewportClientRef = GetClientChecked();
+	DebugCanvasDrawer->InitDebugCanvas(&ViewportClientRef, ViewportClientRef.GetWorld());
 
 	// Set the render target visible to the render thread. Must come before any render thread frame handling.
 	ENQUEUE_RENDER_COMMAND(SetRenderThreadViewportTarget)([
@@ -2174,14 +2281,20 @@ void FSceneViewport::Tick( const FGeometry& AllottedGeometry, double InCurrentTi
 	// properly propagated through to any widgets that depend on that functionality. The SceneViewport and RetainerWidget are the 
 	// only locations where this information exists in Slate, so we push the current scene onto the current
 	// Slate application so that we can leverage it in later calls.
-	if (ViewportClient && ViewportClient->GetWorld() && ViewportClient->GetWorld()->Scene)
+PRAGMA_DISABLE_DEPRECATION_WARNINGS
+	if (FViewportClient* ClientPtr = ViewportClient)
+PRAGMA_ENABLE_DEPRECATION_WARNINGS
 	{
-		FSlateApplication::Get().GetRenderer()->RegisterCurrentScene(ViewportClient->GetWorld()->Scene);
-	}
-	else
-	{
-		FSlateApplication::Get().GetRenderer()->RegisterCurrentScene(nullptr);
+		if (UWorld* const World = ClientPtr->GetWorld())
+		{
+			if (World->Scene)
+			{
+				FSlateApplication::Get().GetRenderer()->RegisterCurrentScene(World->Scene);
+				return;
+			}
+		}
 	}
+	FSlateApplication::Get().GetRenderer()->RegisterCurrentScene(nullptr);
 }
 
 void FSceneViewport::OnPlayWorldViewportSwapped( const FSceneViewport& OtherViewport )
@@ -2190,7 +2303,9 @@ void FSceneViewport::OnPlayWorldViewportSwapped( const FSceneViewport& OtherView
 	if( GetSizeXY() != OtherViewport.GetSizeXY() )
 	{
 		// Switch to the viewport clients world before processing input
-		FScopedConditionalWorldSwitcher WorldSwitcher( ViewportClient );
+PRAGMA_DISABLE_DEPRECATION_WARNINGS
+		FScopedConditionalWorldSwitcher WorldSwitcher(ViewportClient);
+PRAGMA_ENABLE_DEPRECATION_WARNINGS
 
 		UpdateViewportRHI( false, OtherViewport.GetSizeXY().X, OtherViewport.GetSizeXY().Y, EWindowMode::Windowed, PF_Unknown );
 
@@ -2206,8 +2321,10 @@ void FSceneViewport::OnPlayWorldViewportSwapped( const FSceneViewport& OtherView
 
 void FSceneViewport::SwapStatCommands( const FSceneViewport& OtherViewport )
 {
-	FViewportClient* ClientA = GetClient();
-	FViewportClient* ClientB = OtherViewport.GetClient();
+PRAGMA_DISABLE_DEPRECATION_WARNINGS
+	FViewportClient* const ClientA = ViewportClient;
+	FViewportClient* const ClientB = OtherViewport.ViewportClient;
+PRAGMA_ENABLE_DEPRECATION_WARNINGS
 	check(ClientA && ClientB);
 	// Only swap if both viewports have stats
 	const TArray<FString>* StatsA = ClientA->GetEnabledStats();
diff --git a/Engine/Source/Runtime/Engine/Private/UnrealClient.cpp b/Engine/Source/Runtime/Engine/Private/UnrealClient.cpp
--- a/Engine/Source/Runtime/Engine/Private/UnrealClient.cpp
+++ b/Engine/Source/Runtime/Engine/Private/UnrealClient.cpp
@@ -445,11 +445,12 @@ int32 FStatUnitData::DrawStat(FViewport* InViewport, FCanvas* InCanvas, int32 In
 	float Max_RHITTime = 0.0f;
 	float Max_InputLatencyTime = 0.0f;
 
-	const bool bShowUnitCriticalPathTimes = InViewport->GetClient() ? InViewport->GetClient()->IsStatEnabled(TEXT("UnitCriticalPath")) : false;
-	const bool bShowUnitMaxTimes = InViewport->GetClient() ? InViewport->GetClient()->IsStatEnabled(TEXT("UnitMax")) : false;
-	const bool bShowTSRStatistics = InViewport->GetClient() ? InViewport->GetClient()->IsStatEnabled(TEXT("TSR")) : false;
+	FViewportClient* Client = InViewport->GetClient();
+	const bool bShowUnitCriticalPathTimes = Client && Client->IsStatEnabled(TEXT("UnitCriticalPath"));
+	const bool bShowUnitMaxTimes = Client && Client->IsStatEnabled(TEXT("UnitMax"));
+	const bool bShowTSRStatistics = Client && Client->IsStatEnabled(TEXT("TSR"));
 #if !UE_BUILD_SHIPPING
-	const bool bShowRawUnitTimes = InViewport->GetClient() ? InViewport->GetClient()->IsStatEnabled(TEXT("Raw")) : false;
+	const bool bShowRawUnitTimes = Client && Client->IsStatEnabled(TEXT("Raw"));
 	RenderThreadTimes[CurrentIndex] = bShowRawUnitTimes ? RawRenderThreadTime : RenderThreadTime;
 	GameThreadTimes[CurrentIndex] = bShowRawUnitTimes ? RawGameThreadTime : GameThreadTime;
 	for (uint32 GPUIndex : FRHIGPUMask::All())
@@ -516,7 +517,7 @@ int32 FStatUnitData::DrawStat(FViewport* InViewport, FCanvas* InCanvas, int32 In
 	const bool bStereoRendering = GEngine->IsStereoscopic3D(InViewport);
 	UFont* Font = (!FPlatformProperties::SupportsWindowedMode() && GEngine->GetMediumFont()) ? GEngine->GetMediumFont() : GEngine->GetSmallFont();
 
-	const bool bShowUnitTimeGraph = InViewport->GetClient() ? InViewport->GetClient()->IsStatEnabled(TEXT("UnitGraph")) : false;
+	const bool bShowUnitTimeGraph = Client && Client->IsStatEnabled(TEXT("UnitGraph"));
 	bool bHaveGPUData[MAX_NUM_GPUS] = { false };
 	for (uint32 GPUIndex : FRHIGPUMask::All())
 	{
@@ -1216,7 +1217,7 @@ int32 FStatUnitData::DrawStat(FViewport* InViewport, FCanvas* InCanvas, int32 In
 
 		int32 AlertPrintY = GraphBottomYPos - AlertTimeMS * GraphVerticalPixelsPerMS - OutOfBudgetMarginHeight * AlertPrintHeight;
 
-		const bool bShowFrameTimeInUnitGraph = InViewport->GetClient() ? InViewport->GetClient()->IsStatEnabled(TEXT("UnitTime")) : false;
+		const bool bShowFrameTimeInUnitGraph = Client && Client->IsStatEnabled(TEXT("UnitTime"));
 		for (int32 StatIndex = 0; StatIndex < EGS_Count; ++StatIndex)
 		{
 			int32 LastPrintX = 0xFFFFFFFF;
@@ -1500,8 +1501,31 @@ int32 FStatHitchesData::DrawStat(FViewport* InViewport, FCanvas* InCanvas, int32
 /** Send when a viewport is resized */
 FViewport::FOnViewportResized FViewport::ViewportResizedEvent;
 
-FViewport::FViewport(FViewportClient* InViewportClient)
-	: ViewportClient(InViewportClient)
+FViewport::FViewport(const TStrongPtrVariant<FViewportClient>& Client)
+	: FViewport(Client, Client.Get())
+{
+}
+
+FViewport::FViewport(const TSharedPtr<FViewportClient>& Client)
+	: FViewport(Client, Client.Get())
+{
+}
+
+FViewport::FViewport(FViewportClient* const Client)
+	: FViewport(nullptr, Client)
+{
+}
+
+FViewport::FViewport(UScriptViewportClient* const Client)
+	: FViewport(Client, Client)
+{
+}
+
+FViewport::FViewport(const TStrongPtrVariant<FViewportClient>& InClient, FViewportClient* RawClient)
+	: Client(InClient)
+PRAGMA_DISABLE_DEPRECATION_WARNINGS
+	, ViewportClient(RawClient)
+PRAGMA_ENABLE_DEPRECATION_WARNINGS
 	, WindowMode(IsRunningGame() ? GSystemResolution.WindowMode : EWindowMode::Windowed)
 	, bHitProxiesCached(false)
 	, bHasRequestedToggleFreeze(false)
@@ -1518,7 +1542,10 @@ FViewport::FViewport(FViewportClient* InViewportClient)
 	}
 
 	// Cache the viewport client's hit proxy storage requirement.
-	bRequiresHitProxyStorage = ViewportClient && ViewportClient->RequiresHitProxyStorage();
+PRAGMA_DISABLE_DEPRECATION_WARNINGS
+	FViewportClient* ClientPtr = ViewportClient;
+PRAGMA_ENABLE_DEPRECATION_WARNINGS
+	bRequiresHitProxyStorage = ClientPtr && ClientPtr->RequiresHitProxyStorage();
 #if !WITH_EDITOR && !(UE_BUILD_SHIPPING || UE_BUILD_TEST)
 	if ( bRequiresHitProxyStorage )
 	{
@@ -1598,7 +1625,16 @@ static void HighResScreenshotEndFrame(FDummyViewport* DummyViewport)
 
 void FViewport::HighResScreenshot()
 {
-	if (!ViewportClient->GetEngineShowFlags())
+PRAGMA_DISABLE_DEPRECATION_WARNINGS
+	FViewportClient* const ViewportClientPtr = ViewportClient;
+PRAGMA_ENABLE_DEPRECATION_WARNINGS
+	if (!ViewportClientPtr)
+	{
+		return;
+	}
+
+	FEngineShowFlags* const EngineShowFlags = ViewportClientPtr->GetEngineShowFlags();
+	if (!EngineShowFlags)
 	{
 		return;
 	}
@@ -1607,7 +1643,9 @@ void FViewport::HighResScreenshot()
 	// cleared out before we use it below
 	const FString CachedScreenshotName = FScreenshotRequest::GetFilename();
 
-	FDummyViewport* DummyViewport = new FDummyViewport(ViewportClient);
+PRAGMA_DISABLE_DEPRECATION_WARNINGS
+	FDummyViewport* DummyViewport = new FDummyViewport(ViewportClientPtr);
+PRAGMA_ENABLE_DEPRECATION_WARNINGS
 	DummyViewport->SetInitialSize(FIntPoint(
 		(GScreenshotResolutionX > 0) ? GScreenshotResolutionX : SizeX,
 		(GScreenshotResolutionY > 0) ? GScreenshotResolutionY : SizeY
@@ -1617,11 +1655,11 @@ void FViewport::HighResScreenshot()
 
 	BeginInitResource(DummyViewport);
 
-	bool MaskShowFlagBackup = ViewportClient->GetEngineShowFlags()->HighResScreenshotMask;
-	const auto MotionBlurShowFlagBackup = ViewportClient->GetEngineShowFlags()->MotionBlur;
+	bool MaskShowFlagBackup = EngineShowFlags->HighResScreenshotMask;
+	const auto MotionBlurShowFlagBackup = EngineShowFlags->MotionBlur;
 
-	ViewportClient->GetEngineShowFlags()->SetHighResScreenshotMask(GetHighResScreenshotConfig().bMaskEnabled);
-	ViewportClient->GetEngineShowFlags()->SetMotionBlur(false);
+	EngineShowFlags->SetHighResScreenshotMask(GetHighResScreenshotConfig().bMaskEnabled);
+	EngineShowFlags->SetMotionBlur(false);
 
 	// Forcing 128-bit rendering pipeline
 	static IConsoleVariable* SceneColorFormatVar = IConsoleManager::Get().FindConsoleVariable(TEXT("r.SceneColorFormat"));
@@ -1662,9 +1700,10 @@ void FViewport::HighResScreenshot()
 		{
 			GIsHighResScreenshot = true;
 		}
-		FCanvas Canvas(DummyViewport, NULL, ViewportClient->GetWorld(), ViewportClient->GetWorld()->GetFeatureLevel());
+		UWorld* const World = ViewportClientPtr->GetWorld();
+		FCanvas Canvas(DummyViewport, nullptr, World, World->GetFeatureLevel());
 		{
-			ViewportClient->Draw(DummyViewport, &Canvas);
+			ViewportClientPtr->Draw(DummyViewport, &Canvas);
 		}
 		Canvas.Flush_GameThread();
 
@@ -1679,9 +1718,9 @@ void FViewport::HighResScreenshot()
 		--FrameDelay;
 	}
 
-	ViewportClient->GetEngineShowFlags()->SetHighResScreenshotMask(MaskShowFlagBackup);
-	ViewportClient->GetEngineShowFlags()->MotionBlur = MotionBlurShowFlagBackup;
-	bool bIsScreenshotSaved = ViewportClient->ProcessScreenShots(DummyViewport);
+	EngineShowFlags->SetHighResScreenshotMask(MaskShowFlagBackup);
+	EngineShowFlags->MotionBlur = MotionBlurShowFlagBackup;
+	bool bIsScreenshotSaved = ViewportClientPtr->ProcessScreenShots(DummyViewport);
 
 	SceneColorFormatVar->Set(OldSceneColorFormat, ECVF_SetByCode);
 	PostColorFormatVar->Set(OldPostColorFormat, ECVF_SetByCode);
@@ -1734,7 +1773,8 @@ FRHIGPUMask FViewport::GetGPUMask(FRHICommandListImmediate& RHICmdList) const
 void FViewport::Draw( bool bShouldPresent /*= true */)
 {
 	SCOPED_NAMED_EVENT(FViewport_Draw, FColor::Red);
-	UWorld* World = GetClient()->GetWorld();
+	FViewportClient& ViewportClientRef = GetClientChecked();
+	UWorld* World = ViewportClientRef.GetWorld();
 
 	// See what screenshot related features are required
 	static const auto CVarDumpFrames = IConsoleManager::Get().FindTConsoleVariableDataInt(TEXT("r.BufferVisualizationDumpFrames"));
@@ -1814,17 +1854,16 @@ void FViewport::Draw( bool bShouldPresent /*= true */)
 				}
 			}
 
-			UWorld* ViewportWorld = ViewportClient->GetWorld();
-			FCanvas Canvas(this, nullptr, ViewportWorld, ViewportWorld ? ViewportWorld->GetFeatureLevel() : GMaxRHIFeatureLevel, FCanvas::CDM_DeferDrawing, ViewportClient->ShouldDPIScaleSceneCanvas() ? ViewportClient->GetDPIScale() : 1.0f);
+			FCanvas Canvas(this, nullptr, World, World ? World->GetFeatureLevel() : GMaxRHIFeatureLevel, FCanvas::CDM_DeferDrawing, ViewportClientRef.ShouldDPIScaleSceneCanvas() ? ViewportClientRef.GetDPIScale() : 1.0f);
 			Canvas.SetRenderTargetRect(FIntRect(0, 0, SizeX, SizeY));
 			{
-				ViewportClient->Draw(this, &Canvas);
+				ViewportClientRef.Draw(this, &Canvas);
 			}
 			Canvas.Flush_GameThread();
 				
 			UGameViewportClient::OnViewportRendered().Broadcast(this);
-				
-			ViewportClient->ProcessScreenShots(this);
+
+			ViewportClientRef.ProcessScreenShots(this);
 	
 			// Slate doesn't present immediately. Tag the viewport as requiring vsync so that it happens.
 			SetRequiresVsync(bLockToVsync);
@@ -1896,7 +1935,10 @@ void FViewport::DeferInvalidateHitProxy()
 
 const TArray<FColor>& FViewport::GetRawHitProxyData(FIntRect InRect)
 {
-	FScopedConditionalWorldSwitcher WorldSwitcher(ViewportClient);
+PRAGMA_DISABLE_DEPRECATION_WARNINGS
+	FViewportClient* ViewportClientPtr = ViewportClient;
+PRAGMA_ENABLE_DEPRECATION_WARNINGS
+	FScopedConditionalWorldSwitcher WorldSwitcher(ViewportClientPtr);
 
 	const bool bIsRenderingStereo = GEngine->IsStereoscopic3D( this );
 
@@ -1932,9 +1974,10 @@ const TArray<FColor>& FViewport::GetRawHitProxyData(FIntRect InRect)
 			});
 
 		// Let the viewport client draw its hit proxies.
-		FCanvas Canvas(&HitProxyMap, &HitProxyMap, ViewportClient->GetWorld(), GetFeatureLevel(), FCanvas::CDM_DeferDrawing, ViewportClient->ShouldDPIScaleSceneCanvas() ? ViewportClient->GetDPIScale() : 1.0f);
+		FCanvas Canvas(&HitProxyMap, &HitProxyMap, ViewportClientPtr->GetWorld(), GetFeatureLevel(), FCanvas::CDM_DeferDrawing,
+			ViewportClientPtr->ShouldDPIScaleSceneCanvas() ? ViewportClientPtr->GetDPIScale() : 1.0f);
 		{
-			ViewportClient->Draw(this, &Canvas);
+			ViewportClientPtr->Draw(this, &Canvas);
 		}
 		Canvas.Flush_GameThread();
 
@@ -2049,7 +2092,7 @@ HHitProxy* FViewport::GetHitProxy(int32 X, int32 Y)
 		check(ProxyIndex<ProxyMap.Num());
 		HitProxy = ProxyMap[ProxyIndex];
 	
-		bool bIsOrtho = GetClient()->IsOrtho();
+		bool bIsOrtho = GetClientChecked().IsOrtho();
 
 		for(int32 TestY = 0;TestY < TestSizeY;TestY++)
 		{
@@ -2193,9 +2236,12 @@ FIntRect FViewport::CalculateViewExtents(float AspectRatio, const FIntRect& View
  *	Sets a viewport client if one wasn't provided at construction time.
  *	@param InViewportClient	- The viewport client to set.
  **/
-void FViewport::SetViewportClient( FViewportClient* InViewportClient )
+void FViewport::SetViewportClient(const TStrongPtrVariant<FViewportClient>& InViewportClient)
 {
-	ViewportClient = InViewportClient;
+	Client = InViewportClient;
+PRAGMA_DISABLE_DEPRECATION_WARNINGS
+	ViewportClient = InViewportClient.Get();
+PRAGMA_ENABLE_DEPRECATION_WARNINGS
 }
 
 void FViewport::ReleaseRHI()
@@ -2299,6 +2345,15 @@ FString FViewport::FHitProxyMap::GetReferencerName() const
 	return TEXT("FViewport::FHitProxyMap");
 }
 
+FViewportClient& FViewport::GetClientChecked() const
+{
+PRAGMA_DISABLE_DEPRECATION_WARNINGS
+	FViewportClient* Result = ViewportClient;
+PRAGMA_ENABLE_DEPRECATION_WARNINGS
+	check(Result);
+	return *Result;
+}
+
 /**
  * Globally enables/disables rendering
  *
@@ -2528,12 +2583,36 @@ float FCommonViewportClient::GetDPIDerivedResolutionFraction() const
  * FDummyViewport
  */
 
-FDummyViewport::FDummyViewport(FViewportClient* InViewportClient)
-	: FViewport(InViewportClient)
+FDummyViewport::FDummyViewport(const TStrongPtrVariant<FViewportClient>& InClient)
+	: FDummyViewport(InClient, InClient.Get())
+{
+}
+
+FDummyViewport::FDummyViewport(const TSharedPtr<FViewportClient>& Client)
+	: FDummyViewport(Client, Client.Get())
+{
+}
+
+FDummyViewport::FDummyViewport(FViewportClient* const InClient)
+	: FDummyViewport(nullptr, InClient)
+{
+}
+
+FDummyViewport::FDummyViewport(UScriptViewportClient* const InClient)
+	: FDummyViewport(InClient, InClient)
+{
+}
+
+FDummyViewport::FDummyViewport(const TStrongPtrVariant<FViewportClient>& InClient, FViewportClient* RawClient)
+	: FViewport(InClient, RawClient)
 	, DebugCanvas(NULL)
 {
-	UWorld* CurWorld = (InViewportClient != NULL ? InViewportClient->GetWorld() : NULL);
-	DebugCanvas = new FCanvas(this, NULL, CurWorld, (CurWorld != NULL ? CurWorld->GetFeatureLevel() : GMaxRHIFeatureLevel));
+PRAGMA_DISABLE_DEPRECATION_WARNINGS
+	FViewportClient* const ClientPtr = ViewportClient;
+PRAGMA_ENABLE_DEPRECATION_WARNINGS
+
+	UWorld* World = ClientPtr ? ClientPtr->GetWorld() : nullptr;
+	DebugCanvas = new FCanvas(this, nullptr, World, World ? World->GetFeatureLevel() : GMaxRHIFeatureLevel);
 
 	DebugCanvas->SetAllowedModes(0);
 }
diff --git a/Engine/Source/Runtime/Engine/Public/DummyViewport.h b/Engine/Source/Runtime/Engine/Public/DummyViewport.h
--- a/Engine/Source/Runtime/Engine/Public/DummyViewport.h
+++ b/Engine/Source/Runtime/Engine/Public/DummyViewport.h
@@ -11,7 +11,20 @@
 class FDummyViewport : public FViewport
 {
 public:
-	ENGINE_API FDummyViewport(FViewportClient* InViewportClient);
+	ENGINE_API FDummyViewport(const TStrongPtrVariant<FViewportClient>& Client);
+
+	ENGINE_API FDummyViewport(const TSharedPtr<FViewportClient>& Client);
+
+	UE_DEPRECATED(5.8, "Use FDummyViewport(TStrongPtrVariant<FViewportClient>) instead")
+	ENGINE_API FDummyViewport(FViewportClient* const Client);
+
+	// This constructor prevents ambiguity between FViewportClient* and TStrongPtrVariant<FViewportClient> when given a UScriptViewportClient
+	ENGINE_API FDummyViewport(UScriptViewportClient* const Client);
+
+private:
+	ENGINE_API FDummyViewport(const TStrongPtrVariant<FViewportClient>& Client, FViewportClient* RawClient);
+
+public:
 
 	ENGINE_API virtual ~FDummyViewport();
 
diff --git a/Engine/Source/Runtime/Engine/Public/Slate/SceneViewport.h b/Engine/Source/Runtime/Engine/Public/Slate/SceneViewport.h
--- a/Engine/Source/Runtime/Engine/Public/Slate/SceneViewport.h
+++ b/Engine/Source/Runtime/Engine/Public/Slate/SceneViewport.h
@@ -32,7 +32,20 @@ DECLARE_DELEGATE_OneParam( FOnSceneViewportResize, FVector2D );
 class FSceneViewport : public FViewport, public ISlateViewport
 {
 public:
-	ENGINE_API FSceneViewport( FViewportClient* InViewportClient, TSharedPtr<SViewport> InViewportWidget );
+	ENGINE_API FSceneViewport(const TStrongPtrVariant<FViewportClient>& Client, const TSharedPtr<SViewport>& Widget);
+
+	ENGINE_API FSceneViewport(const TSharedPtr<FViewportClient>& Client, const TSharedPtr<SViewport>& Widget);
+
+	UE_DEPRECATED(5.8, "Use FSceneViewport(TStrongPtrVariant<FViewportClient>, ...) instead")
+	ENGINE_API FSceneViewport(FViewportClient* const Client, const TSharedPtr<SViewport>& Widget);
+
+	// This constructor prevents ambiguity between FViewportClient* and TStrongPtrVariant<FViewportClient> when given a UScriptViewportClient
+	ENGINE_API FSceneViewport(UScriptViewportClient* const Client, const TSharedPtr<SViewport>& Widget);
+
+private:
+	ENGINE_API FSceneViewport(const TStrongPtrVariant<FViewportClient>& Client, FViewportClient* const RawClient, const TSharedPtr<SViewport>& Widget);
+
+public:
 	ENGINE_API ~FSceneViewport();
 
 	virtual void* GetWindow() override { return NULL; }
diff --git a/Engine/Source/Runtime/Engine/Public/UnrealClient.h b/Engine/Source/Runtime/Engine/Public/UnrealClient.h
--- a/Engine/Source/Runtime/Engine/Public/UnrealClient.h
+++ b/Engine/Source/Runtime/Engine/Public/UnrealClient.h
@@ -18,6 +18,7 @@
 #include "GenericPlatform/GenericWindow.h"
 #include "DynamicRenderScaling.h"
 #include "RHITypes.h"
+#include "Templates/PointerVariants.h"
 
 class FCanvas;
 class FRDGBuilder;
@@ -424,6 +425,8 @@ struct FStatHitchesData
 	int32 DrawStat(FViewport* InViewport, FCanvas* InCanvas, int32 InX, int32 InY);
 };
 
+class UScriptViewportClient;
+
 /**
  * Encapsulates the I/O of a viewport.
  * The viewport display is implemented using the platform independent RHI.
@@ -437,7 +440,20 @@ public:
 	ENGINE_API static FOnViewportResized ViewportResizedEvent;
 
 	// Constructor.
-	ENGINE_API FViewport(FViewportClient* InViewportClient);
+	ENGINE_API FViewport(const TStrongPtrVariant<FViewportClient>& Client);
+
+	ENGINE_API FViewport(const TSharedPtr<FViewportClient>& Client);
+
+	UE_DEPRECATED(5.8, "Use FViewport(TStrongPtrVariant<FViewportClient>) instead")
+	ENGINE_API FViewport(FViewportClient* const Client);
+
+	// This constructor prevents ambiguity between FViewportClient* and TStrongPtrVariant<FViewportClient> when given a UScriptViewportClient 
+	ENGINE_API FViewport(UScriptViewportClient* const Client);
+
+protected:
+	ENGINE_API FViewport(const TStrongPtrVariant<FViewportClient>& Client, FViewportClient* RawClient);
+
+public:
 	// Destructor
 	ENGINE_API virtual ~FViewport();
 
@@ -663,13 +679,55 @@ PRAGMA_ENABLE_DEPRECATION_WARNINGS
 	 *	Sets a viewport client if one wasn't provided at construction time.
 	 *	@param InViewportClient	- The viewport client to set.
 	 **/
-	ENGINE_API virtual void SetViewportClient( FViewportClient* InViewportClient );
+	ENGINE_API virtual void SetViewportClient(const TStrongPtrVariant<FViewportClient>& InViewportClient);
+
+	UE_REWRITE void SetViewportClient(const TSharedPtr<FViewportClient>& InViewportClient)
+	{
+		SetViewportClient(TStrongPtrVariant<FViewportClient>(InViewportClient));
+	}
+
+	UE_REWRITE void SetViewportClient(TSharedPtr<FViewportClient>&& InViewportClient)
+	{
+		SetViewportClient(TStrongPtrVariant<FViewportClient>(MoveTemp(InViewportClient)));
+	}
+
+	UE_REWRITE void SetViewportClient(TYPE_OF_NULLPTR InViewportClient)
+	{
+		SetViewportClient(TStrongPtrVariant<FViewportClient>(nullptr));
+	}
+
+	template<typename ClientType>
+	requires UE::CConvertibleTo<ClientType*, FViewportClient*> && (!UE::CUObject<ClientType>)
+	UE_DEPRECATED(5.8, "Use SetViewportClient(const TStrongPtrVariant<FViewportClient>&) instead")
+	void SetViewportClient(ClientType* InViewportClient)
+	{
+		Client.Reset(); // No way to set this safely. If caller is using deprecated setter, safe to assume they're still only using raw pointer getter
+		PRAGMA_DISABLE_DEPRECATION_WARNINGS
+		ViewportClient = InViewportClient;
+		PRAGMA_ENABLE_DEPRECATION_WARNINGS
+	}
 
 	//~ Begin FRenderTarget Interface.
 	virtual FIntPoint GetSizeXY() const override { return FIntPoint(SizeX, SizeY); }
 	FIntPoint GetInitialPositionXY() const { return FIntPoint(InitialPositionX, InitialPositionY); }
-	// Accessors.
-	FViewportClient* GetClient() const { return ViewportClient; }
+
+	const TStrongPtrVariant<FViewportClient>& GetClientPtr() const
+	{
+		return Client;
+	}
+
+	ENGINE_API FViewportClient& GetClientChecked() const;
+
+	/** Returns a raw pointer to the client.
+	 * Use when you know the Viewport will stay alive throughout the scope in which you access the returned pointer copy.
+	 * Do not Viewport->GetClient()->... — either test the result, or use GetClientChecked() instead.
+	 */
+	FViewportClient* GetClient() const
+	{
+		PRAGMA_DISABLE_DEPRECATION_WARNINGS
+		return ViewportClient;
+		PRAGMA_ENABLE_DEPRECATION_WARNINGS
+	}
 
 	/**
 	 * Globally enables/disables rendering
@@ -755,6 +813,9 @@ PRAGMA_ENABLE_DEPRECATION_WARNINGS
 protected:
 
 	/** The viewport's client. */
+	TStrongPtrVariant<FViewportClient> Client;
+
+	UE_DEPRECATED(5.8, "Use Client instead")
 	FViewportClient* ViewportClient;
 
 	/**
diff --git a/Engine/Source/Runtime/StreamingPauseRendering/Private/StreamingPauseRendering.cpp b/Engine/Source/Runtime/StreamingPauseRendering/Private/StreamingPauseRendering.cpp
--- a/Engine/Source/Runtime/StreamingPauseRendering/Private/StreamingPauseRendering.cpp
+++ b/Engine/Source/Runtime/StreamingPauseRendering/Private/StreamingPauseRendering.cpp
@@ -125,9 +125,10 @@ void FStreamingPauseRenderingModule::BeginStreamingPause( FViewport* GameViewpor
 
 		if(SizeXY.X > 0 && SizeXY.Y > 0 && bRenderLastFrame)
 		{
-			FViewportClient* ViewportClient = GameViewport->GetClient();
+			const TStrongPtrVariant<FViewportClient>& ViewportClientVariant = GameViewport->GetClientPtr();
+			FViewportClient* const ViewportClient = ViewportClientVariant.Get();
 
-			SceneViewport = MakeShareable(new FSceneViewport(ViewportClient, ViewportWidget));
+			SceneViewport = MakeShared<FSceneViewport>(ViewportClientVariant, ViewportWidget);
 
 			SceneViewport->UpdateViewportRHI(false, SizeXY.X, SizeXY.Y, EWindowMode::Fullscreen, PF_Unknown);
 
diff --git a/Engine/Source/Runtime/UMG/Private/Components/Viewport.cpp b/Engine/Source/Runtime/UMG/Private/Components/Viewport.cpp
--- a/Engine/Source/Runtime/UMG/Private/Components/Viewport.cpp
+++ b/Engine/Source/Runtime/UMG/Private/Components/Viewport.cpp
@@ -340,7 +340,7 @@ class SAutoRefreshViewport : public SViewport
 		
 
 		ViewportClient = MakeShareable(new FUMGViewportClient(PreviewScene.Get()));
-		Viewport = MakeShareable(new FSceneViewport(ViewportClient.Get(), SharedThis(this)));
+		Viewport = MakeShared<FSceneViewport>(ViewportClient, SharedThis(this));
 
 		// The viewport widget needs an interface so it knows what should render
 		SetViewportInterface(Viewport.ToSharedRef());
