New translator: Cull output properties according to blend mode.
Before the change the new translator would evaluate connected sub graphs to metallic etc in translucent blend modes.
Now we cull the output evaluation and use a default value which matches the old translator behavior.

#rnx
#Synced-CL 52611640

diff --git a/Engine/Source/Runtime/Engine/Private/Materials/MaterialIRModuleBuilder.cpp b/Engine/Source/Runtime/Engine/Private/Materials/MaterialIRModuleBuilder.cpp
--- a/Engine/Source/Runtime/Engine/Private/Materials/MaterialIRModuleBuilder.cpp
+++ b/Engine/Source/Runtime/Engine/Private/Materials/MaterialIRModuleBuilder.cpp
@@ -113,7 +113,7 @@ struct FMaterialIRModuleBuilderImpl
 		Module->FeatureLevel = Builder->FeatureLevel;
 		Module->QualityLevel = Builder->QualityLevel;
 		Module->BlendMode = Builder->BlendMode;
-		
+
 		// Declare a entry point to evaluate the vertex stage.
 		Module->AddEntryPoint(TEXTVIEW("VertexStage"), MIR::Stage_Vertex, 1);
 
@@ -172,14 +172,18 @@ struct FMaterialIRModuleBuilderImpl
 		SubstrateTranslatorData.ProcessSubstrateTreeAndTopology(MaterialShadingModels, MaterialGraphNodePreviewExpression, bDoClearFunctionStack);
 	}
 
-	// Some properties may be deactivated in the material, such as CustomizedUVx properties. This function returns whether given
-	// property is active in the translating material.
+	// This function returns whether a given property is active in the translating material.
+	// Some properties may be deactivated in the material, such as CustomizedUVx properties, or properties that are inactive for the material's blend mode.
 	bool PropertyIsActiveInMaterial(const EMaterialProperty Property)
 	{
 		if (IsMaterialPropertyCustomizedUV(Property) && Property >= MP_CustomizedUVs0 + Emitter.BaseMaterial->NumCustomizedUVs)
 		{
 			return false;
 		}
+		if (!Emitter.MaterialInterface->IsPropertyActive(Property))
+		{
+			return false;
+		}
 		return true;
 	}
 
@@ -203,12 +207,6 @@ struct FMaterialIRModuleBuilderImpl
 					continue;
 				}
 
-				// Skip evaluating shading model property if not set to "From Expression"
-				if (Property == MP_ShadingModel && !Emitter.MaterialInterface->IsShadingModelFromMaterialExpression())
-				{
-					continue;
-				}
-
 				// Read the material input associated to this property
 				FMaterialInputDescription Input;
 				GetExpressionInputDescription(Property, Input);
@@ -717,8 +715,8 @@ struct FMaterialIRModuleBuilderImpl
 
 		for (EMaterialProperty Property : MaterialAttributesAggregate::GetMaterialProperties())
 		{
-			// Ignore CustomUVs properties that haven't been enabled in the base material.
-			if (!PropertyIsActiveInMaterial(Property))
+			// Ignore deactivated properties except ShadingModel which we always need to emit.
+			if (Property != MP_ShadingModel && !PropertyIsActiveInMaterial(Property))
 			{
 				continue;
 			}
