- Mark textures referenced from non-pixel shader stages with their shader frequency mask during MIR compilation
- Add per-stage frequency tracking for texture uniforms so downstream systems know which shader stages reference each texture

#Synced-CL 52631319

diff --git a/Engine/Source/Runtime/Engine/Private/Materials/MaterialIRValueAnalyzer.cpp b/Engine/Source/Runtime/Engine/Private/Materials/MaterialIRValueAnalyzer.cpp
--- a/Engine/Source/Runtime/Engine/Private/Materials/MaterialIRValueAnalyzer.cpp
+++ b/Engine/Source/Runtime/Engine/Private/Materials/MaterialIRValueAnalyzer.cpp
@@ -25,6 +25,8 @@
 #include "UObject/Package.h"
 #include "Shader/Preshader2.h"
 
+#include "ShaderCompiler.h"
+
 static void AnalysisError_Analyze(FMaterialIRValueAnalyzer& Analyzer, MIR::FAnalysisError* AnalysisError)
 {
 	Analyzer.Module->AddError(AnalysisError->Expression, AnalysisError->Message.ToString());
@@ -439,6 +441,35 @@ static void Call_AnalyzeInStage(FMaterialIRValueAnalyzer& Analyzer, MIR::FCall*
 	Call->Function->Analysis_CallStageMask |= MIR::StageToMask(Stage);
 }
 
+static void TextureUniform_AnalyzeInStage(FMaterialIRValueAnalyzer& Analyzer, MIR::FTextureUniform* TextureUniform, MIR::EStage Stage)
+{
+	EMaterialTextureParameterType ParamType;
+	switch (TextureUniform->Type.GetKind())
+	{
+		case MIR::ETypeKind::Texture2D:			ParamType = EMaterialTextureParameterType::Standard2D; break;
+		case MIR::ETypeKind::TextureCube:		ParamType = EMaterialTextureParameterType::Cube;       break;
+		case MIR::ETypeKind::Texture2DArray:	ParamType = EMaterialTextureParameterType::Array2D;    break;
+		case MIR::ETypeKind::TextureCubeArray:	ParamType = EMaterialTextureParameterType::ArrayCube;  break;
+		case MIR::ETypeKind::Texture3D:			ParamType = EMaterialTextureParameterType::Volume;     break;
+		default:								return;		// VT / RVT handled via VT stacks
+	}
+
+	EShaderFrequency Frequency;
+	switch (Stage)
+	{
+		case MIR::EStage::Stage_Vertex:		Frequency = SF_Vertex;	break;
+		case MIR::EStage::Stage_Pixel:		Frequency = SF_Pixel;	break;
+		case MIR::EStage::Stage_Compute:	Frequency = SF_Compute;	break;
+		default:							checkNoEntry();			return;
+	}
+
+	Analyzer.CompilationOutput->UniformExpressionSet.AddTextureShaderFrequencyMask(
+		ParamType,
+		TextureUniform->Analysis_UniformIndex,
+		static_cast<EShaderFrequencyMask>(1u << Frequency)
+	);
+}
+
 static void RecurseMarkPreshaderTemp(MIR::FValue* Value)
 {
 	for (MIR::FValue* Use : Value->GetUses())
@@ -730,6 +761,7 @@ void FMaterialIRValueAnalyzer::AnalyzeInStage(MIR::FValue* Value, MIR::EStage St
 		EXPAND_CASE(VTPageTableRead);
 		EXPAND_CASE(PartialDerivative);
 		EXPAND_CASE(Call);
+		EXPAND_CASE(TextureUniform);
 		default: break;
 	}
 
diff --git a/Engine/Source/Runtime/Engine/Public/MaterialShared.h b/Engine/Source/Runtime/Engine/Public/MaterialShared.h
--- a/Engine/Source/Runtime/Engine/Public/MaterialShared.h
+++ b/Engine/Source/Runtime/Engine/Public/MaterialShared.h
@@ -741,6 +741,12 @@ public:
 	const FMaterialTextureCollectionParameterInfo& GetTextureCollectionParameter(int32 Index) const { return UniformTextureCollectionParameters[Index]; }
 
 	inline int32 GetNumTextures(EMaterialTextureParameterType Type) const { return UniformTextureParameters[(uint32)Type].Num(); }
+#if WITH_EDITORONLY_DATA
+	void AddTextureShaderFrequencyMask(EMaterialTextureParameterType Type, int32 Index, EShaderFrequencyMask Mask)
+	{
+		UniformTextureParameters[(int32)Type][Index].ShaderFrequencyMask |= Mask;
+	}
+#endif
 	ENGINE_API void GetGameThreadTextureValue(EMaterialTextureParameterType Type, int32 Index, const UMaterialInterface* MaterialInterface, const FMaterial& Material, UTexture*& OutValue, bool bAllowOverride = true) const;
 	ENGINE_API void GetTextureValue(EMaterialTextureParameterType Type, int32 Index, const FMaterialRenderContext& Context, const FMaterial& Material, const UTexture*& OutValue) const;
 	ENGINE_API void GetTextureValue(int32 Index, const FMaterialRenderContext& Context, const FMaterial& Material, const URuntimeVirtualTexture*& OutValue) const;
