- Fix incorrect component count (1 vs 3) in UMaterialExpressionTangentOutput::Build

#Synced-CL 52631319

diff --git a/Engine/Source/Runtime/Engine/Private/Materials/MaterialExpressionsToMIR.cpp b/Engine/Source/Runtime/Engine/Private/Materials/MaterialExpressionsToMIR.cpp
--- a/Engine/Source/Runtime/Engine/Private/Materials/MaterialExpressionsToMIR.cpp
+++ b/Engine/Source/Runtime/Engine/Private/Materials/MaterialExpressionsToMIR.cpp
@@ -4981,7 +4981,7 @@ void UMaterialExpressionSwitch::Build(MIR::FEmitter& Em)
 
 void UMaterialExpressionTangentOutput::Build(MIR::FEmitter& Em)
 {
-	FValueRef InputValue = Em.CastToFloat(Em.Input(&Input), 1);
+	FValueRef InputValue = Em.CastToFloat(Em.Input(&Input), 3);
 	Em.SetCustomOutputs(TEXTVIEW("TangentOutput"), { &InputValue, 1 }, MIR::EMaterialOutputFrequency::PerPixel);
 }
 
