- Fix for refraction output not appending RefractionDepthBias for translucent / SLW materials
- Fix incorrect default for UMaterialExpressionSingleLayerWaterMaterialOutput::Build

#Synced-CL 52631319

diff --git a/Engine/Source/Runtime/Engine/Private/Materials/MaterialExpressionsToMIR.cpp b/Engine/Source/Runtime/Engine/Private/Materials/MaterialExpressionsToMIR.cpp
--- a/Engine/Source/Runtime/Engine/Private/Materials/MaterialExpressionsToMIR.cpp
+++ b/Engine/Source/Runtime/Engine/Private/Materials/MaterialExpressionsToMIR.cpp
@@ -5030,7 +5030,7 @@ void UMaterialExpressionSingleLayerWaterMaterialOutput::Build(MIR::FEmitter& Em)
 		ScatteringCoefficientsValue ? Em.CastToFloat(ScatteringCoefficientsValue, 3) : Em.ConstantFloat3({ 0.f, 0.f, 0.f }),
 		AbsorptionCoefficientsValue ? Em.CastToFloat(AbsorptionCoefficientsValue, 3) : Em.ConstantFloat3({ 0.f, 0.f, 0.f }),
 		PhaseGValue ? Em.CastToFloat(PhaseGValue, 1) : Em.ConstantFloat(0.f),
-		Em.CastToFloat(Em.InputDefaultFloat(&ColorScaleBehindWater, 0.f), 1),
+		Em.CastToFloat(Em.InputDefaultFloat(&ColorScaleBehindWater, 1.f), 1),
 	};
 	Em.SetCustomOutputs(TEXTVIEW("SingleLayerWaterMaterialOutput"), OutputValues, MIR::EMaterialOutputFrequency::PerPixel);
 }
diff --git a/Engine/Source/Runtime/Engine/Private/Materials/MaterialIRModuleBuilder.cpp b/Engine/Source/Runtime/Engine/Private/Materials/MaterialIRModuleBuilder.cpp
--- a/Engine/Source/Runtime/Engine/Private/Materials/MaterialIRModuleBuilder.cpp
+++ b/Engine/Source/Runtime/Engine/Private/Materials/MaterialIRModuleBuilder.cpp
@@ -822,6 +822,13 @@ struct FMaterialIRModuleBuilderImpl
 				MIR::FValue* SelectionColor = Emitter.NamedPrimitiveUniform(NAME_SelectionColor, Emitter.ConstantFloat4(UE::Math::TVector4<float>{ 0.f, 0.f, 0.f, 0.f }));
 				PropertyValue = Emitter.Lerp(PropertyValue, Emitter.Swizzle(SelectionColor, MIR::FSwizzleMask::XYZ()), Emitter.Subscript(SelectionColor, 3));
 			}
+			// Append refraction depth bias to refraction for translucent and single layer water materials
+			else if (Property == MP_Refraction && (IsTranslucentBlendMode(Builder->BlendMode) || Builder->MaterialInterface->GetShadingModels().HasShadingModel(MSM_SingleLayerWater)))
+			{
+				MIR::FValueRef UserRefraction = Emitter.Swizzle(PropertyValue, MIR::FSwizzleMask::XY());
+				MIR::FValueRef RefractionDepthBias = Emitter.CastToFloat(Emitter.NamedPrimitiveUniform(FName(TEXT("RefractionDepthBias")), Emitter.ConstantFloat(Emitter.BaseMaterial->RefractionDepthBias)), 1);
+				PropertyValue = Emitter.Vector3(UserRefraction, RefractionDepthBias);
+			}
 			// Pack SubsurfaceColor (float3) with SubsurfaceProfile ID (float1) into float4,
 			// matching the old translator's behavior (HLSLMaterialTranslator.cpp:1623-1636).
 			else if (Property == MP_SubsurfaceColor)
