Added support for tracking of source state on the game thread for ASTP so that it's possible to query blend data.

#Synced-CL 52519763

diff --git a/Engine/Plugins/Runtime/MassGameplay/Source/MassRepresentation/Private/MassVisualizationComponent.cpp b/Engine/Plugins/Runtime/MassGameplay/Source/MassRepresentation/Private/MassVisualizationComponent.cpp
--- a/Engine/Plugins/Runtime/MassGameplay/Source/MassRepresentation/Private/MassVisualizationComponent.cpp
+++ b/Engine/Plugins/Runtime/MassGameplay/Source/MassRepresentation/Private/MassVisualizationComponent.cpp
@@ -778,7 +778,7 @@ void UMassVisualizationComponent::EndVisualChanges()
 									const TArray<FAnimSequenceTransformProviderSequence>& Sequences = ASTPDI->GetSequences();
 									FMassSkinnedMeshResolvedAnimState State = {
 										.SequenceIndex = ASTPDI->GetSequenceIndex(TrackId, 0),
-										.Position = ASTPDI->GetPosition(TrackId, 0),
+										.Position = ASTPDI->GetPosition(TrackId, 0).TargetPosition,
 										.PlayRate = ASTPDI->GetPlayRate(TrackId, 0),
 										.AnimSequence = Sequences.IsValidIndex(State.SequenceIndex) ? Sequences[State.SequenceIndex].Sequence.Get() : nullptr
 									};
@@ -869,7 +869,7 @@ void UMassVisualizationComponent::EndVisualChanges()
 								const TArray<FAnimSequenceTransformProviderSequence>& Sequences = ASTPDI->GetSequences();
 								FMassSkinnedMeshResolvedAnimState State = {
 									.SequenceIndex = ASTPDI->GetSequenceIndex(TrackId, 0),
-									.Position = ASTPDI->GetPosition(TrackId, 0),
+									.Position = ASTPDI->GetPosition(TrackId, 0).TargetPosition,
 									.PlayRate = ASTPDI->GetPlayRate(TrackId, 0),
 									.AnimSequence = Sequences.IsValidIndex(State.SequenceIndex) ? Sequences[State.SequenceIndex].Sequence.Get() : nullptr
 								};
diff --git a/Engine/Source/Runtime/Engine/Classes/Animation/AnimSequenceTransformProviderData.h b/Engine/Source/Runtime/Engine/Classes/Animation/AnimSequenceTransformProviderData.h
--- a/Engine/Source/Runtime/Engine/Classes/Animation/AnimSequenceTransformProviderData.h
+++ b/Engine/Source/Runtime/Engine/Classes/Animation/AnimSequenceTransformProviderData.h
@@ -47,6 +47,29 @@ enum class EAnimSequenceTrackMode : uint8
 	BlendSpace,
 };
 
+USTRUCT(BlueprintType)
+struct FAnimSequenceTrackPosition
+{
+	GENERATED_USTRUCT_BODY()
+
+	UPROPERTY(BlueprintReadOnly, Category = Animation)
+	float SourcePosition = 0.0f;
+
+	UPROPERTY(BlueprintReadOnly, Category = Animation)
+	float TargetPosition = 0.0f;
+
+	UPROPERTY(BlueprintReadOnly, Category = Animation)
+	float BlendWeight = 0.0f;
+
+	UPROPERTY(BlueprintReadOnly, Category = Animation)
+	float BlendTime = 0.0f;
+
+	bool IsBlending() const
+	{
+		return BlendWeight > 0.0f && BlendWeight < 1.0f;
+	}
+};
+
 UENUM(BlueprintType)
 enum class EAnimSequenceTransformProviderLayerBlendMode : uint8
 {
@@ -277,7 +300,6 @@ public:
 		Manual.Position = InPosition;
 	}
 
-
 	uint32 GetNumSamples() const
 	{
 		return bBlendSpace ? BlendSpace.NumSamples : 1;
@@ -338,6 +360,17 @@ public:
 		return (WorldTime - Auto.ReferenceTimestamp) * Auto.PlayRate.GetFloat();
 	}
 
+	double GetReferenceTimestamp() const
+	{
+		check(!IsManual());
+		return bBlendSpace ? BlendSpace.ReferenceTimestamp : Auto.ReferenceTimestamp;
+	}
+
+	float GetManualPosition() const
+	{
+		return Manual.Position;
+	}
+
 	void SetLoopMode(EAnimSequenceTrackLoopMode NewLoopMode)
 	{
 		LoopMode = static_cast<uint32>(NewLoopMode);
@@ -424,25 +457,30 @@ public:
 		LayerWeight = PrevLayerWeight;
 	}
 
-	static float WrapPosition(EAnimSequenceTrackLoopMode LoopMode, float Position, float SequenceLength)
+	static float WrapPosition(EAnimSequenceTrackLoopMode LoopMode, double Position, double SequenceLength)
 	{
 		if (LoopMode == EAnimSequenceTrackLoopMode::Loop)
 		{
-			if (SequenceLength >= 0.0f)
+			if (SequenceLength > 0.0)
 			{
-				Position = SequenceLength == 0.0f ? 0.0f : FMath::Fmod(Position, SequenceLength);
+				Position = FMath::Fmod(Position, SequenceLength);
 
-				if (Position < 0.0f)
+				if (Position < 0.0)
 				{
 					Position += SequenceLength;
 				}
 			}
+			else
+			{
+				Position = 0.0;
+			}
 		}
 		else
 		{
-			Position = FMath::Clamp(Position, 0.0f, SequenceLength);
+			Position = FMath::Clamp(Position, 0.0, SequenceLength);
 		}
-		return Position;
+
+		return static_cast<float>(Position);
 	}
 
 private:
@@ -949,6 +987,10 @@ private:
 	UPROPERTY(EditAnywhere, BlueprintReadOnly, Category = TransformProvider, meta = (AllowPrivateAccess = "true", EditCondition = "BlendSpaceList == nullptr", EditConditionHides))
 	TArray<FAnimSequenceTransformProviderBlendSpace> BlendSpaces;
 
+	// Enable game-thread playback tracking for blend state queries via GetPosition.
+	UPROPERTY(EditAnywhere, BlueprintReadOnly, Category = TransformProvider, meta = (AllowPrivateAccess = "true"))
+	bool bEnableBlendTracking = false;
+
 	double GetWorldTime() const;
 
 	FRenderCommandFence DestroyFence;
@@ -1050,9 +1092,9 @@ public:
 	UFUNCTION(BlueprintCallable, Category = "TransformProviders|AnimSequence")
 	ENGINE_API EAnimSequenceTrackLoopMode GetLoopMode(int32 TrackIndex, int32 LayerIndex) const;
 
-	// Returns the current wrapped playback position for a track layer. Returns 0 for blend space mode.
+	// Returns the current playback position state for a track layer. Source/target blend info requires bEnableBlendTracking.
 	UFUNCTION(BlueprintCallable, Category = "TransformProviders|AnimSequence")
-	ENGINE_API float GetPosition(int32 TrackIndex, int32 LayerIndex) const;
+	ENGINE_API FAnimSequenceTrackPosition GetPosition(int32 TrackIndex, int32 LayerIndex) const;
 
 	// Returns the 2D blend space position for a track layer, or zero if not in blend space mode.
 	UFUNCTION(BlueprintCallable, Category = "TransformProviders|AnimSequence")
@@ -1081,6 +1123,9 @@ private:
 	FRenderCommandFence InstanceDestroyFence;
 	UE::FMutex SubmitChangesMutex;
 
+	class FPlaybackTracker;
+	TPimplPtr<FPlaybackTracker> PlaybackTracker;
+
 	// Per-blend-space resolution data (one per BlendSpaces entry)
 	struct FBlendSpaceCache
 	{
diff --git a/Engine/Source/Runtime/Engine/Private/Animation/AnimSequenceTransformProviderData.cpp b/Engine/Source/Runtime/Engine/Private/Animation/AnimSequenceTransformProviderData.cpp
--- a/Engine/Source/Runtime/Engine/Private/Animation/AnimSequenceTransformProviderData.cpp
+++ b/Engine/Source/Runtime/Engine/Private/Animation/AnimSequenceTransformProviderData.cpp
@@ -1092,6 +1092,180 @@ EDataValidationResult UAnimSequenceTransformProviderData::IsDataValid(FDataValid
 
 ///////////////////////////////////////////////////////////////////////////////////////////////////
 
+class UAnimSequenceTransformProviderDataInstance::FPlaybackTracker
+{
+public:
+	explicit FPlaybackTracker(int32 InNumLayers)
+		: NumLayers(InNumLayers)
+	{}
+
+	void SetSourceData(int32 TrackIndex, int32 LayerIndex, const FAnimSequenceTrackPackedData& Current, float InBlendTime, double WorldTime)
+	{
+		const int32 RequiredSize = (TrackIndex + 1) * NumLayers;
+		if (States.Num() < RequiredSize)
+		{
+			NumTracks = TrackIndex + 1;
+			States.SetNum(RequiredSize);
+		}
+
+		FBlendState& State = GetState(TrackIndex, LayerIndex);
+
+		if (State.BlendStartTime < 0.0)
+		{
+			State.BlendStartTime = 0.0;
+			return;
+		}
+
+		if (InBlendTime > 0.0f)
+		{
+			if (Current.IsManual())
+			{
+				State.ManualPosition = Current.GetManualPosition();
+			}
+			else
+			{
+				State.ReferenceTimestamp = Current.GetReferenceTimestamp();
+			}
+
+			State.PlayRate = FFloat16(Current.GetPlayRate());
+			State.SequenceIndex = Current.IsBlendSpace() ? 0 : Current.GetSequenceIndex();
+			State.LoopMode = static_cast<uint32>(Current.GetLoopMode());
+			State.bAutoPlay = Current.IsAutoPlay();
+			State.bBlendSpace = Current.IsBlendSpace();
+			State.BlendStartTime = WorldTime;
+			State.BlendTime = FFloat16(InBlendTime);
+		}
+	}
+
+	FAnimSequenceTrackPosition GetPosition(
+		int32 TrackIndex,
+		int32 LayerIndex,
+		const FAnimSequenceTrackPackedData& CurrentTarget,
+		double WorldTime,
+		TConstArrayView<FAnimSequenceTransformProviderSequence> InSequences) const
+	{
+		FAnimSequenceTrackPosition Result;
+		Result.TargetPosition = DerivePosition(CurrentTarget, WorldTime, InSequences);
+
+		if (!IsValidIndex(TrackIndex, LayerIndex))
+		{
+			Result.SourcePosition = Result.TargetPosition;
+			return Result;
+		}
+
+		const FBlendState& State = GetState(TrackIndex, LayerIndex);
+
+		if (State.BlendStartTime >= 0.0)
+		{
+			const float BlendTimeValue = State.BlendTime.GetFloat();
+			const float Elapsed = static_cast<float>(WorldTime - State.BlendStartTime);
+			Result.BlendWeight = (BlendTimeValue > 0.0f) ? FMath::Clamp(Elapsed / BlendTimeValue, 0.0f, 1.0f) : 0.0f;
+			Result.BlendTime = BlendTimeValue;
+
+			if (Result.IsBlending())
+			{
+				Result.SourcePosition = DerivePosition(State, WorldTime, InSequences);
+			}
+			else
+			{
+				Result.SourcePosition = Result.TargetPosition;
+				Result.BlendWeight = 0.0f;
+			}
+		}
+		else
+		{
+			Result.SourcePosition = Result.TargetPosition;
+		}
+
+		return Result;
+	}
+
+private:
+	struct FBlendState
+	{
+		uint32 SequenceIndex : 20 = 0;
+		uint32 LoopMode      : 1  = 0;
+		uint32 bAutoPlay     : 1  = 0;
+		uint32 bBlendSpace   : 1  = 0;
+		union
+		{
+			double ReferenceTimestamp;
+			float ManualPosition;
+		};
+		double BlendStartTime     = -1.0;
+		FFloat16 PlayRate         = FFloat16(1.0f);
+		FFloat16 BlendTime        = FFloat16(0.0f);
+	};
+
+	FBlendState& GetState(int32 TrackIndex, int32 LayerIndex)
+	{
+		return States[TrackIndex * NumLayers + LayerIndex];
+	}
+
+	const FBlendState& GetState(int32 TrackIndex, int32 LayerIndex) const
+	{
+		return States[TrackIndex * NumLayers + LayerIndex];
+	}
+
+	bool IsValidIndex(int32 TrackIndex, int32 LayerIndex) const
+	{
+		const int32 Index = TrackIndex * NumLayers + LayerIndex;
+		return Index >= 0 && Index < States.Num();
+	}
+
+	static float DerivePosition(const FAnimSequenceTrackPackedData& Data, double WorldTime, TConstArrayView<FAnimSequenceTransformProviderSequence> InSequences)
+	{
+		if (Data.IsBlendSpace())
+		{
+			return static_cast<float>(Data.GetTimePosition(WorldTime));
+		}
+
+		const int32 SequenceIndex = Data.GetSequenceIndex();
+
+		if (!InSequences.IsValidIndex(SequenceIndex) || !InSequences[SequenceIndex].Sequence)
+		{
+			return 0.0f;
+		}
+
+		return FAnimSequenceTrackPackedData::WrapPosition(
+			Data.GetLoopMode(),
+			Data.GetTimePosition(WorldTime),
+			InSequences[SequenceIndex].Sequence->GetPlayLength());
+	}
+
+	float DerivePosition(const FBlendState& State, double WorldTime, TConstArrayView<FAnimSequenceTransformProviderSequence> InSequences) const
+	{
+		if (State.bAutoPlay || State.bBlendSpace)
+		{
+			const double RawPosition = (WorldTime - State.ReferenceTimestamp) * State.PlayRate.GetFloat();
+
+			if (State.bBlendSpace)
+			{
+				// Blend space has no single sequence length, return raw time.
+				return static_cast<float>(RawPosition);
+			}
+
+			if (InSequences.IsValidIndex(State.SequenceIndex) && InSequences[State.SequenceIndex].Sequence)
+			{
+				return FAnimSequenceTrackPackedData::WrapPosition(
+					static_cast<EAnimSequenceTrackLoopMode>(State.LoopMode),
+					RawPosition,
+					InSequences[State.SequenceIndex].Sequence->GetPlayLength());
+			}
+
+			return static_cast<float>(RawPosition);
+		}
+
+		return State.ManualPosition;
+	}
+
+	TArray<FBlendState> States;
+	int32 NumTracks = 0;
+	int32 NumLayers = 0;
+};
+
+///////////////////////////////////////////////////////////////////////////////////////////////////
+
 UAnimSequenceTransformProviderDataInstance* UAnimSequenceTransformProviderDataInstance::CreateAnimSequenceTransformProviderDataInstance(UAnimSequenceTransformProviderData* ProviderData, UInstancedSkinnedMeshComponent* Owner)
 {
 	if (IsValid(ProviderData) && IsValid(Owner))
@@ -1112,6 +1286,7 @@ void UAnimSequenceTransformProviderDataInstance::SetParentProvider(UAnimSequence
 	Sequences = ProviderData->Sequences;
 	Layers = ProviderData->Layers;
 	BlendSpaces = ProviderData->BlendSpaces;
+	bEnableBlendTracking = ProviderData->bEnableBlendTracking;
 
 #if WITH_EDITORONLY_DATA
 	SequenceList = ProviderData->SequenceList;
@@ -1188,6 +1363,8 @@ void UAnimSequenceTransformProviderDataInstance::BeginDestroy()
 		// Move destruction onto the render thread.
 	});
 
+	PlaybackTracker.Reset();
+
 	Super::BeginDestroy();
 }
 
@@ -1235,6 +1412,16 @@ void UAnimSequenceTransformProviderDataInstance::InitLayerPools(bool bMarkDirty)
 	bTrackPoolsDirty = true;
 	BuildBlendSpaceCaches();
 
+	// Create, recreate, or destroy playback tracker based on parent setting.
+	if (bEnableBlendTracking)
+	{
+		PlaybackTracker = MakePimpl<FPlaybackTracker>(Layers.Num());
+	}
+	else
+	{
+		PlaybackTracker.Reset();
+	}
+
 	if (bMarkDirty)
 	{
 		MarkOwnerRenderStateDirty();
@@ -1533,6 +1720,11 @@ bool UAnimSequenceTransformProviderDataInstance::SetAutoPlayData(int32 TrackInde
 
 	if (Pool.IsActiveIndex(TrackIndex))
 	{
+		if (PlaybackTracker)
+		{
+			PlaybackTracker->SetSourceData(TrackIndex, LayerIndex, Pool.GetData(TrackIndex), Data.BlendTime, GetWorldTime());
+		}
+
 		if (Pool.UpdateTrack(TrackIndex, FAnimSequenceTrackPackedData::Pack(Data, GetWorldTime())))
 		{
 			MarkOwnerRenderStateDirty();
@@ -1554,6 +1746,11 @@ bool UAnimSequenceTransformProviderDataInstance::SetManualData(int32 TrackIndex,
 
 	if (Pool.IsActiveIndex(TrackIndex))
 	{
+		if (PlaybackTracker)
+		{
+			PlaybackTracker->SetSourceData(TrackIndex, LayerIndex, Pool.GetData(TrackIndex), Data.BlendTime, GetWorldTime());
+		}
+
 		if (Pool.UpdateTrack(TrackIndex, FAnimSequenceTrackPackedData::Pack(Data, GetSequencePlayLength(Data.SequenceIndex))))
 		{
 			MarkOwnerRenderStateDirty();
@@ -1646,24 +1843,41 @@ EAnimSequenceTrackMode UAnimSequenceTransformProviderDataInstance::GetTrackMode(
 	return EAnimSequenceTrackMode::Manual;
 }
 
-float UAnimSequenceTransformProviderDataInstance::GetPosition(int32 TrackIndex, int32 LayerIndex) const
+FAnimSequenceTrackPosition UAnimSequenceTransformProviderDataInstance::GetPosition(int32 TrackIndex, int32 LayerIndex) const
 {
 	if (!IsValidTrackAndLayer(TrackIndex, LayerIndex))
 	{
-		return 0.0f;
+		return {};
+	}
+
+	if (PlaybackTracker)
+	{
+		return PlaybackTracker->GetPosition(TrackIndex, LayerIndex, TrackPools[LayerIndex].GetData(TrackIndex), GetWorldTime(), Sequences);
 	}
 
+	// No tracker. Return position with Source = Target and no blend info.
+	FAnimSequenceTrackPosition Result;
 	const FAnimSequenceTrackPackedData& Data = TrackPools[LayerIndex].GetData(TrackIndex);
-	const int32 SeqIndex = Data.GetSequenceIndex();
 
-	if (SeqIndex == INDEX_NONE || !Sequences.IsValidIndex(SeqIndex))
+	if (Data.IsBlendSpace())
 	{
-		return 0.0f;
+		// Blend space has no single sequence — return raw time.
+		Result.TargetPosition = static_cast<float>(Data.GetTimePosition(GetWorldTime()));
+	}
+	else
+	{
+		const int32 SeqIndex = Data.GetSequenceIndex();
+		if (SeqIndex != INDEX_NONE && Sequences.IsValidIndex(SeqIndex))
+		{
+			Result.TargetPosition = FAnimSequenceTrackPackedData::WrapPosition(
+				Data.GetLoopMode(),
+				Data.GetTimePosition(GetWorldTime()),
+				GetSequencePlayLength(SeqIndex));
+		}
 	}
 
-	const double RawTime = Data.GetTimePosition(GetWorldTime());
-	const float PlayLength = GetSequencePlayLength(SeqIndex);
-	return FAnimSequenceTrackPackedData::WrapPosition(Data.GetLoopMode(), static_cast<float>(RawTime), PlayLength);
+	Result.SourcePosition = Result.TargetPosition;
+	return Result;
 }
 
 bool UAnimSequenceTransformProviderDataInstance::SetBlendSpaceData(int32 TrackIndex, int32 LayerIndex, const FAnimSequenceTrackBlendSpaceData& InData)
@@ -1679,6 +1893,11 @@ bool UAnimSequenceTransformProviderDataInstance::SetBlendSpaceData(int32 TrackIn
 		return false;
 	}
 
+	if (PlaybackTracker)
+	{
+		PlaybackTracker->SetSourceData(TrackIndex, LayerIndex, TrackPools[LayerIndex].GetData(TrackIndex), InData.BlendTime, GetWorldTime());
+	}
+
 	TrackPools[LayerIndex].UpdateTrack(TrackIndex, FAnimSequenceTrackPackedData::Pack(InData, GetWorldTime()));
 	MarkOwnerRenderStateDirty();
 	return true;
