When adding a DynamicAsset if the BundleData.Bundles is empty remove the TypeData aswell.
This fixes an issue where empty bundles fail to update or set the manager references because no assets exist.

#Synced-CL 52564791

diff --git a/Engine/Source/Runtime/Engine/Private/AssetManager.cpp b/Engine/Source/Runtime/Engine/Private/AssetManager.cpp
--- a/Engine/Source/Runtime/Engine/Private/AssetManager.cpp
+++ b/Engine/Source/Runtime/Engine/Private/AssetManager.cpp
@@ -1737,15 +1737,15 @@ bool UAssetManager::AddDynamicAsset(const FPrimaryAssetId& PrimaryAssetId, const
 			*PrimaryAssetId.ToString(), *OldData->GetAssetPtr().ToString(), *AssetPath.ToString());
 	}
 
-	TypeData.FindOrAddAsset(PrimaryAssetId.PrimaryAssetName, FSoftObjectPtr(AssetPath), AssetPath, AssetPathMap);
-
 	if (BundleData.Bundles.Num() > 0)
 	{
+		TypeData.FindOrAddAsset(PrimaryAssetId.PrimaryAssetName, FSoftObjectPtr(AssetPath), AssetPath, AssetPathMap);
 		CachedAssetBundles.Add(PrimaryAssetId, MakeShared<FAssetBundleData, ESPMode::ThreadSafe>(BundleData));
 	}
 	else if (OldData)
 	{
 		CachedAssetBundles.Remove(PrimaryAssetId);
+		TypeData.RemoveAsset(PrimaryAssetId.PrimaryAssetName, AssetPathMap);
 	}
 
 	return true;
