[StreamableManager] Stop keeping released streamable targets reachable during reference collection

AddReferencedObjects now skips entries with no active handles and no pending async
load request. Such entries are stale - the last handle was released but PreGarbageCollect
has not yet pruned them from StreamableItems. Unconditionally collecting their references
was keeping target assets reachable past their release point, which caused the plugin leak
detector to report false positives.

This race occurs when a handle is released from BeginDestroy of a garbage object:
OnPreGarbageCollect has already run for that cycle, so the stale entry lingers until the next one.
Stale redirector references in StreamableRedirects are also skipped for consistency.

#rnx
#tests cooked build

diff --git a/Engine/Source/Runtime/Engine/Private/StreamableManager.cpp b/Engine/Source/Runtime/Engine/Private/StreamableManager.cpp
--- a/Engine/Source/Runtime/Engine/Private/StreamableManager.cpp
+++ b/Engine/Source/Runtime/Engine/Private/StreamableManager.cpp
@@ -1675,15 +1675,28 @@ void FStreamableManager::OnPreGarbageCollect()
 
 void FStreamableManager::AddReferencedObjects(FReferenceCollector& Collector)
 {
+	TSet<FSoftObjectPath, DefaultKeyFuncs<FSoftObjectPath>, TInlineSetAllocator<256>> StaleRedirects;
 	// If there are active streamable handles in the editor, this will cause the user to Force Delete, which is irritating but necessary because weak pointers cannot be used here
-	for (auto& Pair : StreamableItems)
+	for (const TTuple<FSoftObjectPath, FStreamable*>& Pair : StreamableItems)
 	{
-		Collector.AddStableReference(&Pair.Value->Target);
+		// It's possible that last handle got released, but PreGarbageCollect was not called yet.
+		// Skip adding stale references, it's guaranteed that they will be dropped on next GC.
+		if (!Pair.Value->ActiveHandles.IsEmpty() || Pair.Value->bAsyncLoadRequestOutstanding)
+		{
+			Collector.AddStableReference(&Pair.Value->Target);
+		}
+		else 
+		{
+			StaleRedirects.Add(Pair.Key);
+		}
 	}
-	
+
 	for (TPair<FSoftObjectPath, FRedirectedPath>& Pair : StreamableRedirects)
 	{
-		Collector.AddStableReference(&Pair.Value.LoadedRedirector);
+		if (!StaleRedirects.Contains(Pair.Value.NewPath))
+		{
+			Collector.AddStableReference(&Pair.Value.LoadedRedirector);
+		}
 	}
 }
 
