Fix for issue where Substrate FrontMaterial attribute was being added unnecessarily.

#rnx
#Synced-CL 52631319

diff --git a/Engine/Source/Runtime/Engine/Private/Materials/MaterialCachedData.cpp b/Engine/Source/Runtime/Engine/Private/Materials/MaterialCachedData.cpp
--- a/Engine/Source/Runtime/Engine/Private/Materials/MaterialCachedData.cpp
+++ b/Engine/Source/Runtime/Engine/Private/Materials/MaterialCachedData.cpp
@@ -854,26 +854,20 @@ void FMaterialCachedExpressionData::DoPostAnalyzeChecks(const UMaterial& BaseMat
 
 				const bool bAIsMaterialAttributes = TracedA.Expression->IsResultMaterialAttributes(TracedA.OutputIndex);
 				const bool bBIsMaterialAttributes = TracedB.Expression->IsResultMaterialAttributes(TracedB.OutputIndex);
+				const bool bAIsSubstrateAttributes = TracedA.Expression->IsResultSubstrateMaterial(TracedA.OutputIndex);
+				const bool bBIsSubstrateAttributes = TracedB.Expression->IsResultSubstrateMaterial(TracedB.OutputIndex);
 
 				if (bAIsMaterialAttributes != bBIsMaterialAttributes)
 				{
 					for (int32 PropertyIndex = 0; PropertyIndex < MP_MaterialAttributes; ++PropertyIndex)
 					{
 						// Substrate specific property handled separately.
-						if (PropertyIndex != MP_FrontMaterial)
+						if (PropertyIndex != MP_FrontMaterial || bAIsSubstrateAttributes != bBIsSubstrateAttributes)
 						{
 							SetPropertyConnected((EMaterialProperty)PropertyIndex);
 						}
 					}
 				}
-
-				const bool bAIsSubstrateAttributes = TracedA.Expression->IsResultSubstrateMaterial(TracedA.OutputIndex);
-				const bool bBIsSubstrateAttributes = TracedB.Expression->IsResultSubstrateMaterial(TracedB.OutputIndex);
-
-				if (bAIsSubstrateAttributes != bBIsSubstrateAttributes)
-				{
-					SetPropertyConnected(MP_FrontMaterial);
-				}
 			};
 
 			if (UMaterialExpressionStaticSwitch* StaticSwitch = Cast<UMaterialExpressionStaticSwitch>(ParameterData.Expression))
