- Add GetSafeDivisor to PreShader2 and replicate to the same ops as preshader

#Synced-CL 52631319

diff --git a/Engine/Source/Runtime/Engine/Private/Shader/Preshader2.cpp b/Engine/Source/Runtime/Engine/Private/Shader/Preshader2.cpp
--- a/Engine/Source/Runtime/Engine/Private/Shader/Preshader2.cpp
+++ b/Engine/Source/Runtime/Engine/Private/Shader/Preshader2.cpp
@@ -75,6 +75,23 @@ static const TCHAR* GPreshader2LinearSwizzles[4][5] =
 	{ TEXT(".?"), TEXT(".w"), TEXT(".??"), TEXT(".???"), TEXT(".????") },
 };
 
+// Matches GetSafeDivisor in ShaderValue.cpp -- clamps near-zero divisors to avoid inf/NaN.
+template<typename T>
+FORCENOINLINE static T GetSafeDivisor(T Number)
+{
+	if (FMath::Abs(Number) < (T)UE_DELTA)
+	{
+		return (Number < 0) ? -(T)UE_DELTA : (T)UE_DELTA;
+	}
+	return Number;
+}
+
+template<>
+FORCENOINLINE int32 GetSafeDivisor<int32>(int32 Number)
+{
+	return Number != 0 ? Number : 1;
+}
+
 struct FOpAbs { template<typename T> static T Call(T Value) { return FMath::Abs(Value); } };
 struct FOpACos { template<typename T> static T Call(T Value) { return FMath::Acos(Value); } };
 struct FOpASin { template<typename T> static T Call(T Value) { return FMath::Asin(Value); } };
@@ -94,15 +111,7 @@ struct FOpLogarithm10 { template<typename T> static T Call(T Value)
 	return FMath::Loge(Value) * LogToLog10;
 }};
 struct FOpNegate { template<typename T> static T Call(T Value) { return -Value; } };
-struct FOpReciprocal { template<typename T> static T Call(T Value)
-{
-	// Generate a safe divisor.
-	if (FMath::Abs(Value) < (T)UE_DELTA)
-	{
-		Value = (Value < 0) ? -(T)UE_DELTA : (T)UE_DELTA;
-	}
-	return (T)1 / Value;
-}};
+struct FOpReciprocal { template<typename T> static T Call(T Value) { return (T)1 / GetSafeDivisor(Value); } };
 struct FOpRound { template<typename T> static T Call(T Value) { return FMath::RoundToFloat(Value); } };
 struct FOpSaturate { template<typename T> static T Call(T Value) { return FMath::Clamp(Value, (T)0, (T)1); } };
 struct FOpSign { template<typename T> static T Call(T Value) { return FMath::Sign(Value); } };
@@ -114,7 +123,7 @@ struct FOpTruncate { template<typename T> static T Call(T Value) { return FMath:
 struct FOpAdd { template<typename T> static T Call(T A, T B) { return A + B; } };
 struct FOpSubtract { template<typename T> static T Call(T A, T B) { return A - B; } };
 struct FOpMultiply { template<typename T> static T Call(T A, T B) { return A * B; } };
-struct FOpDivide { template<typename T> static T Call(T A, T B) { return A / B; } };
+struct FOpDivide { template<typename T> static T Call(T A, T B) { return A / GetSafeDivisor(B); } };
 struct FOpMin { template<typename T> static T Call(T A, T B) { return FMath::Min(A, B); } };
 struct FOpMax { template<typename T> static T Call(T A, T B) { return FMath::Max(A, B); } };
 struct FOpATan2 { template<typename T> static T Call(T A, T B) { return FMath::Atan2(A, B); } };
@@ -122,9 +131,9 @@ struct FOpATan2 { template<typename T> static T Call(T A, T B) { return FMath::A
 // Note:  we are not defining a separate opcode for Fmod and Modulo.  Operations on integer will do Modulo, while operations on float will do Fmod.
 struct FOpFmod
 {
-	static float Call(float A, float B) { return FMath::Fmod(A, B); }
-	static double Call(double A, double B) { return FMath::Fmod(A, B); }
-	static int32 Call(int32 A, int32 B) { return A % B; }
+	static float Call(float A, float B) { return FMath::Fmod(A, GetSafeDivisor(B)); }
+	static double Call(double A, double B) { return FMath::Fmod(A, GetSafeDivisor(B)); }
+	static int32 Call(int32 A, int32 B) { return A % GetSafeDivisor(B); }
 };
 
 // Ops that aren't per-component.  These are dummy structures for template specialization or overloading, and don't need an operator implementation.
