Fixed default material on Nanite moving all ISKM instances to the origin.

#Synced-CL 52519763

diff --git a/Engine/Source/Runtime/Engine/Private/InstancedSkinnedMeshSceneProxyDesc.cpp b/Engine/Source/Runtime/Engine/Private/InstancedSkinnedMeshSceneProxyDesc.cpp
--- a/Engine/Source/Runtime/Engine/Private/InstancedSkinnedMeshSceneProxyDesc.cpp
+++ b/Engine/Source/Runtime/Engine/Private/InstancedSkinnedMeshSceneProxyDesc.cpp
@@ -45,13 +45,7 @@ FPrimitiveSceneProxy* FInstancedSkinnedMeshSceneProxyDesc::CreateSceneProxy(cons
 				Nanite::FMaterialAudit NaniteMaterials{};
 				const bool bSetMaterialUsageFlags = true;
 				Nanite::FNaniteResourcesHelper::AuditMaterials(&Desc, NaniteMaterials, bSetMaterialUsageFlags);
-
-				const bool bForceNaniteForMasked = false;
-				const bool bIsMaskingAllowed = Nanite::IsMaskingAllowed(Desc.GetWorld(), bForceNaniteForMasked);
-				if (NaniteMaterials.IsValid(bIsMaskingAllowed))
-				{
-					Result = ::new FNaniteInstancedSkinnedMeshSceneProxy(NaniteMaterials, Desc, SkelMeshRenderData);
-				}
+				Result = ::new FNaniteInstancedSkinnedMeshSceneProxy(NaniteMaterials, Desc, SkelMeshRenderData);
 			}
 			else if (!MeshObject->IsNaniteMesh())
 			{
