UCharacterMovementComponent::FindWaterLine does not null-test CharacterOnwer pointer, crashes

Also protect against potential nullptr in hit component

#rnx
#Synced-CL 52565968
#Synced-CL 52565968

#Synced-CL 52357887

diff --git a/Engine/Source/Runtime/Engine/Private/Components/CharacterMovementComponent.cpp b/Engine/Source/Runtime/Engine/Private/Components/CharacterMovementComponent.cpp
--- a/Engine/Source/Runtime/Engine/Private/Components/CharacterMovementComponent.cpp
+++ b/Engine/Source/Runtime/Engine/Private/Components/CharacterMovementComponent.cpp
@@ -4784,12 +4784,18 @@ FVector UCharacterMovementComponent::FindWaterLine(FVector InWater, FVector Outo
 {
 	FVector Result = OutofWater;
 
+	if (!IsValid(CharacterOwner))
+	{
+		return Result;
+	}
+
 	TArray<FHitResult> Hits;
 	GetWorld()->LineTraceMultiByChannel(Hits, OutofWater, InWater, UpdatedComponent->GetCollisionObjectType(), FCollisionQueryParams(SCENE_QUERY_STAT(FindWaterLine), true, CharacterOwner));
 
 	for (const FHitResult& Check : Hits)
 	{
-		if ( !CharacterOwner->IsOwnedBy(Check.HitObjectHandle.FetchActor()) && !Check.Component.Get()->IsWorldGeometry() )
+		const UPrimitiveComponent* HitComponent = Check.Component.Get();
+		if (!CharacterOwner->IsOwnedBy(Check.HitObjectHandle.FetchActor()) && HitComponent && !HitComponent->IsWorldGeometry())
 		{
 			APhysicsVolume *W = Check.HitObjectHandle.FetchActor<APhysicsVolume>();
 			if ( W && W->bWaterVolume )
