New material translator: Use LWC for world position derivatives according to whether world position is camera relative.
Fixes shader compilation errors in lowered code for camera relative world position.

#rnx
#Synced-CL 52659694

diff --git a/Engine/Source/Runtime/Engine/Private/Materials/MaterialExpressionsToMIR.cpp b/Engine/Source/Runtime/Engine/Private/Materials/MaterialExpressionsToMIR.cpp
--- a/Engine/Source/Runtime/Engine/Private/Materials/MaterialExpressionsToMIR.cpp
+++ b/Engine/Source/Runtime/Engine/Private/Materials/MaterialExpressionsToMIR.cpp
@@ -2209,8 +2209,15 @@ static FValueRef EmitWorldPosition(MIR::FEmitter& Em, EWorldPositionIncludedOffs
 					break;
 
 				default:
-					// Expression emitter assumes these to be LWC for analytic derivative evaluation, so we promote them on load
-					Printer << TEXTVIEW("WSPromote(Parameters.WorldPosition_") << MIR::ExternTag_DD << TEXT(')');
+					// Camera-relative variants use float vectors others require promotion to LWC.
+					if (WorldPositionShaderOffset == WPT_Default || WorldPositionShaderOffset == WPT_ExcludeAllShaderOffsets)
+					{
+						Printer << TEXTVIEW("WSPromote(Parameters.WorldPosition_") << MIR::ExternTag_DD << TEXT(')');
+					}
+					else
+					{
+						Printer << TEXTVIEW("Parameters.WorldPosition_") << MIR::ExternTag_DD;
+					}
 					break;
 			}
 		}
