Add bUsesOpacity to new translator compilation output.
It was only added to the old translator path recently.

#rnx
#Synced-CL 52659694

diff --git a/Engine/Source/Runtime/Engine/Private/Materials/MaterialIRModuleBuilder.cpp b/Engine/Source/Runtime/Engine/Private/Materials/MaterialIRModuleBuilder.cpp
--- a/Engine/Source/Runtime/Engine/Private/Materials/MaterialIRModuleBuilder.cpp
+++ b/Engine/Source/Runtime/Engine/Private/Materials/MaterialIRModuleBuilder.cpp
@@ -1280,6 +1280,7 @@ struct FMaterialIRModuleBuilderImpl
 		CompilationOutput.bUsesPixelDepthOffset     = MaterialPropertyHasNonZeroValue(MP_PixelDepthOffset);
 		CompilationOutput.bUsesAnisotropy           = MaterialPropertyHasNonZeroValue(MP_Anisotropy);
 		CompilationOutput.bUsesOpacityMask          = Module->IsMaterialPropertyUsed(MP_OpacityMask);
+		CompilationOutput.bUsesOpacity              = Module->IsMaterialPropertyUsed(MP_Opacity);
 		
 		// Check if displacement is actually used (disable if it was left at the default, invalid constant value of -1)
 		CompilationOutput.bUsesDisplacement 	 |= DoesPlatformSupportNanite(Module->GetShaderPlatform()) && Builder->MaterialInterface->IsTessellationEnabled() && Module->IsMaterialPropertyUsed(MP_Displacement);
