Fix uninitialized property values in shader code generated by the material translator code.
This would cause DX11 shader compilation to crash.
Issue comes from recent change to avoid compiling unused properties. But even when they are unused, leaving them uninitialized is bad for the DX11 compiler.
We now generate the necessary property code with default values.

#rnx
#Synced-CL 52659694

diff --git a/Engine/Source/Runtime/Engine/Private/Materials/MaterialIRModuleBuilder.cpp b/Engine/Source/Runtime/Engine/Private/Materials/MaterialIRModuleBuilder.cpp
--- a/Engine/Source/Runtime/Engine/Private/Materials/MaterialIRModuleBuilder.cpp
+++ b/Engine/Source/Runtime/Engine/Private/Materials/MaterialIRModuleBuilder.cpp
@@ -715,8 +715,8 @@ struct FMaterialIRModuleBuilderImpl
 
 		for (EMaterialProperty Property : MaterialAttributesAggregate::GetMaterialProperties())
 		{
-			// Ignore deactivated properties except ShadingModel which we always need to emit.
-			if (Property != MP_ShadingModel && !PropertyIsActiveInMaterial(Property))
+			// Skip CustomizedUV slots that aren't enabled as there is no output to emit for these.
+			if (IsMaterialPropertyCustomizedUV(Property) && Property >= MP_CustomizedUVs0 + Emitter.BaseMaterial->NumCustomizedUVs)
 			{
 				continue;
 			}
@@ -728,17 +728,24 @@ struct FMaterialIRModuleBuilderImpl
 			// This holds the value being set to this property.
 			MIR::FValue* PropertyValue = nullptr;
 
-			// If the material attributes value is valid, extract this property attribute from the material attributes aggregate value
-			// and manually flow it into the this property material expression input pin.
-			if (MaterialAttributesValue)
+			// Only evaluate expressions for properties that are active for the current blend mode.
+			// Inactive properties will skip expression evaluation and leave PropertyValue as nullptr. Then we fall through to the default value path below. 
+			// SetMaterialOutput is still emitted so the shader variable is always initialized.
+			bool const bIsPropertyActive = (Property == MP_ShadingModel) || Emitter.MaterialInterface->IsPropertyActive(Property);
+			if (bIsPropertyActive)
 			{
-				PropertyValue = Emitter.Subscript(MaterialAttributesValue, MaterialAttributesAggregate::MaterialPropertyToAttributeIndex(Property));
-			}
-			else
-			{
-				// Otherwise grab the value from the individual attribute pin.
-				FlowValueThroughConnection(InputDesc.Input);
-				PropertyValue = MIR::Internal::FetchValueFromExpressionInput(this, InputDesc.Input);
+				// If the material attributes value is valid, extract this property attribute from the material attributes aggregate value
+				// and manually flow it into the this property material expression input pin.
+				if (MaterialAttributesValue)
+				{
+					PropertyValue = Emitter.Subscript(MaterialAttributesValue, MaterialAttributesAggregate::MaterialPropertyToAttributeIndex(Property));
+				}
+				else
+				{
+					// Otherwise grab the value from the individual attribute pin.
+					FlowValueThroughConnection(InputDesc.Input);
+					PropertyValue = MIR::Internal::FetchValueFromExpressionInput(this, InputDesc.Input);
+				}
 			}
 			
 			// Determine the value to assign to this material attribute "SetMaterialOutput" instruction
