Migrate UE_LOG to UE_LOGF.

This should reduce memory usage. The conversion was performed automatically via ConvertUELog.py (in Engine/Extras/UnsupportedTools).

diff --git a/Engine/Plugins/Animation/IKRig/Source/IKRigEditor/Private/RetargetEditor/IKRetargetBatchOperation.cpp b/Engine/Plugins/Animation/IKRig/Source/IKRigEditor/Private/RetargetEditor/IKRetargetBatchOperation.cpp
--- a/Engine/Plugins/Animation/IKRig/Source/IKRigEditor/Private/RetargetEditor/IKRetargetBatchOperation.cpp
+++ b/Engine/Plugins/Animation/IKRig/Source/IKRigEditor/Private/RetargetEditor/IKRetargetBatchOperation.cpp
@@ -289,7 +289,7 @@ void UIKRetargetBatchOperation::RetargetAssets(
 	Processor.Initialize(Params);
 	if (!Processor.IsInitialized())
 	{
-		UE_LOG(LogTemp, Warning, TEXT("Unable to initialize the IK Retargeter. Newly created animations were not retargeted!"));
+		UE_LOGF(LogTemp, Warning, "Unable to initialize the IK Retargeter. Newly created animations were not retargeted!");
 		return;
 	}
 
@@ -532,7 +532,7 @@ void UIKRetargetBatchOperation::ConvertAnimation(
 	{
 		for (const FName& MissingCurve : Pair.Value)
 		{
-			UE_LOG(LogTemp, Warning, TEXT("IK Retarget Batch Operation: Missing speed curve, %s on sequence %s."), *MissingCurve.ToString(), *Pair.Key);	
+			UE_LOGF(LogTemp, Warning, "IK Retarget Batch Operation: Missing speed curve, %ls on sequence %ls.", *MissingCurve.ToString(), *Pair.Key);	
 		}
 	}
 }
@@ -623,7 +623,7 @@ void UIKRetargetBatchOperation::ApplyCurveOps(const FIKRetargetBatchOperationCon
 			
 		if(!bLoadedData)
 		{
-			UE_LOG(LogTemp, Warning, TEXT("No animation curves found for retargeting in Anim Sequence: %s"), *SourceSequence->GetName());
+			UE_LOGF(LogTemp, Warning, "No animation curves found for retargeting in Anim Sequence: %ls", *SourceSequence->GetName());
 		}
 
 		TargetSeqController.CloseBracket(bShouldTransact);
@@ -765,7 +765,7 @@ void UIKRetargetBatchOperation::NotifyUserOfResults(
 		// log details of what assets were created
 		for (const UObject* NewAsset : NewAssets)
 		{
-			UE_LOG(LogTemp, Display, TEXT("Duplicate and Retarget - New Asset Created: %s"), *NewAsset->GetName());
+			UE_LOGF(LogTemp, Display, "Duplicate and Retarget - New Asset Created: %ls", *NewAsset->GetName());
 		}
 	
 		// notify user that retarget completed
@@ -937,14 +937,14 @@ void UIKRetargetBatchOperation::RunRetarget(FIKRetargetBatchOperationContext& Co
 	const int32 NumAssets = GenerateAssetLists(Context);
 	if (NumAssets == 0)
 	{
-		UE_LOG(LogTemp, Warning, TEXT("Batch retarget aborted. No animation assets were specified."));
+		UE_LOGF(LogTemp, Warning, "Batch retarget aborted. No animation assets were specified.");
 		return;
 	}
 
 	// validate a retarget asset was provided
 	if (!Context.IKRetargetAsset)
 	{
-		UE_LOG(LogTemp, Warning, TEXT("Batch retarget aborted. No IK Retargeter asset was specified."));
+		UE_LOGF(LogTemp, Warning, "Batch retarget aborted. No IK Retargeter asset was specified.");
 		return;
 	}
 
@@ -961,7 +961,7 @@ void UIKRetargetBatchOperation::RunRetarget(FIKRetargetBatchOperationContext& Co
 		
 		if (!Context.SourceMesh)
 		{
-			UE_LOG(LogTemp, Warning, TEXT("Batch retarget aborted. No source mesh was specified and the source IK Rig did not have one. "));
+			UE_LOGF(LogTemp, Warning, "Batch retarget aborted. No source mesh was specified and the source IK Rig did not have one. ");
 			return;	
 		}
 	}
@@ -978,7 +978,7 @@ void UIKRetargetBatchOperation::RunRetarget(FIKRetargetBatchOperationContext& Co
 		
 		if (!Context.TargetMesh)
 		{
-			UE_LOG(LogTemp, Warning, TEXT("Batch retarget aborted. No target mesh was specified and the target IK Rig did not have one. "));
+			UE_LOGF(LogTemp, Warning, "Batch retarget aborted. No target mesh was specified and the target IK Rig did not have one. ");
 			return;
 		}
 	}
diff --git a/Engine/Plugins/Animation/IKRig/Source/IKRigEditor/Private/RetargetEditor/IKRetargeterController.cpp b/Engine/Plugins/Animation/IKRig/Source/IKRigEditor/Private/RetargetEditor/IKRetargeterController.cpp
--- a/Engine/Plugins/Animation/IKRig/Source/IKRigEditor/Private/RetargetEditor/IKRetargeterController.cpp
+++ b/Engine/Plugins/Animation/IKRig/Source/IKRigEditor/Private/RetargetEditor/IKRetargeterController.cpp
@@ -199,7 +199,7 @@ int32 UIKRetargeterController::AddRetargetOp(const FString InIKRetargetOpType) c
 	UScriptStruct* OpType = FindObject<UScriptStruct>(nullptr, *InIKRetargetOpType);
 	if (!OpType)
 	{
-		UE_LOG(LogIKRigEditor, Warning, TEXT("Specified retarget op type was not found, %s."), *InIKRetargetOpType);
+		UE_LOGF(LogIKRigEditor, Warning, "Specified retarget op type was not found, %ls.", *InIKRetargetOpType);
 		return INDEX_NONE;
 	}
 
@@ -211,7 +211,7 @@ bool UIKRetargeterController::RemoveRetargetOp(const int32 InOpIndex) const
 	FIKRetargetOpBase* Op = GetRetargetOpByIndex(InOpIndex);
 	if (!Op)
 	{
-		UE_LOG(LogIKRigEditor, Warning, TEXT("Retarget Op not removed. Invalid index, %d."), InOpIndex);
+		UE_LOGF(LogIKRigEditor, Warning, "Retarget Op not removed. Invalid index, %d.", InOpIndex);
 		return false;
 	}
 
@@ -400,7 +400,7 @@ void UIKRetargeterController::RunOpInitialSetup(const int32 InOpIndex) const
 	FIKRetargetOpBase* Op = GetRetargetOpByIndex(InOpIndex);
 	if (!Op)
 	{
-		UE_LOG(LogIKRigEditor, Warning, TEXT("Retarget Op initial setup skipped. Invalid index, %d."), InOpIndex);
+		UE_LOGF(LogIKRigEditor, Warning, "Retarget Op initial setup skipped. Invalid index, %d.", InOpIndex);
 		return;
 	}
 	
@@ -602,19 +602,19 @@ bool UIKRetargeterController::MoveRetargetOpInStack(int32 OpToMoveIndex, int32 T
 	
 	if (!Asset->RetargetOps.IsValidIndex(OpToMoveIndex))
 	{
-		UE_LOG(LogIKRigEditor, Warning, TEXT("Retarget Op not moved. Invalid source index, %d."), OpToMoveIndex);
+		UE_LOGF(LogIKRigEditor, Warning, "Retarget Op not moved. Invalid source index, %d.", OpToMoveIndex);
 		return false;
 	}
 
 	if (!Asset->RetargetOps.IsValidIndex(TargetIndex))
 	{
-		UE_LOG(LogIKRigEditor, Warning, TEXT("Retarget Op not moved. Invalid target index, %d."), TargetIndex);
+		UE_LOGF(LogIKRigEditor, Warning, "Retarget Op not moved. Invalid target index, %d.", TargetIndex);
 		return false;
 	}
 
 	if (OpToMoveIndex == TargetIndex)
 	{
-		UE_LOG(LogIKRigEditor, Warning, TEXT("Retarget Op not moved. Source and target index cannot be the same."));
+		UE_LOGF(LogIKRigEditor, Warning, "Retarget Op not moved. Source and target index cannot be the same.");
 		return false;
 	}
 
@@ -639,7 +639,7 @@ bool UIKRetargeterController::SetRetargetOpEnabled(int32 RetargetOpIndex, bool b
 {
 	if (!Asset->RetargetOps.IsValidIndex(RetargetOpIndex))
 	{
-		UE_LOG(LogIKRigEditor, Warning, TEXT("Retarget op not found. Invalid index, %d."), RetargetOpIndex);
+		UE_LOGF(LogIKRigEditor, Warning, "Retarget op not found. Invalid index, %d.", RetargetOpIndex);
 		return false;
 	}
 	
@@ -654,7 +654,7 @@ bool UIKRetargeterController::GetRetargetOpEnabled(int32 RetargetOpIndex) const
 {
 	if (!Asset->RetargetOps.IsValidIndex(RetargetOpIndex))
 	{
-		UE_LOG(LogIKRigEditor, Warning, TEXT("Invalid retarget op index, %d."), RetargetOpIndex);
+		UE_LOGF(LogIKRigEditor, Warning, "Invalid retarget op index, %d.", RetargetOpIndex);
 		return false;
 	}
 
@@ -781,13 +781,13 @@ int32 UIKRetargeterController::AddRetargetOp(const UScriptStruct* InRetargetOpTy
 {
 	if (!InRetargetOpType)
 	{
-		UE_LOG(LogIKRigEditor, Warning, TEXT("Could not add retarget operation. No type specified."));
+		UE_LOGF(LogIKRigEditor, Warning, "Could not add retarget operation. No type specified.");
 		return INDEX_NONE;
 	}
 
 	if (!InRetargetOpType->IsChildOf(FIKRetargetOpBase::StaticStruct()))
 	{
-		UE_LOG(LogIKRigEditor, Warning, TEXT("Could not add retarget operations. Invalid op-type specified. Must be child of FIKRetargetOpBase."));
+		UE_LOGF(LogIKRigEditor, Warning, "Could not add retarget operations. Invalid op-type specified. Must be child of FIKRetargetOpBase.");
 		return INDEX_NONE;
 	}
 	
@@ -800,7 +800,7 @@ int32 UIKRetargeterController::AddRetargetOp(const UScriptStruct* InRetargetOpTy
 			const FInstancedStruct& Op = Asset->RetargetOps[OpIndex];
 			if (Op.GetScriptStruct() == InRetargetOpType)
 			{
-				UE_LOG(LogIKRigEditor, Warning, TEXT("Op not added. It is a singleton and the stack already contains an op of that type."));
+				UE_LOGF(LogIKRigEditor, Warning, "Op not added. It is a singleton and the stack already contains an op of that type.");
 				return INDEX_NONE;
 			}
 		}
@@ -1208,7 +1208,7 @@ FName UIKRetargeterController::DuplicateRetargetPose( const FName PoseToDuplicat
 	TMap<FName, FIKRetargetPose>& Poses = GetRetargetPoses(SourceOrTarget);
 	if (!Poses.Contains(PoseToDuplicate))
 	{
-		UE_LOG(LogIKRigEditor, Warning, TEXT("Trying to duplicate pose that does not exist, %s."), *PoseToDuplicate.ToString());
+		UE_LOGF(LogIKRigEditor, Warning, "Trying to duplicate pose that does not exist, %ls.", *PoseToDuplicate.ToString());
 		return NAME_None; // cannot duplicate pose that doesn't exist
 	}
 
@@ -1237,14 +1237,14 @@ bool UIKRetargeterController::RenameRetargetPose(const FName OldPoseName, const
 	// does the old pose exist?
 	if (!GetRetargetPoses(SourceOrTarget).Contains(OldPoseName))
 	{
-		UE_LOG(LogIKRigEditor, Warning, TEXT("Trying to rename pose that does not exist, %s."), *OldPoseName.ToString());
+		UE_LOGF(LogIKRigEditor, Warning, "Trying to rename pose that does not exist, %ls.", *OldPoseName.ToString());
 		return false;
 	}
 
 	// do not allow renaming the default pose (this is disallowed from the UI, but must be done here as well for API usage)
 	if (OldPoseName == UIKRetargeter::GetDefaultPoseName())
 	{
-		UE_LOG(LogIKRigEditor, Warning, TEXT("Trying to rename the default pose. This is not allowed."));
+		UE_LOGF(LogIKRigEditor, Warning, "Trying to rename the default pose. This is not allowed.");
     	return false;
 	}
 
@@ -1332,7 +1332,7 @@ bool UIKRetargeterController::SetCurrentRetargetPose(FName NewCurrentPose, const
 	const TMap<FName, FIKRetargetPose>& Poses = GetRetargetPoses(SourceOrTarget);
 	if (!Poses.Contains(NewCurrentPose))
 	{
-		UE_LOG(LogIKRigEditor, Warning, TEXT("Trying to set current pose to a pose that does not exist, %s."), *NewCurrentPose.ToString());
+		UE_LOGF(LogIKRigEditor, Warning, "Trying to set current pose to a pose that does not exist, %ls.", *NewCurrentPose.ToString());
 		return false;
 	}
 	
diff --git a/Engine/Plugins/ChaosClothAssetDataflowNodes/Source/ChaosClothAssetDataflowNodes/Private/ChaosClothAsset/RemeshNode.cpp b/Engine/Plugins/ChaosClothAssetDataflowNodes/Source/ChaosClothAssetDataflowNodes/Private/ChaosClothAsset/RemeshNode.cpp
--- a/Engine/Plugins/ChaosClothAssetDataflowNodes/Source/ChaosClothAssetDataflowNodes/Private/ChaosClothAsset/RemeshNode.cpp
+++ b/Engine/Plugins/ChaosClothAssetDataflowNodes/Source/ChaosClothAssetDataflowNodes/Private/ChaosClothAsset/RemeshNode.cpp
@@ -1696,7 +1696,7 @@ namespace UE::Chaos::ClothAsset::Private
 
 				if (!ProjectionTarget2D.IsTriangle(NearestTriangle))
 				{
-					UE_LOG(LogChaosClothAssetDataflowNodes, Verbose, TEXT("UE::Chaos::ClothAsset::RemeshNode: Failed to find a nearest triangle on the input 2D Sim mesh. Ensure the 2D sim mesh has no empty or degenerate triangles before remeshing"));
+					UE_LOGF(LogChaosClothAssetDataflowNodes, Verbose, "UE::Chaos::ClothAsset::RemeshNode: Failed to find a nearest triangle on the input 2D Sim mesh. Ensure the 2D sim mesh has no empty or degenerate triangles before remeshing");
 					continue;
 				}
 
@@ -1705,14 +1705,14 @@ namespace UE::Chaos::ClothAsset::Private
 
 				if (!ProjectionTarget3D.IsTriangle(NearestTriangle))
 				{
-					UE_LOG(LogChaosClothAssetDataflowNodes, Verbose, TEXT("UE::Chaos::ClothAsset::RemeshNode: A triangle ID that is valid on the 2D Sim mesh was not found on the 3D sim mesh. Check Simulation faces"));
+					UE_LOGF(LogChaosClothAssetDataflowNodes, Verbose, "UE::Chaos::ClothAsset::RemeshNode: A triangle ID that is valid on the 2D Sim mesh was not found on the 3D sim mesh. Check Simulation faces");
 					continue;
 				}
 
 				const FIndex3i Triangle3D = ProjectionTarget3D.GetTriangle(NearestTriangle);
 				if (!ProjectionTarget3D.IsVertex(Triangle3D[0]) || !ProjectionTarget3D.IsVertex(Triangle3D[1]) || !ProjectionTarget3D.IsVertex(Triangle3D[2]))
 				{
-					UE_LOG(LogChaosClothAssetDataflowNodes, Verbose, TEXT("UE::Chaos::ClothAsset::RemeshNode: A triangle ID is valid on the 3D Sim mesh has vertex indices that are not valid. Check Simulation faces"));
+					UE_LOGF(LogChaosClothAssetDataflowNodes, Verbose, "UE::Chaos::ClothAsset::RemeshNode: A triangle ID is valid on the 3D Sim mesh has vertex indices that are not valid. Check Simulation faces");
 					continue;
 				}
 
diff --git a/Engine/Plugins/ChaosVD/Source/ChaosVD/Private/Trace/ChaosVDTraceProvider.cpp b/Engine/Plugins/ChaosVD/Source/ChaosVD/Private/Trace/ChaosVDTraceProvider.cpp
--- a/Engine/Plugins/ChaosVD/Source/ChaosVD/Private/Trace/ChaosVDTraceProvider.cpp
+++ b/Engine/Plugins/ChaosVD/Source/ChaosVD/Private/Trace/ChaosVDTraceProvider.cpp
@@ -151,7 +151,7 @@ int32 FChaosVDTraceProvider::RemapSolverID(int32 SolverID)
 		InternalRecording->ReserveSolverID_AssumesLocked(RemappedSolverID);
 	}
 
-	UE_LOG(LogChaosVDEditor, Verbose, TEXT("[%hs] Remapped solver id from [%d] to [%d]."), __func__, SolverID, RemappedSolverID);
+	UE_LOGF(LogChaosVDEditor, Verbose, "[%s] Remapped solver id from [%d] to [%d].", __func__, SolverID, RemappedSolverID);
 
 	RemappedSolversIDs.Emplace(SolverID, RemappedSolverID);
 
@@ -168,7 +168,7 @@ int32 FChaosVDTraceProvider::GetRemappedSolverID(int32 SolverID)
 	}
 	else
 	{
-		UE_LOG(LogChaosVDEditor, Verbose, TEXT("[%hs] Failed to get remapped solver id [%d]. Data that references the invalid solver id will be ignored."), __func__, SolverID);
+		UE_LOGF(LogChaosVDEditor, Verbose, "[%s] Failed to get remapped solver id [%d]. Data that references the invalid solver id will be ignored.", __func__, SolverID);
 		RemappedSolverID = INDEX_NONE; 
 	}
 
@@ -264,7 +264,7 @@ FChaosVDSolverFrameData* FChaosVDTraceProvider::GetCurrentSolverFrame(const int3
 
 	if (bIsInvalidSolverID)
 	{
-		UE_LOG(LogChaosVDEditor, Verbose, TEXT("[%hs] was called with an invalid solver ID. Data that references the invalid solver id will be ignored."), __func__);
+		UE_LOGF(LogChaosVDEditor, Verbose, "[%s] was called with an invalid solver ID. Data that references the invalid solver id will be ignored.", __func__);
 		return nullptr;
 	}
 
@@ -348,8 +348,8 @@ bool FChaosVDTraceProvider::ProcessBinaryData(const int32 DataID)
 				{
 					if (DataProcessorPtr->ShouldAlwaysSkip())
 					{
-						UE_LOG(LogChaosVDEditor, Log,
-							TEXT("[%s] Silently skipping [%s] — processor opted out via ShouldAlwaysSkip()."),
+						UE_LOGF(LogChaosVDEditor, Log,
+							"[%ls] Silently skipping [%ls] — processor opted out via ShouldAlwaysSkip().",
 							ANSI_TO_TCHAR(__FUNCTION__), *UnprocessedData->TypeName);
 						UnprocessedDataByID.Remove(DataID);
 						DataProcessedSoFarCounter++;
@@ -364,8 +364,8 @@ bool FChaosVDTraceProvider::ProcessBinaryData(const int32 DataID)
 						{
 							// Treat deliberately-skipped data as a clean success so the rest of the load continues normally.
 							TypesSkippedInSafeMode.Add(UnprocessedData->TypeName);
-							UE_LOG(LogChaosVDEditor, Log,
-								TEXT("[%s] Skipping non-backwards-compatible data type [%s]  - Safe Loading Mode is enabled. See CVD General Settings to change this."),
+							UE_LOGF(LogChaosVDEditor, Log,
+								"[%ls] Skipping non-backwards-compatible data type [%ls]  - Safe Loading Mode is enabled. See CVD General Settings to change this.",
 								ANSI_TO_TCHAR(__FUNCTION__), *UnprocessedData->TypeName);
 							UnprocessedDataByID.Remove(DataID);
 							DataProcessedSoFarCounter++;
@@ -375,8 +375,8 @@ bool FChaosVDTraceProvider::ProcessBinaryData(const int32 DataID)
 						{
 							// Record that we're about to process potentially unsafe data so we can warn after loading.
 							UnsafeTypesProcessed.Add(UnprocessedData->TypeName);
-							UE_LOG(LogChaosVDEditor, Log,
-								TEXT("[%s] Processing non-backwards-compatible data type [%s] with Safe Loading Mode disabled. Enable Safe Loading Mode in CVD General Settings if CVD crashes unexpectedly."),
+							UE_LOGF(LogChaosVDEditor, Log,
+								"[%ls] Processing non-backwards-compatible data type [%ls] with Safe Loading Mode disabled. Enable Safe Loading Mode in CVD General Settings if CVD crashes unexpectedly.",
 								ANSI_TO_TCHAR(__FUNCTION__), *UnprocessedData->TypeName);
 						}
 					}
@@ -389,7 +389,7 @@ bool FChaosVDTraceProvider::ProcessBinaryData(const int32 DataID)
 					}
 					else
 					{
-						UE_LOG(LogChaosVDEditor, Verbose, TEXT("[%s] Failed to serialize Binary Data with ID [%d] | Type [%s]"), ANSI_TO_TCHAR(__FUNCTION__), DataID, *UnprocessedData->TypeName);
+						UE_LOGF(LogChaosVDEditor, Verbose, "[%ls] Failed to serialize Binary Data with ID [%d] | Type [%ls]", ANSI_TO_TCHAR(__FUNCTION__), DataID, *UnprocessedData->TypeName);
 						TypesFailedToSerialize.Add(UnprocessedData->TypeName);
 						DataProcessedSoFarCounter++;
 					}
@@ -452,7 +452,7 @@ void FChaosVDTraceProvider::HandleAnalysisComplete()
 
 	UnprocessedDataByID.Reset();
 
-	UE_LOG(LogChaosVDEditor, Log, TEXT("Trace Analysis complete for session [%s] | Calculating data loaded stats..."), Session.GetName());
+	UE_LOGF(LogChaosVDEditor, Log, "Trace Analysis complete for session [%ls] | Calculating data loaded stats...", Session.GetName());
 
 	static const FNumberFormattingOptions SizeFormattingOptions = FNumberFormattingOptions().SetMinimumFractionalDigits(2).SetMaximumFractionalDigits(2);
 
@@ -463,7 +463,7 @@ void FChaosVDTraceProvider::HandleAnalysisComplete()
 		{
 			uint64 ProcessedBytes = DataProcessor.Value->GetProcessedBytes();
 			TotalBytes += ProcessedBytes;
-			UE_LOG(LogChaosVDEditor, Log, TEXT("Data loaded for type [%s]  => [%s] "), DataProcessor.Key.IsEmpty() ? TEXT("Invalid") : DataProcessor.Key.GetData(), *FText::AsMemory(ProcessedBytes, &SizeFormattingOptions,nullptr, EMemoryUnitStandard::IEC).ToString());
+			UE_LOGF(LogChaosVDEditor, Log, "Data loaded for type [%ls]  => [%ls] ", DataProcessor.Key.IsEmpty() ? TEXT("Invalid") : DataProcessor.Key.GetData(), *FText::AsMemory(ProcessedBytes, &SizeFormattingOptions,nullptr, EMemoryUnitStandard::IEC).ToString());
 		}
 	}
 
@@ -474,10 +474,10 @@ void FChaosVDTraceProvider::HandleAnalysisComplete()
 		int32 NumOfGameFramesProcessed = Recording->GetAvailableGameFramesNumber();
 		double AvgTimePerFrameSeconds = TotalTimeProcessingFrames / NumOfGameFramesProcessed;
 		
-		UE_LOG(LogChaosVDEditor, Log, TEXT(" [%d] Game frames Processed at [%f] ms per frame on average"), NumOfGameFramesProcessed, AvgTimePerFrameSeconds * 1000.0);
+		UE_LOGF(LogChaosVDEditor, Log, " [%d] Game frames Processed at [%f] ms per frame on average", NumOfGameFramesProcessed, AvgTimePerFrameSeconds * 1000.0);
 	}
 
-	UE_LOG(LogChaosVDEditor, Log, TEXT("Total size of loaded data => [%s]"), *FText::AsMemory(TotalBytes, &SizeFormattingOptions,nullptr, EMemoryUnitStandard::IEC).ToString());
+	UE_LOGF(LogChaosVDEditor, Log, "Total size of loaded data => [%ls]", *FText::AsMemory(TotalBytes, &SizeFormattingOptions,nullptr, EMemoryUnitStandard::IEC).ToString());
 }
 
 void FChaosVDTraceProvider::ShowSafeModeSkippedDialog()
diff --git a/Engine/Plugins/Experimental/Animation/RelativeIKOp/Source/RelativeIKOp/Private/RelativeIKOp.cpp b/Engine/Plugins/Experimental/Animation/RelativeIKOp/Source/RelativeIKOp/Private/RelativeIKOp.cpp
--- a/Engine/Plugins/Experimental/Animation/RelativeIKOp/Source/RelativeIKOp/Private/RelativeIKOp.cpp
+++ b/Engine/Plugins/Experimental/Animation/RelativeIKOp/Source/RelativeIKOp/Private/RelativeIKOp.cpp
@@ -2204,12 +2204,12 @@ FVector FIKRetargetRelativeIKOp::CalcReferenceShapeScale3D(FKShapeElem* ShapeEle
 
 	case EAggCollisionShape::Convex:
 		{
-			// UE_LOG(LogAnimation, Warning, TEXT("Unsupported: Shape is a Convex"));
+			// UE_LOGF(LogAnimation, Warning, "Unsupported: Shape is a Convex");
 			break;
 		}
 	default:
 		{
-			// UE_LOG(LogAnimation, Warning, TEXT("Unknown or unsupported shape type"));
+			// UE_LOGF(LogAnimation, Warning, "Unknown or unsupported shape type");
 			break;
 		}
 	}
diff --git a/Engine/Plugins/Experimental/EditorDataStorage/Source/TedsCore/Private/Queries/TypedElementExtendedQueryStore.cpp b/Engine/Plugins/Experimental/EditorDataStorage/Source/TedsCore/Private/Queries/TypedElementExtendedQueryStore.cpp
--- a/Engine/Plugins/Experimental/EditorDataStorage/Source/TedsCore/Private/Queries/TypedElementExtendedQueryStore.cpp
+++ b/Engine/Plugins/Experimental/EditorDataStorage/Source/TedsCore/Private/Queries/TypedElementExtendedQueryStore.cpp
@@ -83,7 +83,7 @@ namespace UE::Editor::DataStorage
 			}
 			StringBuilder.Append(TEXT("}"));
 		}
-		UE_LOG(LogEditorDataStorage, Error, TEXT("%s"), StringBuilder.ToString());
+		UE_LOGF(LogEditorDataStorage, Error, "%ls", StringBuilder.ToString());
 	}
 	
 	bool FTedsQueryArchetypeMatchOverride::Match(const FMassArchetypeCompositionDescriptor& CompositionDescriptor) const
diff --git a/Engine/Plugins/Experimental/EditorDataStorage/Source/TedsCore/Private/TableManager.cpp b/Engine/Plugins/Experimental/EditorDataStorage/Source/TedsCore/Private/TableManager.cpp
--- a/Engine/Plugins/Experimental/EditorDataStorage/Source/TedsCore/Private/TableManager.cpp
+++ b/Engine/Plugins/Experimental/EditorDataStorage/Source/TedsCore/Private/TableManager.cpp
@@ -62,7 +62,7 @@ namespace UE::Editor::DataStorage
 			RemoveTableTag(Composition);
 			Type = ETableType::Derivative;
 
-			UE_LOG(LogEditorDataStorage, VeryVerbose, TEXT("New table derivative '%s' from template '%s"), 
+			UE_LOGF(LogEditorDataStorage, VeryVerbose, "New table derivative '%ls' from template '%ls", 
 				*TableName.ToString(), 
 				*TableInfo.Name.ToString());
 		}
@@ -70,7 +70,7 @@ namespace UE::Editor::DataStorage
 		{
 			Type = ETableType::Declared;
 
-			UE_LOG(LogEditorDataStorage, VeryVerbose, TEXT("New table declared '%s'"), *TableName.ToString());
+			UE_LOGF(LogEditorDataStorage, VeryVerbose, "New table declared '%ls'", *TableName.ToString());
 		}
 
 		FMassArchetypeCreationParams ArchetypeCreationParams;
@@ -345,7 +345,7 @@ namespace UE::Editor::DataStorage
 						});
 					MassLookup.Add(Archetype, NewTable);
 					Tables[ConvertTableHandleToIndex(Parent)].Variants.Add(NewTable);
-					UE_LOG(LogEditorDataStorage, VeryVerbose, TEXT("New relation table variant for '%s'"), *ParentFName.ToString());
+					UE_LOGF(LogEditorDataStorage, VeryVerbose, "New relation table variant for '%ls'", *ParentFName.ToString());
 				}
 				return;
 			}
@@ -366,7 +366,7 @@ namespace UE::Editor::DataStorage
 			TableHandle Variant = MassLookup.Add(Archetype, NewTable);
 			Tables[ConvertTableHandleToIndex(Parent)].Variants.Add(NewTable);
 
-			UE_LOG(LogEditorDataStorage, VeryVerbose, TEXT("New table variant for '%s'"), *ParentName);
+			UE_LOGF(LogEditorDataStorage, VeryVerbose, "New table variant for '%ls'", *ParentName);
 		}
 	}
 
diff --git a/Engine/Plugins/Experimental/EditorDataStorage/Source/TedsCore/Private/TypedElementDatabase.cpp b/Engine/Plugins/Experimental/EditorDataStorage/Source/TedsCore/Private/TypedElementDatabase.cpp
--- a/Engine/Plugins/Experimental/EditorDataStorage/Source/TedsCore/Private/TypedElementDatabase.cpp
+++ b/Engine/Plugins/Experimental/EditorDataStorage/Source/TedsCore/Private/TypedElementDatabase.cpp
@@ -2250,7 +2250,7 @@ void UEditorDataStorage::TickCooperativeTasks()
 						// under heavy load this can trigger for even well behaving tasks. It's also possible that the OS takes control away from
 						// the editor for a while to give other programs a chance to run which can also cause spikes. This becomes useful if there's
 						// a pattern of frequent runs going over budget.
-						UE_LOG(LogEditorDataStorage, Verbose, TEXT("Time sliced task: '%s' took %.2fms, but was allotted %.2fms"),
+						UE_LOGF(LogEditorDataStorage, Verbose, "Time sliced task: '%ls' took %.2fms, but was allotted %.2fms",
 							*CurrentTask.Name.ToString(), TaskDuration.GetTotalMilliseconds(), RemainingTime.GetTotalMilliseconds());
 					}
 				}
diff --git a/Engine/Plugins/Experimental/MeshModelingToolsetExp/Source/MeshModelingToolsEditorOnlyExp/Private/MeshVertexAttributePaintToolBase.cpp b/Engine/Plugins/Experimental/MeshModelingToolsetExp/Source/MeshModelingToolsEditorOnlyExp/Private/MeshVertexAttributePaintToolBase.cpp
--- a/Engine/Plugins/Experimental/MeshModelingToolsetExp/Source/MeshModelingToolsEditorOnlyExp/Private/MeshVertexAttributePaintToolBase.cpp
+++ b/Engine/Plugins/Experimental/MeshModelingToolsetExp/Source/MeshModelingToolsEditorOnlyExp/Private/MeshVertexAttributePaintToolBase.cpp
@@ -70,11 +70,11 @@ namespace UE::MeshVertexAttributePaintToolBase::Private
 		switch (InMessageType)
 		{
 		case ELogVerbosity::Warning:
-			UE_LOG(LogMeshVertexAttributePaintToolBase, Warning, TEXT("%s"), *InMessage.ToString());
+			UE_LOGF(LogMeshVertexAttributePaintToolBase, Warning, "%ls", *InMessage.ToString());
 			break;
 		case ELogVerbosity::Error:
 			State = SNotificationItem::CS_Fail;
-			UE_LOG(LogMeshVertexAttributePaintToolBase, Error, TEXT("%s"), *InMessage.ToString());
+			UE_LOGF(LogMeshVertexAttributePaintToolBase, Error, "%ls", *InMessage.ToString());
 			break;
 		default:
 			break; // don't log anything unless a warning or error
diff --git a/Engine/Plugins/Experimental/MeshPartition/Source/MeshPartitionEditor/Private/MeshPartitionCollisionGeneration.cpp b/Engine/Plugins/Experimental/MeshPartition/Source/MeshPartitionEditor/Private/MeshPartitionCollisionGeneration.cpp
--- a/Engine/Plugins/Experimental/MeshPartition/Source/MeshPartitionEditor/Private/MeshPartitionCollisionGeneration.cpp
+++ b/Engine/Plugins/Experimental/MeshPartition/Source/MeshPartitionEditor/Private/MeshPartitionCollisionGeneration.cpp
@@ -599,7 +599,7 @@ namespace UE::MeshPartition::Collision
 				Settings.SimplificationSettings.bCutMeshAlongMaterialBoundaries, Settings.SimplificationSettings.MaterialBoundaryCutSnapTolerance);
 			
 			Timer.Stop();
-			UE_LOG(LogMegaMeshEditor, Verbose, TEXT("CutConvertMeshToComplexCollision: Prepared collision data, %.2f seconds, %d verts / %d tris"), Time, MeshCollisionData.Vertices.Num(), MeshCollisionData.Indices.Num());
+			UE_LOGF(LogMegaMeshEditor, Verbose, "CutConvertMeshToComplexCollision: Prepared collision data, %.2f seconds, %d verts / %d tris", Time, MeshCollisionData.Vertices.Num(), MeshCollisionData.Indices.Num());
 		}
 
 		if (bApplyQEMSimplify)
@@ -618,7 +618,7 @@ namespace UE::MeshPartition::Collision
 			Timer.Stop();
 			float ReducedVertPct = Mesh.VertexCount() ? 100.0f * (float)MeshCollisionData.Vertices.Num() / (float)Mesh.VertexCount() : 0.f;
 			float ReducedTriPct = Mesh.TriangleCount() ? 100.0f * (float)MeshCollisionData.Indices.Num() / (float)Mesh.TriangleCount() : 0.f;
-			UE_LOG(LogMegaMeshEditor, Display, TEXT("ConvertMeshToComplexCollision: QEM simplified collision data, %.2f seconds, Initial: %d verts / %d tris, Simplified: %d (%.1f%%) verts / %d (%.1f%%) tris"),
+			UE_LOGF(LogMegaMeshEditor, Display, "ConvertMeshToComplexCollision: QEM simplified collision data, %.2f seconds, Initial: %d verts / %d tris, Simplified: %d (%.1f%%) verts / %d (%.1f%%) tris",
 				Time, Mesh.VertexCount(), Mesh.TriangleCount(), MeshCollisionData.Vertices.Num(), ReducedVertPct, MeshCollisionData.Indices.Num(), ReducedTriPct);
 		}
 		else if (bApplyClusterSimplify)
@@ -698,14 +698,14 @@ namespace UE::MeshPartition::Collision
 			Timer.Stop();
 			float ReducedVertPct = Mesh.VertexCount() ? 100.0f * (float)MeshCollisionData.Vertices.Num() / (float)Mesh.VertexCount() : 0.f;
 			float ReducedTriPct = Mesh.TriangleCount() ? 100.0f * (float)MeshCollisionData.Indices.Num() / (float)Mesh.TriangleCount() : 0.f;
-			UE_LOG(LogMegaMeshEditor, Display, TEXT("ConvertMeshToComplexCollision: Prepared simplified collision data with edge length %f, %.2f seconds, Initial: %d verts / %d tris, Simplified: %d (%.1f%%) verts / %d (%.1f%%) tris"),
+			UE_LOGF(LogMegaMeshEditor, Display, "ConvertMeshToComplexCollision: Prepared simplified collision data with edge length %f, %.2f seconds, Initial: %d verts / %d tris, Simplified: %d (%.1f%%) verts / %d (%.1f%%) tris",
 				SimplifyOptions.TargetEdgeLength, Time, Mesh.VertexCount(), Mesh.TriangleCount(), MeshCollisionData.Vertices.Num(), ReducedVertPct, MeshCollisionData.Indices.Num(), ReducedTriPct);
 		}
 		else
 		{
 			// No simplification to apply
 
-			UE_LOG(LogMegaMeshEditor, Verbose, TEXT("ConvertMeshToComplexCollision: Prepared collision data. %d verts / %d tris"), MeshCollisionData.Vertices.Num(), MeshCollisionData.Indices.Num());
+			UE_LOGF(LogMegaMeshEditor, Verbose, "ConvertMeshToComplexCollision: Prepared collision data. %d verts / %d tris", MeshCollisionData.Vertices.Num(), MeshCollisionData.Indices.Num());
 		}
 		// Compact the material IDs, transferring any used physical materials to the result
 		CompactMaterialIndices(MeshCollisionData, ReferencePhysicalMaterials, CollisionData.PhysicalMaterials);
diff --git a/Engine/Plugins/Experimental/MeshPartition/Source/MeshPartitionEditor/Private/MeshPartitionCollisionTransformer.cpp b/Engine/Plugins/Experimental/MeshPartition/Source/MeshPartitionEditor/Private/MeshPartitionCollisionTransformer.cpp
--- a/Engine/Plugins/Experimental/MeshPartition/Source/MeshPartitionEditor/Private/MeshPartitionCollisionTransformer.cpp
+++ b/Engine/Plugins/Experimental/MeshPartition/Source/MeshPartitionEditor/Private/MeshPartitionCollisionTransformer.cpp
@@ -91,7 +91,7 @@ void FCollisionTransformer::BuildCollisionData(const MeshPartition::FTransformer
 			
 	if (!ensure(TransformerUnit.MeshData != nullptr))
 	{
-		UE_LOG(LogMegaMeshEditor, Error, TEXT("Cannot get MeshData."));
+		UE_LOGF(LogMegaMeshEditor, Error, "Cannot get MeshData.");
 		return;
 	}
 	    		
diff --git a/Engine/Plugins/Experimental/ToolsetRegistry/Source/ToolsetRegistry/Private/ToolsetRegistry/UToolsetRegistry.cpp b/Engine/Plugins/Experimental/ToolsetRegistry/Source/ToolsetRegistry/Private/ToolsetRegistry/UToolsetRegistry.cpp
--- a/Engine/Plugins/Experimental/ToolsetRegistry/Source/ToolsetRegistry/Private/ToolsetRegistry/UToolsetRegistry.cpp
+++ b/Engine/Plugins/Experimental/ToolsetRegistry/Source/ToolsetRegistry/Private/ToolsetRegistry/UToolsetRegistry.cpp
@@ -22,7 +22,7 @@ void UToolsetRegistry::RegisterToolsetClass(TSubclassOf<UToolsetDefinition> InTo
 {
 	if (!InToolsetClass)
 	{
-		UE_LOG(LogToolsetRegistry, Error, TEXT("Cannot register null toolset class"));
+		UE_LOGF(LogToolsetRegistry, Error, "Cannot register null toolset class");
 		return;
 	}
 
@@ -37,25 +37,25 @@ void UToolsetRegistry::RegisterToolsetClass(TSubclassOf<UToolsetDefinition> InTo
 		FString ToolsetName = FBlueprintLibraryToolset::GetToolsetClassName(InToolsetClass);
 		if (!Toolset->HasValidTools())
 		{
-			UE_LOG(LogToolsetRegistry, Warning,
-				   TEXT("Unable to register Toolset %s: invalid tool list"), *ToolsetName);
+			UE_LOGF(LogToolsetRegistry, Warning,
+				   "Unable to register Toolset %ls: invalid tool list", *ToolsetName);
 			return;
 		}
 
 		if (ToolsetRegistrySubsystem.GetValue()->ToolsetRegistry.RegisterToolset(Toolset))
 		{
-			UE_LOG(LogToolsetRegistry, Display,
-				TEXT("Registering Toolset %s"), *ToolsetName);
+			UE_LOGF(LogToolsetRegistry, Display,
+				"Registering Toolset %ls", *ToolsetName);
 		}
 		else
 		{
-			UE_LOG(LogToolsetRegistry, Warning,
-				TEXT("Unable to register Toolset %s"), *ToolsetName);
+			UE_LOGF(LogToolsetRegistry, Warning,
+				"Unable to register Toolset %ls", *ToolsetName);
 		}
 	}
 	else
 	{
-		UE_LOG(LogToolsetRegistry, Error, TEXT("AIToolsetRegistrySubsystem unavailable"));
+		UE_LOGF(LogToolsetRegistry, Error, "AIToolsetRegistrySubsystem unavailable");
 	}
 }
 
@@ -68,7 +68,7 @@ void UToolsetRegistry::UnregisterToolsetClass(TSubclassOf<UToolsetDefinition> In
 
 	if (!InToolsetClass)
 	{
-		UE_LOG(LogToolsetRegistry, Warning, TEXT("Attempting to unregister null toolset class"));
+		UE_LOGF(LogToolsetRegistry, Warning, "Attempting to unregister null toolset class");
 		return;
 	}
 
@@ -81,7 +81,7 @@ void UToolsetRegistry::UnregisterToolsetClass(TSubclassOf<UToolsetDefinition> In
 			ToolsetRegistrySubsystem.GetValue()->ToolsetRegistry.Find(ToolsetName);
 		if (Handler)
 		{
-			UE_LOG(LogToolsetRegistry, Display, TEXT("Unregistering Toolset %s"), *ToolsetName);
+			UE_LOGF(LogToolsetRegistry, Display, "Unregistering Toolset %ls", *ToolsetName);
 			ToolsetRegistrySubsystem.GetValue()->ToolsetRegistry.UnregisterToolset(Handler);
 		}
 	}
@@ -151,7 +151,7 @@ FString UToolsetRegistry::GetToolsetJsonSchema(TSubclassOf<UToolsetDefinition> I
 {
 	if (!InToolsetClass)
 	{
-		UE_LOG(LogToolsetRegistry, Warning, TEXT("Attempting to get schema of null toolset class"));
+		UE_LOGF(LogToolsetRegistry, Warning, "Attempting to get schema of null toolset class");
 		return FString();
 	}
 
diff --git a/Engine/Plugins/MovieScene/MovieRenderPipeline/Source/MovieRenderPipelineCore/Private/Graph/MovieGraphPipeline.cpp b/Engine/Plugins/MovieScene/MovieRenderPipeline/Source/MovieRenderPipelineCore/Private/Graph/MovieGraphPipeline.cpp
--- a/Engine/Plugins/MovieScene/MovieRenderPipeline/Source/MovieRenderPipelineCore/Private/Graph/MovieGraphPipeline.cpp
+++ b/Engine/Plugins/MovieScene/MovieRenderPipeline/Source/MovieRenderPipelineCore/Private/Graph/MovieGraphPipeline.cpp
@@ -51,7 +51,7 @@ UMovieGraphPipeline::UMovieGraphPipeline()
 
 void UMovieGraphPipeline::Initialize(UMoviePipelineExecutorJob* InJob, const FMovieGraphInitConfig& InitConfig)
 {
-	UE_LOG(LogMovieRenderPipeline, Log, TEXT("Initializing MovieGraph Render"));
+	UE_LOGF(LogMovieRenderPipeline, Log, "Initializing MovieGraph Render");
 	if (!ensureAlwaysMsgf(InJob, TEXT("MovieGraph cannot be initialized with a null job. Aborting.")))
 	{
 		//Shutdown(true);
@@ -140,7 +140,7 @@ void UMovieGraphPipeline::Initialize(UMoviePipelineExecutorJob* InJob, const FMo
 	else
 	{
 		// Although the evaluated graph was created for Unreal Insights, it will lead to a general failure to initialize the pipeline.
-		UE_LOG(LogMovieRenderPipeline, Error, TEXT("Unable to initialize Movie Pipeline. The following error(s) were encountered: %s"), *OutError);
+		UE_LOGF(LogMovieRenderPipeline, Error, "Unable to initialize Movie Pipeline. The following error(s) were encountered: %ls", *OutError);
 	}
 
 	//for (ULevelStreaming* Level : GetWorld()->GetStreamingLevels())
@@ -150,18 +150,18 @@ void UMovieGraphPipeline::Initialize(UMoviePipelineExecutorJob* InJob, const FMo
 	//	if (StreamingClass == ULevelStreamingDynamic::StaticClass())
 	//	{
 	//		const FString NonPrefixedLevelName = UWorld::StripPIEPrefixFromPackageName(Level->GetWorldAssetPackageName(), GetWorld()->StreamingLevelsPrefix);
-	//		UE_LOG(LogMovieRenderPipeline, Warning, TEXT("Sub-level '%s' is set to blueprint streaming and will not be visible during a render unless a Sequencer Visibility Track controls its visibility or you have written other code to handle loading it."),
+	//		UE_LOGF(LogMovieRenderPipeline, Warning, "Sub-level '%ls' is set to blueprint streaming and will not be visible during a render unless a Sequencer Visibility Track controls its visibility or you have written other code to handle loading it.",
 	//			*NonPrefixedLevelName);
 	//	}
 	//}
 
 	// ToDo: Print some information, Job Name, Level Sequence, etc.
-	UE_LOG(LogMovieRenderPipeline, Log, TEXT("Finished initializing MovieGraph Render"));
+	UE_LOGF(LogMovieRenderPipeline, Log, "Finished initializing MovieGraph Render");
 
 	const bool bNoShotsToRender = ActiveShotList.Num() == 0;
 	if (bNoShotsToRender)
 	{
-		UE_LOG(LogMovieRenderPipeline, Warning, TEXT("No shots detected to render. Either all outside playback range, or disabled via shot mask. Shutting down."));
+		UE_LOGF(LogMovieRenderPipeline, Warning, "No shots detected to render. Either all outside playback range, or disabled via shot mask. Shutting down.");
 	}
 
 	// If shutting down immediately, we'll transition directly to finish as there is no work to do.
@@ -200,7 +200,7 @@ void UMovieGraphPipeline::DuplicateJobAndConfiguration()
 
 		if (!GeneratedGraph)
 		{
-			UE_LOG(LogMovieRenderPipeline, Error, TEXT("DuplicateJobAndConfiguration: Failed to generate graph from Basic config. The job will have no graph configuration."));
+			UE_LOGF(LogMovieRenderPipeline, Error, "DuplicateJobAndConfiguration: Failed to generate graph from Basic config. The job will have no graph configuration.");
 		}
 
 		// Pre-populate the duplication map so DuplicateConfigRecursive skips these already-transient graphs
@@ -339,12 +339,12 @@ void UMovieGraphPipeline::LoadPreviewWidget()
 		}
 		else
 		{
-			UE_LOG(LogMovieRenderPipeline, Warning, TEXT("Failed to create Preview Screen UMG Widget. No in-game overlay available."));
+			UE_LOGF(LogMovieRenderPipeline, Warning, "Failed to create Preview Screen UMG Widget. No in-game overlay available.");
 		}
 	}
 	else
 	{
-		UE_LOG(LogMovieRenderPipeline, Warning, TEXT("Failed to find Preview Screen UMG Widget class. No in-game overlay available."));
+		UE_LOGF(LogMovieRenderPipeline, Warning, "Failed to find Preview Screen UMG Widget class. No in-game overlay available.");
 	}
 }
 
@@ -474,7 +474,7 @@ void UMovieGraphPipeline::UpdateLayerContentsInRenderLayerSubsystem(const UMovie
 
 					if (!bFoundModifiedCollection)
 					{
-						UE_LOG(LogMovieRenderPipeline, Warning, TEXT("The modifier '%s' specified a collection '%s', but the collection couldn't be found."),
+						UE_LOGF(LogMovieRenderPipeline, Warning, "The modifier '%ls' specified a collection '%ls', but the collection couldn't be found.",
 							*CollectionModifierInstance->GetName(), *ModifiedCollectionName.ToString());
 					}
 				}
@@ -549,7 +549,7 @@ void UMovieGraphPipeline::BuildShotListFromDataSource()
 		else
 		{
 			GraphTimeStepInstances.Add(NewObject<UMovieGraphLinearTimeStep>(this));
-			UE_LOG(LogMovieRenderPipeline, Warning, TEXT("The shot '%s' specified a time step of type '%s', but it could not be loaded. Defaulting to linear."),
+			UE_LOGF(LogMovieRenderPipeline, Warning, "The shot '%ls' specified a time step of type '%ls', but it could not be loaded. Defaulting to linear.",
 				*Shot->OuterName, *SamplingMethodNode->SamplingMethodClass.GetAssetName());
 		}
 		
@@ -619,7 +619,7 @@ void UMovieGraphPipeline::OnEngineTickBeginFrame()
 	LLM_SCOPE_BYNAME(TEXT("MovieGraphBeginFrame"));
 
 
-	UE_LOG(LogMovieRenderPipeline, VeryVerbose, TEXT("MovieGraph OnEngineTickBeginFrame"));
+	UE_LOGF(LogMovieRenderPipeline, VeryVerbose, "MovieGraph OnEngineTickBeginFrame");
 	switch (PipelineState)
 	{
 	case EMovieRenderPipelineState::Uninitialized:
@@ -663,7 +663,7 @@ void UMovieGraphPipeline::TickProducingFrames()
 	
 	if (bShutdownRequested)
 	{
-		UE_LOG(LogMovieRenderPipeline, Log, TEXT("MovieGraph TickProducingFrames: Async Shutdown Requested, abandoning remaining work and moving to Finalize."));
+		UE_LOGF(LogMovieRenderPipeline, Log, "MovieGraph TickProducingFrames: Async Shutdown Requested, abandoning remaining work and moving to Finalize.");
 		TransitionToState(EMovieRenderPipelineState::Finalize);
 		return;
 	}
@@ -756,7 +756,7 @@ void UMovieGraphPipeline::TickPostFinalizeExport(const bool bInForceFinish)
 	// This step assumes you have produced data and filled the data structures.
 	check(PipelineState == EMovieRenderPipelineState::Export);
 	
-	UE_LOG(LogMovieRenderPipeline, Verbose, TEXT("[%d] PostFinalize Export (Start)."), GFrameCounter);
+	UE_LOGF(LogMovieRenderPipeline, Verbose, "[%d] PostFinalize Export (Start).", GFrameCounter);
 
 	// Loop through any extensions (such as XML export) and let them export using all of the
 	// data that was generated during this run such as containers, output names and lengths.
@@ -786,7 +786,7 @@ void UMovieGraphPipeline::TickPostFinalizeExport(const bool bInForceFinish)
 		FPlatformProcess::Sleep(1.f);
 	} while (true);
 
-	UE_LOG(LogMovieRenderPipeline, Verbose, TEXT("[%d] PostFinalize Export (End)."), GFrameCounter);
+	UE_LOGF(LogMovieRenderPipeline, Verbose, "[%d] PostFinalize Export (End).", GFrameCounter);
 
 	// If an output container is still working, we'll early out to keep the UI responsive.
 	// If they've forced a finish this will have to be true before we can reach this block.
@@ -848,11 +848,11 @@ void UMovieGraphPipeline::StartUnrealInsightsCapture(UMovieGraphEvaluatedConfig*
 	if (bTraceStarted)
 	{
 		bDidStartInsightsCapture = true;
-		UE_LOG(LogMovieRenderPipeline, Log, TEXT("Started capturing UnrealInsights trace file to %s"), *FinalFilePath);
+		UE_LOGF(LogMovieRenderPipeline, Log, "Started capturing UnrealInsights trace file to %ls", *FinalFilePath);
 	}
 	else
 	{
-		UE_LOG(LogMovieRenderPipeline, Warning, TEXT("Failed to start capturing UnrealInsights trace. Is there already a trace session in progress?"));
+		UE_LOGF(LogMovieRenderPipeline, Warning, "Failed to start capturing UnrealInsights trace. Is there already a trace session in progress?");
 	}
 }
 
@@ -1033,7 +1033,7 @@ void UMovieGraphPipeline::TeardownShot(const TObjectPtr<UMoviePipelineExecutorSh
 	// tick the new shot will be initialized and processed.
 	if (CurrentShotIndex >= (ActiveShotList.Num() - 1))
 	{
-		UE_LOG(LogMovieRenderPipeline, Log, TEXT("MovieGraph Finished rendering last shot. Moving to Finalize to finish writing items to disk."));
+		UE_LOGF(LogMovieRenderPipeline, Log, "MovieGraph Finished rendering last shot. Moving to Finalize to finish writing items to disk.");
 		TransitionToState(EMovieRenderPipelineState::Finalize);
 	}
 
@@ -1063,7 +1063,7 @@ void UMovieGraphPipeline::SetSoloShot(const TObjectPtr<UMoviePipelineExecutorSho
 	// such as re-using shots or using different parts of shots. To resolve this, we expand the entire tree needed for a given range, render it, and then restore the original
 	// values before moving onto the next shot so that each shot has no effect on the others.
 	{
-		UE_LOG(LogMovieRenderPipeline, Log, TEXT("Expanding Shot %d/%d (Shot: %s Camera: %s)"), CurrentShotIndex + 1, ActiveShotList.Num(), *InShot->OuterName, *InShot->InnerName);
+		UE_LOGF(LogMovieRenderPipeline, Log, "Expanding Shot %d/%d (Shot: %ls Camera: %ls)", CurrentShotIndex + 1, ActiveShotList.Num(), *InShot->OuterName, *InShot->InnerName);
 
 		// Enable the one hierarchy we do want for rendering. We will re-disable it later when we restore the current Sequence state.
 		GetDataSourceInstance()->UnmuteShot(InShot);
@@ -1111,7 +1111,7 @@ void UMovieGraphPipeline::ExpandShot(const TObjectPtr<UMoviePipelineExecutorShot
 	FFrameTime StartTimeInRoot = FFrameRate::TransformTime(InShot->ShotInfo.TotalOutputRangeRoot.GetLowerBoundValue(), InTickResolution, InDisplayRate);
 	if (bIsPrePass && StartTimeInRoot.GetSubFrame() != 0.f)
 	{
-		UE_LOG(LogMovieRenderPipeline, Warning, TEXT("Shot/Camera \"%s [%s]\" starts on a sub-frame. Rendered range has been rounded to the previous frame to make output numbers align with Sequencer."), 
+		UE_LOGF(LogMovieRenderPipeline, Warning, "Shot/Camera \"%ls [%ls]\" starts on a sub-frame. Rendered range has been rounded to the previous frame to make output numbers align with Sequencer.", 
 			*InShot->OuterName, *InShot->InnerName);
 		FFrameNumber NewStartFrame = FFrameRate::TransformTime(FFrameTime(StartTimeInRoot.GetFrame()), InDisplayRate, InTickResolution).FloorToFrame();
 		InShot->ShotInfo.TotalOutputRangeRoot.SetLowerBoundValue(NewStartFrame);
@@ -1157,13 +1157,13 @@ void UMovieGraphPipeline::OnEngineTickEndFrame()
 		return;
 	}
 
-	UE_LOG(LogMovieRenderPipeline, VeryVerbose, TEXT("MovieGraph OnEngineTickEndFrame (Start)"));
+	UE_LOGF(LogMovieRenderPipeline, VeryVerbose, "MovieGraph OnEngineTickEndFrame (Start)");
 
 	GraphAudioRendererInstance->ProcessAudioTick();
 	
 	RenderFrame();
 
-	UE_LOG(LogMovieRenderPipeline, VeryVerbose, TEXT("MovieGraph OnEngineTickEndFrame (End)"));
+	UE_LOGF(LogMovieRenderPipeline, VeryVerbose, "MovieGraph OnEngineTickEndFrame (End)");
 }
 
 void UMovieGraphPipeline::RenderFrame()
@@ -1246,7 +1246,7 @@ void UMovieGraphPipeline::RequestShutdownImpl(bool bIsError)
 		// You can call Shutdown during these two, but they won't do anything as we're already shutting down at that point.
 	case EMovieRenderPipelineState::Finalize:
 	case EMovieRenderPipelineState::Export:
-		UE_LOG(LogMovieRenderPipeline, Log, TEXT("MovieGraph Async Shutdown Requested, ignoring due to already being on the way to shutdown."));
+		UE_LOGF(LogMovieRenderPipeline, Log, "MovieGraph Async Shutdown Requested, ignoring due to already being on the way to shutdown.");
 		break;
 	}
 }
@@ -1383,7 +1383,7 @@ void UMovieGraphPipeline::TransitionToState(const EMovieRenderPipelineState InNe
 			PostRenderEvaluatedGraph = GetCurrentJob()->GetGraphPreset()->CreateFlattenedGraph(GetCurrentTraversalContext(bForShot), FlattenError);
 			if (!PostRenderEvaluatedGraph)
 			{
-				UE_LOG(LogMovieRenderPipeline, Warning, TEXT("Generating a post-render evaluated graph failed for job '%s'. Post-render exports will not run. Evaluation error: %s"), *GetCurrentJob()->JobName, *FlattenError);
+				UE_LOGF(LogMovieRenderPipeline, Warning, "Generating a post-render evaluated graph failed for job '%ls'. Post-render exports will not run. Evaluation error: %ls", *GetCurrentJob()->JobName, *FlattenError);
 			}
 
 			// This is called once notifying output containers that all frames that will be submitted have been submitted.
@@ -1478,14 +1478,14 @@ void UMovieGraphPipeline::TransitionToState(const EMovieRenderPipelineState InNe
 
 			//GAreScreenMessagesEnabled = bPrevGScreenMessagesEnabled;
 
-			//UE_LOG(LogMovieRenderPipeline, Log, TEXT("Movie Pipeline completed. Duration: %s"), *(FDateTime::UtcNow() - InitializationTime).ToString());
+			//UE_LOGF(LogMovieRenderPipeline, Log, "Movie Pipeline completed. Duration: %ls", *(FDateTime::UtcNow() - InitializationTime).ToString());
 
 			OutputNodesDataSentTo.Reset();
 
 			GraphDataSourceInstance->RestoreCachedDataPostJob();
 
 			OnMoviePipelineFinishedImpl();
-			UE_LOG(LogMovieRenderPipeline, Log, TEXT("Movie Graph Render completed. Duration: %s"), *(FDateTime::UtcNow() - GraphInitializationTime).ToString());
+			UE_LOGF(LogMovieRenderPipeline, Log, "Movie Graph Render completed. Duration: %ls", *(FDateTime::UtcNow() - GraphInitializationTime).ToString());
 		}
 		break;
 	}
@@ -1508,9 +1508,9 @@ void UMovieGraphPipeline::OnMoviePipelineFinishedImpl()
 	Params.bSuccess = !bShutdownSetErrorFlag;
 	Params.GraphData = GeneratedOutputData;
 	//
-	//UE_LOG(LogMovieRenderPipelineIO, Verbose, TEXT("Files written to disk for entire sequence:"));
+	//UE_LOGF(LogMovieRenderPipelineIO, Verbose, "Files written to disk for entire sequence:");
 	//PrintVerboseLogForFiles(GeneratedShotOutputData);
-	//UE_LOG(LogMovieRenderPipelineIO, Verbose, TEXT("Completed outputting files written to disk."));
+	//UE_LOGF(LogMovieRenderPipelineIO, Verbose, "Completed outputting files written to disk.");
 	//
 
 	OnMoviePipelineWorkFinishedDelegateNative.Broadcast(Params);
@@ -1524,14 +1524,14 @@ void UMovieGraphPipeline::OnMoviePipelineFinishedImpl()
 //		const UMoviePipelineExecutorShot* Shot = OutputData.Shot.Get();
 //		if (Shot)
 //		{
-//			UE_LOG(LogMovieRenderPipelineIO, Verbose, TEXT("Shot: %s [%s]"), *Shot->OuterName, *Shot->InnerName);
+//			UE_LOGF(LogMovieRenderPipelineIO, Verbose, "Shot: %ls [%ls]", *Shot->OuterName, *Shot->InnerName);
 //		}
 //		for (const TPair<FMoviePipelinePassIdentifier, FMoviePipelineRenderPassOutputData>& Pair : OutputData.RenderPassData)
 //		{
-//			UE_LOG(LogMovieRenderPipelineIO, Verbose, TEXT("Render Pass: %s"), *Pair.Key.Name);
+//			UE_LOGF(LogMovieRenderPipelineIO, Verbose, "Render Pass: %ls", *Pair.Key.Name);
 //			for (const FString& FilePath : Pair.Value.FilePaths)
 //			{
-//				UE_LOG(LogMovieRenderPipelineIO, Verbose, TEXT("\t\t%s"), *FilePath);
+//				UE_LOGF(LogMovieRenderPipelineIO, Verbose, "\t\t%ls", *FilePath);
 //			}
 //		}
 //	}
@@ -1558,7 +1558,7 @@ void UMovieGraphPipeline::ProcessOutstandingFinishedFrames()
 			OutputFrame.FileMetadata.Add(TEXT("unreal/frameRenderDuration"), DurationTimeStr);
 
 			// int32 FrameNumber = OutputFrame.TraversalContext.Time.RenderedFrameNumber;
-			// UE_LOG(LogTemp, Log, TEXT("Frame: %d Duration: %s"), FrameNumber, *DurationTimeStr);
+			// UE_LOGF(LogTemp, Log, "Frame: %d Duration: %ls", FrameNumber, *DurationTimeStr);
 		}
 
 		// How we choose which image gets sent to which output container is a little tricky. We use the CDO of each node type on purpose
@@ -1690,7 +1690,7 @@ void UMovieGraphPipeline::ProcessOutstandingFutures(const bool bStallAndProcessA
 		// Sometime futures can be completed, but will be set to an error state (such as we couldn't write to the specified disk path.)
 		if (!OutputFuture.Get<0>().Get())
 		{
-			UE_LOG(LogMovieRenderPipeline, Error, TEXT("Error exporting frame, canceling movie export."));
+			UE_LOGF(LogMovieRenderPipeline, Error, "Error exporting frame, canceling movie export.");
 			RequestShutdown(true);
 			break;
 		}
diff --git a/Engine/Plugins/MovieScene/MovieRenderPipeline/Source/MovieRenderPipelineCore/Private/MoviePipelineBasicConfig.cpp b/Engine/Plugins/MovieScene/MovieRenderPipeline/Source/MovieRenderPipelineCore/Private/MoviePipelineBasicConfig.cpp
--- a/Engine/Plugins/MovieScene/MovieRenderPipeline/Source/MovieRenderPipelineCore/Private/MoviePipelineBasicConfig.cpp
+++ b/Engine/Plugins/MovieScene/MovieRenderPipeline/Source/MovieRenderPipelineCore/Private/MoviePipelineBasicConfig.cpp
@@ -46,14 +46,14 @@ namespace UE::MoviePipeline::BasicConfig
 		const FProperty* Property = FindFProperty<FProperty>(InObject->GetClass(), PropertyName);
 		if (!Property)
 		{
-			UE_LOG(LogMovieRenderPipeline, Warning, TEXT("UMoviePipelineBasicConfig: Could not find property '%s' on '%s'"), *PropertyName.ToString(), *InObject->GetClass()->GetName());
+			UE_LOGF(LogMovieRenderPipeline, Warning, "UMoviePipelineBasicConfig: Could not find property '%ls' on '%ls'", *PropertyName.ToString(), *InObject->GetClass()->GetName());
 			return;
 		}
 
 		void* ValuePtr = Property->ContainerPtrToValuePtr<void>(InObject);
 		if (Property->ImportText_Direct(*Value, ValuePtr, InObject, PPF_None) == nullptr)
 		{
-			UE_LOG(LogMovieRenderPipeline, Warning, TEXT("UMoviePipelineBasicConfig: Failed to set property '%s' on '%s' from value '%s'"),
+			UE_LOGF(LogMovieRenderPipeline, Warning, "UMoviePipelineBasicConfig: Failed to set property '%ls' on '%ls' from value '%ls'",
 				*PropertyName.ToString(), *InObject->GetClass()->GetName(), *Value);
 		}
 	}
@@ -252,7 +252,7 @@ UMovieGraphConfig* UMoviePipelineBasicConfig::GenerateGraph(UMoviePipelineBasicC
 
 	if (!InConfig)
 	{
-		UE_LOG(LogMovieRenderPipeline, Error, TEXT("UMoviePipelineBasicConfig::GenerateGraph called with null config."));
+		UE_LOGF(LogMovieRenderPipeline, Error, "UMoviePipelineBasicConfig::GenerateGraph called with null config.");
 		return nullptr;
 	}
 
@@ -264,14 +264,14 @@ UMovieGraphConfig* UMoviePipelineBasicConfig::GenerateGraph(UMoviePipelineBasicC
 	const UMovieGraphConfig* TemplateGraph = Cast<UMovieGraphConfig>(TemplateGraphPath.TryLoad());
 	if (!TemplateGraph)
 	{
-		UE_LOG(LogMovieRenderPipeline, Error, TEXT("UMoviePipelineBasicConfig::GenerateGraph: Failed to load the template graph."));
+		UE_LOGF(LogMovieRenderPipeline, Error, "UMoviePipelineBasicConfig::GenerateGraph: Failed to load the template graph.");
 		return nullptr;
 	}
 
 	UMovieGraphConfig* Graph = DuplicateObject<UMovieGraphConfig>(TemplateGraph, InOuter);
 	if (!Graph)
 	{
-		UE_LOG(LogMovieRenderPipeline, Error, TEXT("UMoviePipelineBasicConfig::GenerateGraph: Failed to duplicate the template graph."));
+		UE_LOGF(LogMovieRenderPipeline, Error, "UMoviePipelineBasicConfig::GenerateGraph: Failed to duplicate the template graph.");
 		return nullptr;
 	}
 
@@ -398,8 +398,8 @@ UMovieGraphConfig* UMoviePipelineBasicConfig::GenerateGraph(UMoviePipelineBasicC
 			UClass* OutputClass = OutputClassPtr.LoadSynchronous();
 			if (!OutputClass || OutputClass->HasAnyClassFlags(CLASS_Abstract))
 			{
-				UE_LOG(LogMovieRenderPipeline, Warning,
-					TEXT("UMoviePipelineBasicConfig::GenerateGraph: Skipping invalid/abstract output type '%s'."),
+				UE_LOGF(LogMovieRenderPipeline, Warning,
+					"UMoviePipelineBasicConfig::GenerateGraph: Skipping invalid/abstract output type '%ls'.",
 					*OutputClassPtr.ToString());
 				continue;
 			}
@@ -499,8 +499,8 @@ void UMoviePipelineBasicConfig::SaveAsDefault(const UMoviePipelineBasicConfig* I
 
 	if (!bSuccess)
 	{
-		UE_LOG(LogMovieRenderPipeline, Warning,
-			TEXT("UMoviePipelineBasicConfig::SaveAsDefault: Failed to save default config to '%s'."), *PackageFileName);
+		UE_LOGF(LogMovieRenderPipeline, Warning,
+			"UMoviePipelineBasicConfig::SaveAsDefault: Failed to save default config to '%ls'.", *PackageFileName);
 	}
 }
 #endif // WITH_EDITOR
diff --git a/Engine/Plugins/MovieScene/MovieRenderPipeline/Source/MovieRenderPipelineCore/Private/MoviePipelineInProcessExecutor.cpp b/Engine/Plugins/MovieScene/MovieRenderPipeline/Source/MovieRenderPipelineCore/Private/MoviePipelineInProcessExecutor.cpp
--- a/Engine/Plugins/MovieScene/MovieRenderPipeline/Source/MovieRenderPipelineCore/Private/MoviePipelineInProcessExecutor.cpp
+++ b/Engine/Plugins/MovieScene/MovieRenderPipeline/Source/MovieRenderPipelineCore/Private/MoviePipelineInProcessExecutor.cpp
@@ -30,24 +30,24 @@ void UMoviePipelineInProcessExecutor::Start(const UMoviePipelineExecutorJob* InJ
 	{		
 		if (!World)
 		{
-			UE_LOG(LogMovieRenderPipeline, Warning, TEXT("Unable to start movie pipeline job. No current map."));
+			UE_LOGF(LogMovieRenderPipeline, Warning, "Unable to start movie pipeline job. No current map.");
 			OnIndividualPipelineFinished(nullptr);
 			return;
 		}
 
 		if (World != InJob->Map.ResolveObject())
 		{
-			UE_LOG(LogMovieRenderPipeline, Warning, TEXT("Current map '%s' does not match job's map: '%s'. Rendering may happen on a different map than expected and produce incorrect results."), *GetNameSafe(World), *InJob->Map.GetAssetPath().ToString());
+			UE_LOGF(LogMovieRenderPipeline, Warning, "Current map '%ls' does not match job's map: '%ls'. Rendering may happen on a different map than expected and produce incorrect results.", *GetNameSafe(World), *InJob->Map.GetAssetPath().ToString());
 		}
 		
-		UE_LOG(LogMovieRenderPipeline, Log, TEXT("Starting %s"), *GetNameSafe(World));
+		UE_LOGF(LogMovieRenderPipeline, Log, "Starting %ls", *GetNameSafe(World));
 	}
 
 	// Check for null sequence
 	ULevelSequence* LevelSequence = Cast<ULevelSequence>(InJob->Sequence.TryLoad());
 	if (!LevelSequence)
 	{
-		UE_LOG(LogMovieRenderPipeline, Error, TEXT("Failed to load Sequence specified by job: %s"), *InJob->Sequence.ToString());
+		UE_LOGF(LogMovieRenderPipeline, Error, "Failed to load Sequence specified by job: %ls", *InJob->Sequence.ToString());
 
 		FText FailureReason = LOCTEXT("InvalidSequenceFailureDialog", "One or more jobs in the queue has an invalid/null sequence. See log for details.");
 		OnPipelineErrored(nullptr, true, FailureReason);
@@ -78,7 +78,7 @@ void UMoviePipelineInProcessExecutor::Start(const UMoviePipelineExecutorJob* InJ
 		}
 
 		FString LevelPath = InJob->Map.GetLongPackageName();
-		UE_LOG(LogMovieRenderPipeline, Log, TEXT("About to load target map %s"), *LevelPath);
+		UE_LOGF(LogMovieRenderPipeline, Log, "About to load target map %ls", *LevelPath);
 		FCoreUObjectDelegates::PostLoadMapWithWorld.AddUObject(this, &UMoviePipelineInProcessExecutor::OnMapLoadFinished);
 		UGameplayStatics::OpenLevel(World, FName(*LevelPath), true, MapOptions);
 	}
@@ -86,7 +86,7 @@ void UMoviePipelineInProcessExecutor::Start(const UMoviePipelineExecutorJob* InJ
 
 void UMoviePipelineInProcessExecutor::OnMapLoadFinished(UWorld* NewWorld)
 {
-	UE_LOG(LogMovieRenderPipeline, Log, TEXT("Finished map load %s"), *GetNameSafe(NewWorld));
+	UE_LOGF(LogMovieRenderPipeline, Log, "Finished map load %ls", *GetNameSafe(NewWorld));
 
 	// NewWorld can be null if a world is being destroyed.
 	if (!NewWorld)
@@ -101,7 +101,7 @@ void UMoviePipelineInProcessExecutor::OnMapLoadFinished(UWorld* NewWorld)
 	if (CurrentPipelineIndex >= Queue->GetJobs().Num())
 	{
 		FCoreDelegates::OnBeginFrame.RemoveAll(this);
-		UE_LOG(LogMovieRenderPipeline, Error, TEXT("Out of bounds Job Index (%d) in Queue (Length: %d)!"), CurrentPipelineIndex, Queue->GetJobs().Num());
+		UE_LOGF(LogMovieRenderPipeline, Error, "Out of bounds Job Index (%d) in Queue (Length: %d)!", CurrentPipelineIndex, Queue->GetJobs().Num());
 
 		FText FailureReason = LOCTEXT("OutOfBoundsFailureDialog", "The job index is out of bounds, did you clear the Queue while rendering was happening? See log for details.");
 		OnPipelineErrored(nullptr, true, FailureReason);
@@ -152,7 +152,7 @@ void UMoviePipelineInProcessExecutor::OnMapLoadFinished(UWorld* NewWorld)
 	
 	if (ExecutorSettings->InitialDelayFrameCount == 0)
 	{
-		UE_LOG(LogMovieRenderPipeline, Log, TEXT("Zero Initial Delay, initializing..."));
+		UE_LOGF(LogMovieRenderPipeline, Log, "Zero Initial Delay, initializing...");
 		if (PipelineAsLegacy)
 		{
 			PipelineAsLegacy->Initialize(Queue->GetJobs()[CurrentPipelineIndex]);
@@ -175,7 +175,7 @@ void UMoviePipelineInProcessExecutor::OnTick()
 	{
 		if (RemainingInitializationFrames == 0)
 		{
-			UE_LOG(LogMovieRenderPipeline, Log, TEXT("Delay finished, initializing..."));
+			UE_LOGF(LogMovieRenderPipeline, Log, "Delay finished, initializing...");
 			UMoviePipeline* PipelineAsLegacy = Cast<UMoviePipeline>(ActiveMoviePipeline);
 			if (PipelineAsLegacy)
 			{
@@ -207,13 +207,13 @@ void UMoviePipelineInProcessExecutor::OnApplicationQuit()
 	// Only call Shutdown if the pipeline hasn't been finished.
 	if (ActiveMoviePipeline && ActiveMoviePipeline->GetPipelineState() != EMovieRenderPipelineState::Finished)
 	{
-		UE_LOG(LogMovieRenderPipeline, Log, TEXT("MoviePipelineInProcessExecutor: Application quit while Movie Pipeline was still active. Stalling to do full shutdown."));
+		UE_LOGF(LogMovieRenderPipeline, Log, "MoviePipelineInProcessExecutor: Application quit while Movie Pipeline was still active. Stalling to do full shutdown.");
 
 		// This will flush any outstanding work on the movie pipeline (file writes) immediately
 		ActiveMoviePipeline->RequestShutdown(); // Set the Shutdown Requested flag.
 		ActiveMoviePipeline->Shutdown(); // Flush the shutdown.
 
-		UE_LOG(LogMovieRenderPipeline, Log, TEXT("MoviePipelineInProcessExecutor: Stalling finished, pipeline has shut down."));
+		UE_LOGF(LogMovieRenderPipeline, Log, "MoviePipelineInProcessExecutor: Stalling finished, pipeline has shut down.");
 	}
 }
 
@@ -337,7 +337,7 @@ void UMoviePipelineInProcessExecutor::CancelCurrentJob_Implementation()
 {
 	if (ActiveMoviePipeline && ActiveMoviePipeline->GetPipelineState() != EMovieRenderPipelineState::Finished)
 	{
-		UE_LOG(LogMovieRenderPipeline, Log, TEXT("MoviePipelineInProcessExecutor: Application quit while Movie Pipeline was still active. Stalling to do full shutdown."));
+		UE_LOGF(LogMovieRenderPipeline, Log, "MoviePipelineInProcessExecutor: Application quit while Movie Pipeline was still active. Stalling to do full shutdown.");
 
 		// This will flush any outstanding work on the movie pipeline (file writes) immediately
 		ActiveMoviePipeline->RequestShutdown(); // Set the Shutdown Requested flag.
diff --git a/Engine/Plugins/MovieScene/MovieRenderPipeline/Source/MovieRenderPipelineCore/Private/MoviePipelineQueue.cpp b/Engine/Plugins/MovieScene/MovieRenderPipeline/Source/MovieRenderPipelineCore/Private/MoviePipelineQueue.cpp
--- a/Engine/Plugins/MovieScene/MovieRenderPipeline/Source/MovieRenderPipelineCore/Private/MoviePipelineQueue.cpp
+++ b/Engine/Plugins/MovieScene/MovieRenderPipeline/Source/MovieRenderPipelineCore/Private/MoviePipelineQueue.cpp
@@ -109,7 +109,7 @@ UMoviePipelineQueue* UMoviePipelineQueue::CopyFrom(UMoviePipelineQueue* InQueue)
 {
 	if (!InQueue)
 	{
-		UE_LOG(LogMovieRenderPipeline, Warning, TEXT("Cannot copy the contents of a null queue."));
+		UE_LOGF(LogMovieRenderPipeline, Warning, "Cannot copy the contents of a null queue.");
 		return nullptr;
 	}
 
@@ -244,8 +244,8 @@ void UMoviePipelineExecutorJob::PostLoad()
 	// based on !GraphPreset.IsNull() (non-empty path), but the referenced asset may have been deleted.
 	if (ConfigMode == EMoviePipelineConfigMode::Graph && GetGraphPreset() == nullptr)
 	{
-		UE_LOG(LogMovieRenderPipeline, Warning,
-			TEXT("Job '%s' was in Graph mode but the graph preset could not be loaded. Falling back to Preset mode."),
+		UE_LOGF(LogMovieRenderPipeline, Warning,
+			"Job '%ls' was in Graph mode but the graph preset could not be loaded. Falling back to Preset mode.",
 			*GetName());
 		ConfigMode = EMoviePipelineConfigMode::Preset;
 	}
diff --git a/Engine/Plugins/MovieScene/MovieRenderPipeline/Source/MovieRenderPipelineEditor/Private/MoviePipelineNewProcessExecutor.cpp b/Engine/Plugins/MovieScene/MovieRenderPipeline/Source/MovieRenderPipelineEditor/Private/MoviePipelineNewProcessExecutor.cpp
--- a/Engine/Plugins/MovieScene/MovieRenderPipeline/Source/MovieRenderPipelineEditor/Private/MoviePipelineNewProcessExecutor.cpp
+++ b/Engine/Plugins/MovieScene/MovieRenderPipeline/Source/MovieRenderPipelineEditor/Private/MoviePipelineNewProcessExecutor.cpp
@@ -79,7 +79,7 @@ void UMoviePipelineNewProcessExecutor::Execute_Implementation(UMoviePipelineQueu
 	UMoviePipelineQueue* DuplicatedQueue = UMoviePipelineEditorBlueprintLibrary::SaveQueueToManifestFile(InPipelineQueue, ManifestFilePath);
 	if (!DuplicatedQueue)
 	{
-		UE_LOG(LogMovieRenderPipeline, Error, TEXT("Could not save manifest package to disk. Path: %s"), *ManifestFilePath);
+		UE_LOGF(LogMovieRenderPipeline, Error, "Could not save manifest package to disk. Path: %ls", *ManifestFilePath);
 		OnExecutorFinishedImpl();
 		return;
 	}
@@ -227,8 +227,8 @@ void UMoviePipelineNewProcessExecutor::Execute_Implementation(UMoviePipelineQueu
 	const bool bLaunchWindowHidden = false;
 	const uint32 PriorityModifier = 0;
 
-	UE_LOG(LogMovieRenderPipeline, Log, TEXT("Launching a new process to render with the following command line:"));
-	UE_LOG(LogMovieRenderPipeline, Log, TEXT("%s %s"), *ExecutablePath, *FinalCommandLine);
+	UE_LOGF(LogMovieRenderPipeline, Log, "Launching a new process to render with the following command line:");
+	UE_LOGF(LogMovieRenderPipeline, Log, "%ls %ls", *ExecutablePath, *FinalCommandLine);
 	
 	ProcessHandle = FPlatformProcess::CreateProc(
 		*ExecutablePath, *FinalCommandLine, bLaunchDetatched,
@@ -237,7 +237,7 @@ void UMoviePipelineNewProcessExecutor::Execute_Implementation(UMoviePipelineQueu
 
 	if (!ProcessHandle.IsValid())
 	{
-		UE_LOG(LogMovieRenderPipeline, Warning, TEXT("Failed to launch executable for new process render. Executable Path: \"%s\" Command Line: \"%s\""), FPlatformProcess::ExecutablePath(), *FinalCommandLine);
+		UE_LOGF(LogMovieRenderPipeline, Warning, "Failed to launch executable for new process render. Executable Path: \"%ls\" Command Line: \"%ls\"", FPlatformProcess::ExecutablePath(), *FinalCommandLine);
 		// OnPipelineErrored(nullptr, true, LOCTEXT("ProcessFailedToLaunch", "New Process failed to launch. See log for command line argument used."));
 		// OnIndividualPipelineFinished(nullptr);
 	}
@@ -271,7 +271,7 @@ void UMoviePipelineNewProcessExecutor::CheckForProcessFinished()
 		// Log an error for now
 		if (ReturnCode != 0)
 		{
-			UE_LOG(LogMovieRenderPipeline, Warning, TEXT("Process exited with non-success return code. Return Code; %d"), ReturnCode);
+			UE_LOGF(LogMovieRenderPipeline, Warning, "Process exited with non-success return code. Return Code; %d", ReturnCode);
 			// OnPipelineErrored(nullptr, true, LOCTEXT("ProcessNonZeroReturn", "Non-success return code returned. See log for Return Code."));
 		}
 
@@ -297,7 +297,7 @@ void UMoviePipelineNewProcessExecutor::CancelAllJobs_Implementation()
 	}
 	else
 	{
-		UE_LOG(LogMovieRenderPipeline, Warning, TEXT("Attempting to cancel UMoviePipelineNewProcessExecutor job but process has already exited."));
+		UE_LOGF(LogMovieRenderPipeline, Warning, "Attempting to cancel UMoviePipelineNewProcessExecutor job but process has already exited.");
 	}
 }
 
diff --git a/Engine/Plugins/MovieScene/MovieRenderPipeline/Source/MovieRenderPipelineEditor/Private/MoviePipelinePIEExecutor.cpp b/Engine/Plugins/MovieScene/MovieRenderPipeline/Source/MovieRenderPipelineEditor/Private/MoviePipelinePIEExecutor.cpp
--- a/Engine/Plugins/MovieScene/MovieRenderPipeline/Source/MovieRenderPipelineEditor/Private/MoviePipelinePIEExecutor.cpp
+++ b/Engine/Plugins/MovieScene/MovieRenderPipeline/Source/MovieRenderPipelineEditor/Private/MoviePipelinePIEExecutor.cpp
@@ -89,7 +89,7 @@ UMoviePipelinePIEExecutor::UMoviePipelinePIEExecutor()
 
 void UMoviePipelinePIEExecutor::SetAllowUsingUnsavedLevels(const bool bShouldAllow)
 {
-	UE_LOG(LogMovieRenderPipeline, Log, TEXT("Updating the PIE Executor to allow using unsaved levels."));
+	UE_LOGF(LogMovieRenderPipeline, Log, "Updating the PIE Executor to allow using unsaved levels.");
 	
 	bAllowUnsavedLevels = bShouldAllow;
 }
@@ -153,7 +153,7 @@ void UMoviePipelinePIEExecutor::Start(const UMoviePipelineExecutorJob* InJob)
 			}
 			else
 			{
-				UE_LOG(LogMovieRenderPipeline, Warning, TEXT("Failed to evaluate graph to determine output resolution. Traversal error: %s"), *OutTraversalError);
+				UE_LOGF(LogMovieRenderPipeline, Warning, "Failed to evaluate graph to determine output resolution. Traversal error: %ls", *OutTraversalError);
 			}
 		}
 		else if (InJob->IsUsingBasicConfiguration())
@@ -250,7 +250,7 @@ void UMoviePipelinePIEExecutor::OnPIEStartupFinished(bool)
 	if(!ExecutingWorld)
 	{
 		// This only happens if PIE startup fails and they've usually gotten a pop-up dialog already.
-		UE_LOG(LogMovieRenderPipeline, Error, TEXT("Failed to find a PIE UWorld after OnPIEStartupFinished!"));
+		UE_LOGF(LogMovieRenderPipeline, Error, "Failed to find a PIE UWorld after OnPIEStartupFinished!");
 		OnExecutorFinishedImpl();
 		return;
 	}
@@ -265,7 +265,7 @@ void UMoviePipelinePIEExecutor::OnPIEStartupFinished(bool)
 	TSubclassOf<UMoviePipelineBase> PipelineClass = ProjectSettings->DefaultPipeline.TryLoadClass<UMoviePipelineBase>();
 	if (!PipelineClass)
 	{
-		UE_LOG(LogMovieRenderPipeline, Error, TEXT("Failed to load project specified MoviePipeline class type!"));
+		UE_LOGF(LogMovieRenderPipeline, Error, "Failed to load project specified MoviePipeline class type!");
 		OnExecutorFinishedImpl();
 		return;
 	}
@@ -275,12 +275,12 @@ void UMoviePipelinePIEExecutor::OnPIEStartupFinished(bool)
 	const bool bNeedsGraphPipeline = CurrentJob->IsUsingGraphConfiguration() || CurrentJob->IsUsingBasicConfiguration();
 	if (bNeedsGraphPipeline && !PipelineClass->IsChildOf(UMovieGraphPipeline::StaticClass()))
 	{
-		UE_LOG(LogMovieRenderPipeline, Log, TEXT("Using the default MovieGraphPipeline; the pipeline specified in Project Settings (%s) is not compatible with graphs."), *PipelineClass->GetName());
+		UE_LOGF(LogMovieRenderPipeline, Log, "Using the default MovieGraphPipeline; the pipeline specified in Project Settings (%ls) is not compatible with graphs.", *PipelineClass->GetName());
 		PipelineClass = UMovieGraphPipeline::StaticClass();
 	}
 	else if (!bNeedsGraphPipeline && !PipelineClass->IsChildOf(UMoviePipeline::StaticClass()))
 	{
-		UE_LOG(LogMovieRenderPipeline, Log, TEXT("Using the default MoviePipeline; the pipeline specified in Project Settings (%s) is not compatible with normal (non-graph) presets."), *PipelineClass->GetName());
+		UE_LOGF(LogMovieRenderPipeline, Log, "Using the default MoviePipeline; the pipeline specified in Project Settings (%ls) is not compatible with normal (non-graph) presets.", *PipelineClass->GetName());
 		PipelineClass = UMoviePipeline::StaticClass();
 	}
 
@@ -481,11 +481,11 @@ void UMoviePipelinePIEExecutor::OnPIEEnded(bool)
 	// Only call Shutdown if the pipeline hasn't been finished.
 	if (ActiveMoviePipeline && ActiveMoviePipeline->GetPipelineState() != EMovieRenderPipelineState::Finished)
 	{
-		UE_LOG(LogMovieRenderPipeline, Log, TEXT("PIE Ended while Movie Pipeline was still active. Stalling to do full shutdown."));
+		UE_LOGF(LogMovieRenderPipeline, Log, "PIE Ended while Movie Pipeline was still active. Stalling to do full shutdown.");
 
 		// This will flush any outstanding work on the movie pipeline (file writes) immediately
 		ActiveMoviePipeline->Shutdown(true);
-		UE_LOG(LogMovieRenderPipeline, Log, TEXT("MoviePipelinePIEExecutor: Stalling finished, pipeline has shut down."));
+		UE_LOGF(LogMovieRenderPipeline, Log, "MoviePipelinePIEExecutor: Stalling finished, pipeline has shut down.");
 	}
 
 	// Cache off the output data so we can use it in DelayedFinishNotification, at which point ActiveMoviePipeline is null.	
diff --git a/Engine/Plugins/VirtualProduction/Avalanche/Source/AvalancheTransition/Private/Behavior/AvaTransitionBehaviorInstance.cpp b/Engine/Plugins/VirtualProduction/Avalanche/Source/AvalancheTransition/Private/Behavior/AvaTransitionBehaviorInstance.cpp
--- a/Engine/Plugins/VirtualProduction/Avalanche/Source/AvalancheTransition/Private/Behavior/AvaTransitionBehaviorInstance.cpp
+++ b/Engine/Plugins/VirtualProduction/Avalanche/Source/AvalancheTransition/Private/Behavior/AvaTransitionBehaviorInstance.cpp
@@ -181,8 +181,8 @@ bool FAvaTransitionBehaviorInstance::ValidateTransitionScene()
 	if (!TransitionSceneOwner.IsValid())
 	{
 		//todo: support having a behavior instance run without a Transition Scene instanced
-		UE_LOG(LogAvaTransition, Error
-			, TEXT("SetContextRequirements failed for '%s'. Transition Scene Owner is invalid!")
+		UE_LOGF(LogAvaTransition, Error
+			, "SetContextRequirements failed for '%ls'. Transition Scene Owner is invalid!"
 			, *LogContext);
 		return false;
 	}
@@ -190,8 +190,8 @@ bool FAvaTransitionBehaviorInstance::ValidateTransitionScene()
 	FAvaTransitionScene* TransitionScene = TransitionContext.GetTransitionScene();
 	if (!TransitionScene)
 	{
-		UE_LOG(LogAvaTransition, Error
-			, TEXT("SetContextRequirements failed for '%s'. Transition Scene is null")
+		UE_LOGF(LogAvaTransition, Error
+			, "SetContextRequirements failed for '%ls'. Transition Scene is null"
 			, *LogContext);
 		return false;
 	}
@@ -235,8 +235,8 @@ bool FAvaTransitionBehaviorInstance::MakeContext(IAvaTransitionBehavior* InBehav
 
 	if (!IsValid(World))
 	{
-		UE_LOG(LogAvaTransition, Error
-			, TEXT("SetContextRequirements failed for '%s'. World is invalid")
+		UE_LOGF(LogAvaTransition, Error
+			, "SetContextRequirements failed for '%ls'. World is invalid"
 			, *LogContext);
 		return false;
 	}
@@ -244,8 +244,8 @@ bool FAvaTransitionBehaviorInstance::MakeContext(IAvaTransitionBehavior* InBehav
 	const FStateTreeReference& StateTreeReference = InBehavior->GetStateTreeReference();
 	if (!StateTreeReference.GetStateTree())
 	{
-		UE_LOG(LogAvaTransition, Error
-			, TEXT("SetContextRequirements failed for '%s'. StateTree is invalid")
+		UE_LOGF(LogAvaTransition, Error
+			, "SetContextRequirements failed for '%ls'. StateTree is invalid"
 			, *LogContext);
 		return false;
 	}
@@ -286,8 +286,8 @@ bool FAvaTransitionBehaviorInstance::MakeContext(IAvaTransitionBehavior* InBehav
 
 	if (!Context.AreContextDataViewsValid())
 	{
-		UE_LOG(LogAvaTransition, Error
-			, TEXT("SetContextRequirements failed for '%s'. Missing external data requirements. StateTree will not update.")
+		UE_LOGF(LogAvaTransition, Error
+			, "SetContextRequirements failed for '%ls'. Missing external data requirements. StateTree will not update."
 			, *LogContext);
 
 		return false;
diff --git a/Engine/Source/Runtime/Engine/Private/Components/CharacterMovementComponent.cpp b/Engine/Source/Runtime/Engine/Private/Components/CharacterMovementComponent.cpp
--- a/Engine/Source/Runtime/Engine/Private/Components/CharacterMovementComponent.cpp
+++ b/Engine/Source/Runtime/Engine/Private/Components/CharacterMovementComponent.cpp
@@ -991,14 +991,14 @@ void UCharacterMovementComponent::SetUpdatedComponent(USceneComponent* NewUpdate
 		const ACharacter* NewCharacterOwner = Cast<ACharacter>(NewUpdatedComponent->GetOwner());
 		if (NewCharacterOwner == nullptr)
 		{
-			UE_LOG(LogCharacterMovement, Error, TEXT("%s owned by %s must update a component owned by a Character"), *GetName(), *GetNameSafe(NewUpdatedComponent->GetOwner()));
+			UE_LOGF(LogCharacterMovement, Error, "%ls owned by %ls must update a component owned by a Character", *GetName(), *GetNameSafe(NewUpdatedComponent->GetOwner()));
 			return;
 		}
 
 		// check that UpdatedComponent is a Capsule
 		if (Cast<UCapsuleComponent>(NewUpdatedComponent) == nullptr)
 		{
-			UE_LOG(LogCharacterMovement, Error, TEXT("%s owned by %s must update a capsule component"), *GetName(), *GetNameSafe(NewUpdatedComponent->GetOwner()));
+			UE_LOGF(LogCharacterMovement, Error, "%ls owned by %ls must update a capsule component", *GetName(), *GetNameSafe(NewUpdatedComponent->GetOwner()));
 			return;
 		}
 	}
@@ -1123,7 +1123,7 @@ FVector UCharacterMovementComponent::GetPawnCapsuleExtent(const EShrinkCapsuleEx
 		break;
 
 	default:
-		UE_LOG(LogCharacterMovement, Warning, TEXT("Unknown EShrinkCapsuleExtent in UCharacterMovementComponent::GetCapsuleExtent"));
+		UE_LOGF(LogCharacterMovement, Warning, "Unknown EShrinkCapsuleExtent in UCharacterMovementComponent::GetCapsuleExtent");
 		break;
 	}
 
@@ -1864,11 +1864,11 @@ void UCharacterMovementComponent::AdjustProxyCapsuleSize()
 
 		if (NewRadius == 0.f || NewHalfHeight == 0.f)
 		{
-			UE_LOG(LogCharacterMovement, Warning, TEXT("Invalid attempt to shrink Proxy capsule for %s to zero dimension!"), *CharacterOwner->GetName());
+			UE_LOGF(LogCharacterMovement, Warning, "Invalid attempt to shrink Proxy capsule for %ls to zero dimension!", *CharacterOwner->GetName());
 			return;
 		}
 
-		UE_LOG(LogCharacterMovement, Verbose, TEXT("Shrinking capsule for %s from (r=%.3f, h=%.3f) to (r=%.3f, h=%.3f)"), *CharacterOwner->GetName(),
+		UE_LOGF(LogCharacterMovement, Verbose, "Shrinking capsule for %ls from (r=%.3f, h=%.3f) to (r=%.3f, h=%.3f)", *CharacterOwner->GetName(),
 			Radius * ComponentScale, HalfHeight * ComponentScale, NewRadius * ComponentScale, NewHalfHeight * ComponentScale);
 
 		CharacterOwner->GetCapsuleComponent()->SetCapsuleSize(NewRadius, NewHalfHeight, true);
@@ -1885,7 +1885,7 @@ void UCharacterMovementComponent::SimulatedTick(float DeltaSeconds)
 	if (CharacterOwner->IsPlayingNetworkedRootMotionMontage())
 	{
 		bWasSimulatingRootMotion = true;
-		UE_LOG(LogRootMotion, Verbose, TEXT("UCharacterMovementComponent::SimulatedTick"));
+		UE_LOGF(LogRootMotion, Verbose, "UCharacterMovementComponent::SimulatedTick");
 
 		// Tick animations before physics.
 		if( CharacterOwner && CharacterOwner->GetMesh() )
@@ -1921,7 +1921,7 @@ void UCharacterMovementComponent::SimulatedTick(float DeltaSeconds)
 				DrawDebugLine(GetWorld(), OldLocation, NewLocation, FColor::Red, false, 10.f);
 
 				const FMovementBaseInterfaceData* MovementBaseInterfaceData = GetMovementBaseInterfaceData();
-				UE_LOG(LogRootMotion, Log,  TEXT("UCharacterMovementComponent::SimulatedTick DeltaMovement Translation: %s, Rotation: %s, MovementBase: %s"),
+				UE_LOGF(LogRootMotion, Log,  "UCharacterMovementComponent::SimulatedTick DeltaMovement Translation: %ls, Rotation: %ls, MovementBase: %ls",
 					*(NewLocation - OldLocation).ToCompactString(), *(NewRotation - OldRotation).GetNormalized().ToCompactString(), *GetNameSafe(MovementBaseInterfaceData ? MovementBaseInterfaceData->GetMovementBaseObject() : nullptr));
 			}
 			#endif
@@ -1951,7 +1951,7 @@ void UCharacterMovementComponent::SimulatedTick(float DeltaSeconds)
 	{
 		// We have root motion sources and possibly animated root motion
 		bWasSimulatingRootMotion = true;
-		UE_LOG(LogRootMotion, Verbose, TEXT("UCharacterMovementComponent::SimulatedTick"));
+		UE_LOGF(LogRootMotion, Verbose, "UCharacterMovementComponent::SimulatedTick");
 
 		// If we have RootMotionRepMoves, find the most recent important one and set position/rotation to it
 		bool bCorrectedToServer = false;
@@ -1979,7 +1979,7 @@ void UCharacterMovementComponent::SimulatedTick(float DeltaSeconds)
 			CurrentRootMotion.UpdateStateFrom(RootMotionRepMove.RootMotion.AuthoritativeRootMotion, true);
 
 			// Clear out existing RootMotionRepMoves since we've consumed the most recent
-			UE_LOG(LogRootMotion, Log,  TEXT("\tClearing old moves in SimulatedTick (%d)"), CharacterOwner->RootMotionRepMoves.Num());
+			UE_LOGF(LogRootMotion, Log,  "\tClearing old moves in SimulatedTick (%d)", CharacterOwner->RootMotionRepMoves.Num());
 			CharacterOwner->RootMotionRepMoves.Reset();
 		}
 
@@ -2090,7 +2090,7 @@ void UCharacterMovementComponent::SimulatedTick(float DeltaSeconds)
 	}
 	else
 	{
-		UE_LOG(LogCharacterNetSmoothing, Verbose, TEXT("Skipping network smoothing for %s."), *GetNameSafe(CharacterOwner));
+		UE_LOGF(LogCharacterNetSmoothing, Verbose, "Skipping network smoothing for %ls.", *GetNameSafe(CharacterOwner));
 	}
 }
 
@@ -2212,7 +2212,7 @@ void UCharacterMovementComponent::SimulateMovement(float DeltaSeconds)
 	// If base is not resolved on the client, we should not try to simulate at all
 	if (!RepMovementBaseInterfaceData.IsValid() && CharacterOwner->GetReplicatedBasedMovement().bServerHasBaseComponent)
 	{
-		UE_LOG(LogCharacterMovement, Verbose, TEXT("Base for simulated character '%s' is not resolved on client, skipping SimulateMovement"), *CharacterOwner->GetName());
+		UE_LOGF(LogCharacterMovement, Verbose, "Base for simulated character '%ls' is not resolved on client, skipping SimulateMovement", *CharacterOwner->GetName());
 		return;
 	}
 
@@ -2231,7 +2231,7 @@ void UCharacterMovementComponent::SimulateMovement(float DeltaSeconds)
 			{
 				bNetworkUpdateReceived = false;
 				bHandledNetUpdate = true;
-				UE_LOG(LogCharacterMovement, Verbose, TEXT("Proxy %s received net update"), *CharacterOwner->GetName());
+				UE_LOGF(LogCharacterMovement, Verbose, "Proxy %ls received net update", *CharacterOwner->GetName());
 				if (bNetworkGravityDirectionChanged)
 				{
 					SetGravityDirection(CharacterOwner->GetReplicatedGravityDirection());
@@ -2291,7 +2291,7 @@ void UCharacterMovementComponent::SimulateMovement(float DeltaSeconds)
 		// May only need to simulate forward on frames where we haven't just received a new position update.
 		if (!bHandledNetUpdate || !bNetworkSkipProxyPredictionOnNetUpdate || !CharacterMovementCVars::NetEnableSkipProxyPredictionOnNetUpdate)
 		{
-			UE_LOG(LogCharacterMovement, Verbose, TEXT("Proxy %s simulating movement"), *GetNameSafe(CharacterOwner));
+			UE_LOGF(LogCharacterMovement, Verbose, "Proxy %ls simulating movement", *GetNameSafe(CharacterOwner));
 			FStepDownResult StepDownResult;
 			MoveSmooth(Velocity, DeltaSeconds, &StepDownResult);
 
@@ -2367,7 +2367,7 @@ void UCharacterMovementComponent::SimulateMovement(float DeltaSeconds)
 		}
 		else
 		{
-			UE_LOG(LogCharacterMovement, Verbose, TEXT("Proxy %s SKIPPING simulate movement"), *GetNameSafe(CharacterOwner));
+			UE_LOGF(LogCharacterMovement, Verbose, "Proxy %ls SKIPPING simulate movement", *GetNameSafe(CharacterOwner));
 		}
 
 		UpdateCharacterStateAfterMovement(DeltaSeconds);
@@ -2972,7 +2972,7 @@ void UCharacterMovementComponent::PerformMovement(float DeltaSeconds)
 					}
 				}
 				
-				UE_LOG(LogRootMotion, Log,  TEXT("PerformMovement WorldSpaceRootMotion Translation: %s, Rotation: %s, Actor Facing: %s, Velocity: %s")
+				UE_LOGF(LogRootMotion, Log,  "PerformMovement WorldSpaceRootMotion Translation: %ls, Rotation: %ls, Actor Facing: %ls, Velocity: %ls"
 					, *RootMotionParams.GetRootMotionTransform().GetTranslation().ToCompactString()
 					, *RootMotionParams.GetRootMotionTransform().GetRotation().Rotator().ToCompactString()
 					, *CharacterOwner->GetActorForwardVector().ToCompactString()
@@ -3072,12 +3072,12 @@ void UCharacterMovementComponent::PerformMovement(float DeltaSeconds)
 
 				const FMovementBaseInterfaceData* MovementBaseInterfaceData = GetMovementBaseInterfaceData();
 				// Log details.
-				UE_LOG(LogRootMotion, Warning,  TEXT("PerformMovement Resulting DeltaMove Translation: %s, Rotation: %s, MovementBase: %s"),
+				UE_LOGF(LogRootMotion, Warning,  "PerformMovement Resulting DeltaMove Translation: %ls, Rotation: %ls, MovementBase: %ls",
 					*(ResultingLocation - OldLocation).ToCompactString(), *(ResultingRotation - OldActorRotation).GetNormalized().ToCompactString(), *GetNameSafe(MovementBaseInterfaceData ? MovementBaseInterfaceData->GetMovementBaseObject() : nullptr));
 
 				const FVector RMTranslation = RootMotionParams.GetRootMotionTransform().GetTranslation();
 				const FRotator RMRotation = RootMotionParams.GetRootMotionTransform().GetRotation().Rotator();
-				UE_LOG(LogRootMotion, Warning,  TEXT("PerformMovement Resulting DeltaError Translation: %s, Rotation: %s"),
+				UE_LOGF(LogRootMotion, Warning,  "PerformMovement Resulting DeltaError Translation: %ls, Rotation: %ls",
 					*(ResultingLocation - OldLocation - RMTranslation).ToCompactString(), *(ResultingRotation - OldActorRotation - RMRotation).GetNormalized().ToCompactString() );
 			}
 			#endif
@@ -3536,7 +3536,7 @@ void UCharacterMovementComponent::StartNewPhysics(float deltaTime, int32 Iterati
 
 	if (UpdatedComponent->IsSimulatingPhysics())
 	{
-		UE_LOG(LogCharacterMovement, Log, TEXT("UCharacterMovementComponent::StartNewPhysics: UpdateComponent (%s) is simulating physics - aborting."), *UpdatedComponent->GetPathName());
+		UE_LOGF(LogCharacterMovement, Log, "UCharacterMovementComponent::StartNewPhysics: UpdateComponent (%ls) is simulating physics - aborting.", *UpdatedComponent->GetPathName());
 		return;
 	}
 
@@ -3566,7 +3566,7 @@ void UCharacterMovementComponent::StartNewPhysics(float deltaTime, int32 Iterati
 		PhysCustom(deltaTime, Iterations);
 		break;
 	default:
-		UE_LOG(LogCharacterMovement, Warning, TEXT("%s has unsupported movement mode %d"), *CharacterOwner->GetName(), int32(MovementMode));
+		UE_LOGF(LogCharacterMovement, Warning, "%ls has unsupported movement mode %d", *CharacterOwner->GetName(), int32(MovementMode));
 		SetMovementMode(MOVE_None);
 		break;
 	}
@@ -4862,7 +4862,7 @@ void UCharacterMovementComponent::SetGravityDirection(const FVector& InNewGravit
 	{
 		if (!GravityDirection.Equals(NewGravityDir))
 		{
-			UE_LOG(LogCharacterMovement, Verbose, TEXT("SetGravityDirection: From(%s) To(%s)"), *GravityDirection.ToCompactString(), *NewGravityDir.ToCompactString());
+			UE_LOGF(LogCharacterMovement, Verbose, "SetGravityDirection: From(%ls) To(%ls)", *GravityDirection.ToCompactString(), *NewGravityDir.ToCompactString());
 			GravityDirection = NewGravityDir;
 			WorldToGravityTransform = FQuat::FindBetweenNormals(FVector::UpVector, -NewGravityDir);
 			GravityToWorldTransform = WorldToGravityTransform.Inverse();
@@ -4960,7 +4960,7 @@ void UCharacterMovementComponent::PhysFalling(float deltaTime, int32 Iterations)
 		// Apply gravity
 		Velocity = NewFallVelocity(Velocity, Gravity, GravityTime);
 
-		//UE_LOG(LogCharacterMovement, Log, TEXT("dt=(%.6f) OldLocation=(%s) OldVelocity=(%s) OldVelocityWithRootMotion=(%s) NewVelocity=(%s)"), timeTick, *(UpdatedComponent->GetComponentLocation()).ToString(), *OldVelocity.ToString(), *OldVelocityWithRootMotion.ToString(), *Velocity.ToString());
+		//UE_LOGF(LogCharacterMovement, Log, "dt=(%.6f) OldLocation=(%ls) OldVelocity=(%ls) OldVelocityWithRootMotion=(%ls) NewVelocity=(%ls)", timeTick, *(UpdatedComponent->GetComponentLocation()).ToString(), *OldVelocity.ToString(), *OldVelocityWithRootMotion.ToString(), *Velocity.ToString());
 		ApplyRootMotionToVelocity(timeTick);
 		DecayFormerBaseVelocity(timeTick);
 
@@ -5449,10 +5449,10 @@ void UCharacterMovementComponent::StartFalling(int32 Iterations, float remaining
 
 void UCharacterMovementComponent::RevertMove(const FVector& OldLocation, FMovementBaseInterfaceData* OldMovementBaseInterfaceData, const FVector& PreviousBaseLocation, const FFindFloorResult& OldFloor, bool bFailMove)
 {
-	//UE_LOG(LogCharacterMovement, Log, TEXT("RevertMove from %f %f %f to %f %f %f"), CharacterOwner->Location.X, CharacterOwner->Location.Y, CharacterOwner->Location.Z, OldLocation.X, OldLocation.Y, OldLocation.Z);
+	//UE_LOGF(LogCharacterMovement, Log, "RevertMove from %f %f %f to %f %f %f", CharacterOwner->Location.X, CharacterOwner->Location.Y, CharacterOwner->Location.Z, OldLocation.X, OldLocation.Y, OldLocation.Z);
 	UpdatedComponent->SetWorldLocation(OldLocation, false, nullptr, GetTeleportType());
 	
-	//UE_LOG(LogCharacterMovement, Log, TEXT("Now at %f %f %f"), CharacterOwner->Location.X, CharacterOwner->Location.Y, CharacterOwner->Location.Z);
+	//UE_LOGF(LogCharacterMovement, Log, "Now at %f %f %f", CharacterOwner->Location.X, CharacterOwner->Location.Y, CharacterOwner->Location.Z);
 	bJustTeleported = false;
 	// if our previous base couldn't have moved or changed in any physics-affecting way, restore it
 	if (MovementBaseUtility::IsMovementBaseDataValid(OldMovementBaseInterfaceData) && 
@@ -5476,7 +5476,7 @@ void UCharacterMovementComponent::RevertMove(const FVector& OldLocation, FMoveme
 		// end movement now
 		Velocity = FVector::ZeroVector;
 		Acceleration = FVector::ZeroVector;
-		//UE_LOG(LogCharacterMovement, Log, TEXT("%s FAILMOVE RevertMove"), *CharacterOwner->GetName());
+		//UE_LOGF(LogCharacterMovement, Log, "%ls FAILMOVE RevertMove", *CharacterOwner->GetName());
 	}
 }
 
@@ -5516,11 +5516,11 @@ void UCharacterMovementComponent::OnCharacterStuckInGeometry(const FHitResult* H
 			LastStuckWarningTime = RealTimeSeconds;
 			if (Hit == nullptr)
 			{
-				UE_LOG(LogCharacterMovement, Log, TEXT("%s is stuck and failed to move! (%d other events since notify)"), *CharacterOwner->GetName(), StuckWarningCountSinceNotify);
+				UE_LOGF(LogCharacterMovement, Log, "%ls is stuck and failed to move! (%d other events since notify)", *CharacterOwner->GetName(), StuckWarningCountSinceNotify);
 			}
 			else
 			{
-				UE_LOG(LogCharacterMovement, Log, TEXT("%s is stuck and failed to move! Velocity: X=%3.2f Y=%3.2f Z=%3.2f Location: X=%3.2f Y=%3.2f Z=%3.2f Normal: X=%3.2f Y=%3.2f Z=%3.2f PenetrationDepth:%.3f Actor:%s Component:%s BoneName:%s (%d other events since notify)"),
+				UE_LOGF(LogCharacterMovement, Log, "%ls is stuck and failed to move! Velocity: X=%3.2f Y=%3.2f Z=%3.2f Location: X=%3.2f Y=%3.2f Z=%3.2f Normal: X=%3.2f Y=%3.2f Z=%3.2f PenetrationDepth:%.3f Actor:%ls Component:%ls BoneName:%ls (%d other events since notify)",
 					   *GetNameSafe(CharacterOwner),
 					   Velocity.X, Velocity.Y, Velocity.Z,
 					   Hit->Location.X, Hit->Location.Y, Hit->Location.Z,
@@ -5597,13 +5597,13 @@ void UCharacterMovementComponent::MoveAlongFloor(const FVector& InVelocity, floa
 				const FVector PreStepUpLocation = UpdatedComponent->GetComponentLocation();
 				if (!StepUp(GetGravityDirection(), Delta * (1.f - PercentTimeApplied), Hit, OutStepDownResult))
 				{
-					UE_LOG(LogCharacterMovement, Verbose, TEXT("- StepUp (ImpactNormal %s, Normal %s"), *Hit.ImpactNormal.ToString(), *Hit.Normal.ToString());
+					UE_LOGF(LogCharacterMovement, Verbose, "- StepUp (ImpactNormal %ls, Normal %ls", *Hit.ImpactNormal.ToString(), *Hit.Normal.ToString());
 					HandleImpact(Hit, LastMoveTimeSlice, RampVector);
 					SlideAlongSurface(Delta, 1.f - PercentTimeApplied, Hit.Normal, Hit, true);
 				}
 				else
 				{
-					UE_LOG(LogCharacterMovement, Verbose, TEXT("+ StepUp (ImpactNormal %s, Normal %s"), *Hit.ImpactNormal.ToString(), *Hit.Normal.ToString());
+					UE_LOGF(LogCharacterMovement, Verbose, "+ StepUp (ImpactNormal %ls, Normal %ls", *Hit.ImpactNormal.ToString(), *Hit.Normal.ToString());
 					if (!bMaintainHorizontalGroundVelocity)
 					{
 						// Don't recalculate velocity based on this height adjustment, if considering vertical adjustments. Only consider horizontal movement.
@@ -5972,7 +5972,7 @@ void UCharacterMovementComponent::PhysNavWalking(float deltaTime, int32 Iteratio
 	if (bDeltaMoveNearlyZero && bSameNavLocation)
 	{
 		DestNavLocation = CachedNavLocation;
-		UE_LOG(LogNavMeshMovement, VeryVerbose, TEXT("%s using cached navmesh location! (bProjectNavMeshWalking = %d)"), *GetNameSafe(CharacterOwner), bProjectNavMeshWalking);
+		UE_LOGF(LogNavMeshMovement, VeryVerbose, "%ls using cached navmesh location! (bProjectNavMeshWalking = %d)", *GetNameSafe(CharacterOwner), bProjectNavMeshWalking);
 	}
 	else
 	{
@@ -6114,7 +6114,7 @@ FVector UCharacterMovementComponent::ProjectLocationFromNavMesh(float DeltaSecon
 	{
 		if (!bCachedLocationStillValid || bAlwaysCheckFloor)
 		{
-			UE_LOG(LogNavMeshMovement, VeryVerbose, TEXT("ProjectLocationFromNavMesh(): %s interval: %.3f velocity: %s"), *GetNameSafe(CharacterOwner), NavMeshProjectionInterval, *Velocity.ToString());
+			UE_LOGF(LogNavMeshMovement, VeryVerbose, "ProjectLocationFromNavMesh(): %ls interval: %.3f velocity: %ls", *GetNameSafe(CharacterOwner), NavMeshProjectionInterval, *Velocity.ToString());
 
 			FHitResult HitResult;
 			FindBestNavMeshLocation(TraceStart, TraceEnd, CurrentFeetLocation, TargetNavLocation, HitResult);
@@ -6131,7 +6131,7 @@ FVector UCharacterMovementComponent::ProjectLocationFromNavMesh(float DeltaSecon
 		}
 		else
 		{
-			UE_LOG(LogNavMeshMovement, VeryVerbose, TEXT("ProjectLocationFromNavMesh(): %s interval: %.3f velocity: %s [SKIP TRACE]"), *GetNameSafe(CharacterOwner), NavMeshProjectionInterval, *Velocity.ToString());
+			UE_LOGF(LogNavMeshMovement, VeryVerbose, "ProjectLocationFromNavMesh(): %ls interval: %.3f velocity: %ls [SKIP TRACE]", *GetNameSafe(CharacterOwner), NavMeshProjectionInterval, *Velocity.ToString());
 		}
 
 		// Wrap around to maintain same relative offset to tick time changes.
@@ -6303,7 +6303,7 @@ void UCharacterMovementComponent::AdjustFloorHeight()
 		if (OldFloorDist < MIN_FLOOR_DIST && CurrentFloor.LineDist >= MIN_FLOOR_DIST)
 		{
 			// This would cause us to scale unwalkable walls
-			UE_LOG(LogCharacterMovement, VeryVerbose, TEXT("Adjust floor height aborting due to line trace with small floor distance (line: %.2f, sweep: %.2f)"), CurrentFloor.LineDist, CurrentFloor.FloorDist);
+			UE_LOGF(LogCharacterMovement, VeryVerbose, "Adjust floor height aborting due to line trace with small floor distance (line: %.2f, sweep: %.2f)", CurrentFloor.LineDist, CurrentFloor.FloorDist);
 			return;
 		}
 		else
@@ -6321,7 +6321,7 @@ void UCharacterMovementComponent::AdjustFloorHeight()
 		const float AvgFloorDist = (MIN_FLOOR_DIST + MAX_FLOOR_DIST) * 0.5f;
 		const float MoveDist = AvgFloorDist - OldFloorDist;
 		SafeMoveUpdatedComponent(MoveDist * -GetGravityDirection(), UpdatedComponent->GetComponentQuat(), true, AdjustHit );
-		UE_LOG(LogCharacterMovement, VeryVerbose, TEXT("Adjust floor height %.3f (Hit = %d)"), MoveDist, AdjustHit.bBlockingHit);
+		UE_LOGF(LogCharacterMovement, VeryVerbose, "Adjust floor height %.3f (Hit = %d)", MoveDist, AdjustHit.bBlockingHit);
 
 		if (!AdjustHit.IsValidBlockingHit())
 		{
@@ -6388,7 +6388,7 @@ void UCharacterMovementComponent::ProcessLanded(const FHitResult& Hit, float rem
 			if (!bHasNavigationData || NavLocation.NodeRef == INVALID_NAVNODEREF)
 			{
 				GroundMovementMode = MOVE_Walking;
-				UE_LOG(LogNavMeshMovement, Verbose, TEXT("ProcessLanded(): %s tried to go to NavWalking but couldn't find NavMesh! Using Walking instead."), *GetNameSafe(CharacterOwner));
+				UE_LOGF(LogNavMeshMovement, Verbose, "ProcessLanded(): %ls tried to go to NavWalking but couldn't find NavMesh! Using Walking instead.", *GetNameSafe(CharacterOwner));
 			}
 		}
 
@@ -6487,7 +6487,7 @@ void UCharacterMovementComponent::SetNavWalkingPhysics(bool bEnable)
 			}
 			else
 			{
-				UE_LOG(LogCharacterMovement, Warning, TEXT("Can't revert NavWalking collision settings for %s.%s"),
+				UE_LOGF(LogCharacterMovement, Warning, "Can't revert NavWalking collision settings for %ls.%ls",
 					*GetNameSafe(CharacterOwner), *GetNameSafe(UpdatedComponent));
 			}
 		}
@@ -6868,7 +6868,7 @@ void UCharacterMovementComponent::AddImpulse( FVector Impulse, bool bVelocityCha
 			}
 			else
 			{
-				UE_LOG(LogCharacterMovement, Warning, TEXT("Attempt to apply impulse to zero or negative Mass in CharacterMovement"));
+				UE_LOGF(LogCharacterMovement, Warning, "Attempt to apply impulse to zero or negative Mass in CharacterMovement");
 			}
 		}
 
@@ -6886,7 +6886,7 @@ void UCharacterMovementComponent::AddForce( FVector Force )
 		}
 		else
 		{
-			UE_LOG(LogCharacterMovement, Warning, TEXT("Attempt to apply force to zero or negative Mass in CharacterMovement"));
+			UE_LOGF(LogCharacterMovement, Warning, "Attempt to apply force to zero or negative Mass in CharacterMovement");
 		}
 	}
 }
@@ -7052,7 +7052,7 @@ bool UCharacterMovementComponent::IsWithinEdgeTolerance(const FVector& CapsuleLo
 
 void UCharacterMovementComponent::ComputeFloorDist(const FVector& CapsuleLocation, float LineDistance, float SweepDistance, FFindFloorResult& OutFloorResult, float SweepRadius, const FHitResult* DownwardSweepResult) const
 {
-	UE_LOG(LogCharacterMovement, VeryVerbose, TEXT("[Role:%d] ComputeFloorDist: %s at location %s"), (int32)CharacterOwner->GetLocalRole(), *GetNameSafe(CharacterOwner), *CapsuleLocation.ToString());
+	UE_LOGF(LogCharacterMovement, VeryVerbose, "[Role:%d] ComputeFloorDist: %ls at location %ls", (int32)CharacterOwner->GetLocalRole(), *GetNameSafe(CharacterOwner), *CapsuleLocation.ToString());
 	OutFloorResult.Clear();
 
 	float PawnRadius, PawnHalfHeight;
@@ -7212,7 +7212,7 @@ void UCharacterMovementComponent::FindFloor(const FVector& CapsuleLocation, FFin
 		return;
 	}
 
-	UE_LOG(LogCharacterMovement, VeryVerbose, TEXT("[Role:%d] FindFloor: %s at location %s"), (int32)CharacterOwner->GetLocalRole(), *GetNameSafe(CharacterOwner), *CapsuleLocation.ToString());
+	UE_LOGF(LogCharacterMovement, VeryVerbose, "[Role:%d] FindFloor: %ls at location %ls", (int32)CharacterOwner->GetLocalRole(), *GetNameSafe(CharacterOwner), *CapsuleLocation.ToString());
 	check(CharacterOwner->GetCapsuleComponent());
 
 	// Increase height check slightly if walking, to prevent floor height adjustment from later invalidating the floor result.
@@ -7252,7 +7252,7 @@ void UCharacterMovementComponent::FindFloor(const FVector& CapsuleLocation, FFin
 
 			if ( !bForceNextFloorCheck && !IsActorBasePendingKill && bHasValidMovementBase)
 			{
-				//UE_LOG(LogCharacterMovement, Log, TEXT("%s SKIP check for floor"), *CharacterOwner->GetName());
+				//UE_LOGF(LogCharacterMovement, Log, "%ls SKIP check for floor", *CharacterOwner->GetName());
 				OutFloorResult = CurrentFloor;
 				bNeedToValidateFloor = false;
 			}
@@ -7693,7 +7693,7 @@ bool UCharacterMovementComponent::StepUp(const FVector& GravDir, const FVector&
 		const float DeltaZ = (Hit.ImpactPoint | -GravDir) - PawnFloorPointZ;
 		if (DeltaZ > MaxStepHeight)
 		{
-			//UE_LOG(LogCharacterMovement, VeryVerbose, TEXT("- Reject StepUp (too high Height %.3f) up from floor base %f to %f"), DeltaZ, PawnInitialFloorBaseZ, NewLocation.Z);
+			//UE_LOGF(LogCharacterMovement, VeryVerbose, "- Reject StepUp (too high Height %.3f) up from floor base %f to %f", DeltaZ, PawnInitialFloorBaseZ, NewLocation.Z);
 			ScopedStepUpMovement.RevertMove();
 			return false;
 		}
@@ -7705,7 +7705,7 @@ bool UCharacterMovementComponent::StepUp(const FVector& GravDir, const FVector&
 			const bool bNormalTowardsMe = (Delta | Hit.ImpactNormal) < 0.f;
 			if (bNormalTowardsMe)
 			{
-				//UE_LOG(LogCharacterMovement, VeryVerbose, TEXT("- Reject StepUp (unwalkable normal %s opposed to movement)"), *Hit.ImpactNormal.ToString());
+				//UE_LOGF(LogCharacterMovement, VeryVerbose, "- Reject StepUp (unwalkable normal %ls opposed to movement)", *Hit.ImpactNormal.ToString());
 				ScopedStepUpMovement.RevertMove();
 				return false;
 			}
@@ -7714,7 +7714,7 @@ bool UCharacterMovementComponent::StepUp(const FVector& GravDir, const FVector&
 			// It's fine to step down onto an unwalkable normal below us, we will just slide off. Rejecting those moves would prevent us from being able to walk off the edge.
 			if ((Hit.Location | -GravDir) > OldLocationZ)
 			{
-				//UE_LOG(LogCharacterMovement, VeryVerbose, TEXT("- Reject StepUp (unwalkable normal %s above old position)"), *Hit.ImpactNormal.ToString());
+				//UE_LOGF(LogCharacterMovement, VeryVerbose, "- Reject StepUp (unwalkable normal %ls above old position)", *Hit.ImpactNormal.ToString());
 				ScopedStepUpMovement.RevertMove();
 				return false;
 			}
@@ -7723,7 +7723,7 @@ bool UCharacterMovementComponent::StepUp(const FVector& GravDir, const FVector&
 		// Reject moves where the downward sweep hit something very close to the edge of the capsule. This maintains consistency with FindFloor as well.
 		if (!IsWithinEdgeTolerance(Hit.Location, Hit.ImpactPoint, PawnRadius))
 		{
-			//UE_LOG(LogCharacterMovement, VeryVerbose, TEXT("- Reject StepUp (outside edge tolerance)"));
+			//UE_LOGF(LogCharacterMovement, VeryVerbose, "- Reject StepUp (outside edge tolerance)");
 			ScopedStepUpMovement.RevertMove();
 			return false;
 		}
@@ -7731,7 +7731,7 @@ bool UCharacterMovementComponent::StepUp(const FVector& GravDir, const FVector&
 		// Don't step up onto invalid surfaces if traveling higher.
 		if (DeltaZ > 0.f && !CanStepUp(Hit))
 		{
-			//UE_LOG(LogCharacterMovement, VeryVerbose, TEXT("- Reject StepUp (up onto surface with !CanStepUp())"));
+			//UE_LOGF(LogCharacterMovement, VeryVerbose, "- Reject StepUp (up onto surface with !CanStepUp())");
 			ScopedStepUpMovement.RevertMove();
 			return false;
 		}
@@ -8171,7 +8171,7 @@ float UCharacterMovementComponent::GetSimulationTimeStep(float RemainingTime, in
 #if !(UE_BUILD_SHIPPING || UE_BUILD_TEST)
 			if ((s_WarningCount++ < 100) || (GFrameCounter & 15) == 0)
 			{
-				UE_LOG(LogCharacterMovement, Warning, TEXT("GetSimulationTimeStep() - Max iterations %d hit while remaining time %.6f > MaxSimulationTimeStep (%.3f) for '%s', movement '%s'"), MaxSimulationIterations, RemainingTime, MaxSimulationTimeStep, *GetNameSafe(CharacterOwner), *GetMovementName());
+				UE_LOGF(LogCharacterMovement, Warning, "GetSimulationTimeStep() - Max iterations %d hit while remaining time %.6f > MaxSimulationTimeStep (%.3f) for '%ls', movement '%ls'", MaxSimulationIterations, RemainingTime, MaxSimulationTimeStep, *GetNameSafe(CharacterOwner), *GetMovementName());
 			}
 #endif
 		}
@@ -8236,7 +8236,7 @@ void UCharacterMovementComponent::SmoothCorrection(const FVector& OldLocation, c
 			ClientData->MeshTranslationOffset = ClientData->MeshTranslationOffset + NewToOldVector;	
 		}
 
-		UE_LOG(LogCharacterNetSmoothing, Verbose, TEXT("Proxy %s SmoothCorrection(%.2f)"), *GetNameSafe(CharacterOwner), FMath::Sqrt(DistSq));
+		UE_LOGF(LogCharacterNetSmoothing, Verbose, "Proxy %ls SmoothCorrection(%.2f)", *GetNameSafe(CharacterOwner), FMath::Sqrt(DistSq));
 		if (NetworkSmoothingMode == ENetworkSmoothingMode::Linear)
 		{
 			ClientData->OriginalMeshTranslationOffset	= ClientData->MeshTranslationOffset;
@@ -8274,7 +8274,7 @@ void UCharacterMovementComponent::SmoothCorrection(const FVector& OldLocation, c
 			const double OldClientTimeStamp = ClientData->SmoothingClientTimeStamp;
 			ClientData->SmoothingClientTimeStamp = FMath::LerpStable(ClientData->SmoothingServerTimeStamp, OldClientTimeStamp, 0.5);
 
-			UE_LOG(LogCharacterNetSmoothing, VeryVerbose, TEXT("SmoothCorrection: Pull back client from ClientTimeStamp: %.6f to %.6f, ServerTimeStamp: %.6f for %s"),
+			UE_LOGF(LogCharacterNetSmoothing, VeryVerbose, "SmoothCorrection: Pull back client from ClientTimeStamp: %.6f to %.6f, ServerTimeStamp: %.6f for %ls",
 				OldClientTimeStamp, ClientData->SmoothingClientTimeStamp, ClientData->SmoothingServerTimeStamp, *GetNameSafe(CharacterOwner));
 		}
 
@@ -8307,7 +8307,7 @@ void UCharacterMovementComponent::SmoothCorrection(const FVector& OldLocation, c
 		ClientData->LastCorrectionDelta = ClientData->SmoothingServerTimeStamp - ClientData->SmoothingClientTimeStamp;
 		ClientData->LastCorrectionTime = MyWorld->GetTimeSeconds();
 
-		UE_LOG(LogCharacterNetSmoothing, VeryVerbose, TEXT("SmoothCorrection: WorldTime: %.6f, ServerTimeStamp: %.6f, ClientTimeStamp: %.6f, Delta: %.6f for %s"),
+		UE_LOGF(LogCharacterNetSmoothing, VeryVerbose, "SmoothCorrection: WorldTime: %.6f, ServerTimeStamp: %.6f, ClientTimeStamp: %.6f, Delta: %.6f for %ls",
 			MyWorld->GetTimeSeconds(), ClientData->SmoothingServerTimeStamp, ClientData->SmoothingClientTimeStamp, ClientData->LastCorrectionDelta, *GetNameSafe(CharacterOwner));
 
 #if !(UE_BUILD_SHIPPING || UE_BUILD_TEST)
@@ -8414,7 +8414,7 @@ void UCharacterMovementComponent::SmoothClientPosition_Interpolate(float DeltaSe
 				const float CurrentSmoothTime = TargetDelta - RemainingTime;
 				LerpPercent = FMath::Clamp(CurrentSmoothTime / TargetDelta, 0.0f, LerpLimit);
 
-				UE_LOG(LogCharacterNetSmoothing, VeryVerbose, TEXT("Interpolate: WorldTime: %.6f, ServerTimeStamp: %.6f, ClientTimeStamp: %.6f, Elapsed: %.6f, Alpha: %.6f for %s"),
+				UE_LOGF(LogCharacterNetSmoothing, VeryVerbose, "Interpolate: WorldTime: %.6f, ServerTimeStamp: %.6f, ClientTimeStamp: %.6f, Elapsed: %.6f, Alpha: %.6f for %ls",
 					MyWorld->GetTimeSeconds(), ClientData->SmoothingServerTimeStamp, ClientData->SmoothingClientTimeStamp, CurrentSmoothTime, LerpPercent, *GetNameSafe(CharacterOwner));
 			}
 			else
@@ -8500,7 +8500,7 @@ void UCharacterMovementComponent::SmoothClientPosition_Interpolate(float DeltaSe
 		}
 
 #if !(UE_BUILD_SHIPPING || UE_BUILD_TEST)
-		//UE_LOG(LogCharacterNetSmoothing, VeryVerbose, TEXT("SmoothClientPosition_Interpolate %s: Translation: %s Rotation: %s"),
+		//UE_LOGF(LogCharacterNetSmoothing, VeryVerbose, "SmoothClientPosition_Interpolate %ls: Translation: %ls Rotation: %ls",
 		//	*GetNameSafe(CharacterOwner), *ClientData->MeshTranslationOffset.ToString(), *ClientData->MeshRotationOffset.ToString());
 
 		if ( CharacterMovementCVars::NetVisualizeSimulatedCorrections >= 1 )
@@ -8624,14 +8624,14 @@ bool UCharacterMovementComponent::ClientUpdatePositionAfterServerUpdate()
 
 	if (ClientData->SavedMoves.Num() == 0)
 	{
-		UE_LOG(LogNetPlayerMovement, Verbose, TEXT("ClientUpdatePositionAfterServerUpdate No saved moves to replay"));
+		UE_LOGF(LogNetPlayerMovement, Verbose, "ClientUpdatePositionAfterServerUpdate No saved moves to replay");
 
 		// With no saved moves to resimulate, the move the server updated us with is the last move we've done, no resimulation needed.
 		CharacterOwner->bClientResimulateRootMotion = false;
 		if (CharacterOwner->bClientResimulateRootMotionSources)
 		{
 			// With no resimulation, we just update our current root motion to what the server sent us
-			UE_LOG(LogRootMotion, VeryVerbose, TEXT("CurrentRootMotion getting updated to ServerUpdate state: %s"), *CharacterOwner->GetName());
+			UE_LOGF(LogRootMotion, VeryVerbose, "CurrentRootMotion getting updated to ServerUpdate state: %ls", *CharacterOwner->GetName());
 			CurrentRootMotion.UpdateStateFrom(CharacterOwner->SavedRootMotion);
 			CharacterOwner->bClientResimulateRootMotionSources = false;
 		}
@@ -8657,7 +8657,7 @@ bool UCharacterMovementComponent::ClientUpdatePositionAfterServerUpdate()
 	FScopedMeshMovementUpdate ScopedMeshUpdate(CharacterOwner->GetMesh(), CharacterMovementCVars::bDeferCharacterMeshMovementForAllCorrections);
 
 	// Replay moves that have not yet been acked.
-	UE_LOG(LogNetPlayerMovement, Verbose, TEXT("ClientUpdatePositionAfterServerUpdate Replaying %d Moves, starting at Timestamp %f"), ClientData->SavedMoves.Num(), ClientData->SavedMoves[0]->TimeStamp);
+	UE_LOGF(LogNetPlayerMovement, Verbose, "ClientUpdatePositionAfterServerUpdate Replaying %d Moves, starting at Timestamp %f", ClientData->SavedMoves.Num(), ClientData->SavedMoves[0]->TimeStamp);
 	for (int32 i=0; i<ClientData->SavedMoves.Num(); i++)
 	{
 		FSavedMove_Character* const CurrentMove = ClientData->SavedMoves[i].Get();
@@ -8700,7 +8700,7 @@ bool UCharacterMovementComponent::ClientUpdatePositionAfterServerUpdate()
 		// If we were resimulating root motion sources, it's because we had mismatched state
 		// with the server - we just resimulated our SavedMoves and now need to restore
 		// CurrentRootMotion with the latest "good state"
-		UE_LOG(LogRootMotion, VeryVerbose, TEXT("CurrentRootMotion getting updated after ServerUpdate replays: %s"), *CharacterOwner->GetName());
+		UE_LOGF(LogRootMotion, VeryVerbose, "CurrentRootMotion getting updated after ServerUpdate replays: %ls", *CharacterOwner->GetName());
 		CurrentRootMotion.UpdateStateFrom(CharacterOwner->SavedRootMotion);
 		CharacterOwner->bClientResimulateRootMotionSources = false;
 	}
@@ -8755,7 +8755,7 @@ bool UCharacterMovementComponent::ForcePositionUpdate(float DeltaTime)
 	const bool bServerMoveHasOccurred = (ServerData->ServerTimeStampLastServerMove != 0.f);
 	if (bServerMoveHasOccurred)
 	{
-		UE_LOG(LogNetPlayerMovement, Log, TEXT("ForcePositionUpdate: %s (DeltaTime %.2f -> ServerTimeStamp %.2f) Pawn: %s"), *GetNameSafe(CharacterOwner), DeltaTime, ServerData->CurrentClientTimeStamp, *GetNameSafe(CharacterOwner));
+		UE_LOGF(LogNetPlayerMovement, Log, "ForcePositionUpdate: %ls (DeltaTime %.2f -> ServerTimeStamp %.2f) Pawn: %ls", *GetNameSafe(CharacterOwner), DeltaTime, ServerData->CurrentClientTimeStamp, *GetNameSafe(CharacterOwner));
 	}
 
 	// Force movement update.
@@ -8858,7 +8858,7 @@ float FNetworkPredictionData_Client_Character::UpdateTimeStampAndDeltaTime(float
 	// Reset TimeStamp regularly to combat float accuracy decreasing over time.
 	if( CurrentTimeStamp > CharacterMovementComponent.MinTimeBetweenTimeStampResets )
 	{
-		UE_LOG(LogNetPlayerMovement, Log, TEXT("Resetting Client's TimeStamp %f"), CurrentTimeStamp);
+		UE_LOGF(LogNetPlayerMovement, Log, "Resetting Client's TimeStamp %f", CurrentTimeStamp);
 		CurrentTimeStamp -= CharacterMovementComponent.MinTimeBetweenTimeStampResets;
 
 		// Mark all buffered moves as having old time stamps, so we make sure to not resend them.
@@ -8975,7 +8975,7 @@ void UCharacterMovementComponent::ReplicateMoveToServer(float DeltaTime, const F
 
 				// Accumulate multiple transform updates until scope ends.
 				FScopedCapsuleMovementUpdate ScopedMovementUpdate(UpdatedComponent, true);
-				UE_LOG(LogNetPlayerMovement, VeryVerbose, TEXT("CombineMove: add delta %f + %f and revert from %f %f to %f %f"), DeltaTime, PendingMove->DeltaTime, UpdatedComponent->GetComponentLocation().X, UpdatedComponent->GetComponentLocation().Y, OldStartLocation.X, OldStartLocation.Y);
+				UE_LOGF(LogNetPlayerMovement, VeryVerbose, "CombineMove: add delta %f + %f and revert from %f %f to %f %f", DeltaTime, PendingMove->DeltaTime, UpdatedComponent->GetComponentLocation().X, UpdatedComponent->GetComponentLocation().Y, OldStartLocation.X, OldStartLocation.Y);
 
 				NewMove->CombineWith(PendingMove, CharacterOwner, PC, OldStartLocation);
 
@@ -9000,12 +9000,12 @@ void UCharacterMovementComponent::ReplicateMoveToServer(float DeltaTime, const F
 			}
 			else
 			{
-				UE_LOG(LogNetPlayerMovement, Verbose, TEXT("Not combining move [would collide at start location]"));
+				UE_LOGF(LogNetPlayerMovement, Verbose, "Not combining move [would collide at start location]");
 			}
 		}
 		else
 		{
-			UE_LOG(LogNetPlayerMovement, Verbose, TEXT("Not combining move [not allowed by CanCombineWith()]"));
+			UE_LOGF(LogNetPlayerMovement, Verbose, "Not combining move [not allowed by CanCombineWith()]");
 		}
 	}
 
@@ -9044,7 +9044,7 @@ void UCharacterMovementComponent::ReplicateMoveToServer(float DeltaTime, const F
 
 		ClientData->ClientUpdateRealTime = MyWorld->GetRealTimeSeconds();
 
-		UE_CLOG(CharacterOwner && UpdatedComponent, LogNetPlayerMovement, VeryVerbose, TEXT("ClientMove Time %f Acceleration %s Velocity %s Position %s Rotation %s DeltaTime %f Mode %s MovementBase %s.%s (Dynamic:%d) DualMove? %d"),
+		UE_CLOGF(CharacterOwner && UpdatedComponent, LogNetPlayerMovement, VeryVerbose, "ClientMove Time %f Acceleration %ls Velocity %ls Position %ls Rotation %ls DeltaTime %f Mode %ls MovementBase %ls.%ls (Dynamic:%d) DualMove? %d",
 			NewMove->TimeStamp, *NewMove->Acceleration.ToString(), *Velocity.ToString(), *UpdatedComponent->GetComponentLocation().ToString(), *UpdatedComponent->GetComponentRotation().ToCompactString(), NewMove->DeltaTime, *GetMovementName(),
 			*GetNameSafe(NewMove->EndMovementBaseInterfaceData.GetMovementBaseObject()), *NewMove->EndBoneName.ToString(),
 			MovementBaseUtility::IsDynamicBase(&NewMove->EndMovementBaseInterfaceData) ? 1 : 0, ClientData->PendingMove.IsValid() ? 1 : 0);
@@ -9058,13 +9058,13 @@ void UCharacterMovementComponent::ReplicateMoveToServer(float DeltaTime, const F
 		if (ClientData->DebugForcedPacketLossTimerStart > 0.f && (TimeSinceLossStart < CharacterMovementCVars::NetForceClientServerMoveLossDuration))
 		{
 			bSendServerMove = false;
-			UE_LOG(LogNetPlayerMovement, Log, TEXT("Drop ServerMove, %.2f time remains"), CharacterMovementCVars::NetForceClientServerMoveLossDuration - TimeSinceLossStart);
+			UE_LOGF(LogNetPlayerMovement, Log, "Drop ServerMove, %.2f time remains", CharacterMovementCVars::NetForceClientServerMoveLossDuration - TimeSinceLossStart);
 		}
 		else if (CharacterMovementCVars::NetForceClientServerMoveLossPercent != 0.f && (RandomStream.FRand() < CharacterMovementCVars::NetForceClientServerMoveLossPercent))
 		{
 			bSendServerMove = false;
 			ClientData->DebugForcedPacketLossTimerStart = (CharacterMovementCVars::NetForceClientServerMoveLossDuration > 0) ? MyWorld->RealTimeSeconds : 0.0f;
-			UE_LOG(LogNetPlayerMovement, Log, TEXT("Drop ServerMove, %.2f time remains"), CharacterMovementCVars::NetForceClientServerMoveLossDuration);
+			UE_LOGF(LogNetPlayerMovement, Log, "Drop ServerMove, %.2f time remains", CharacterMovementCVars::NetForceClientServerMoveLossDuration);
 		}
 		else
 		{
@@ -9155,7 +9155,7 @@ void UCharacterMovementComponent::CallServerMovePacked(const FSavedMove_Characte
 
 	if (ServerMoveBitWriter.PackageMap == nullptr)
 	{
-		UE_LOG(LogNetPlayerMovement, Error, TEXT("CallServerMovePacked: Failed to find a NetConnection/PackageMap for data serialization!"));
+		UE_LOGF(LogNetPlayerMovement, Error, "CallServerMovePacked: Failed to find a NetConnection/PackageMap for data serialization!");
 		return;
 	}
 
@@ -9166,7 +9166,7 @@ void UCharacterMovementComponent::CallServerMovePacked(const FSavedMove_Characte
 	// Serialize move struct into a bit stream
 	if (!MoveDataContainer.Serialize(*this, ServerMoveBitWriter, ServerMoveBitWriter.PackageMap) || ServerMoveBitWriter.IsError())
 	{
-		UE_LOG(LogNetPlayerMovement, Error, TEXT("CallServerMovePacked: Failed to serialize out movement data!"));
+		UE_LOGF(LogNetPlayerMovement, Error, "CallServerMovePacked: Failed to serialize out movement data!");
 		return;
 	}
 
@@ -9296,11 +9296,11 @@ void UCharacterMovementComponent::ServerMoveOld_Implementation
 
 	if( !VerifyClientTimeStamp(OldTimeStamp, *ServerData) )
 	{
-		UE_LOG(LogNetPlayerMovement, VeryVerbose, TEXT("ServerMoveOld: TimeStamp expired. %f, CurrentTimeStamp: %f, Character: %s"), OldTimeStamp, ServerData->CurrentClientTimeStamp, *GetNameSafe(CharacterOwner));
+		UE_LOGF(LogNetPlayerMovement, VeryVerbose, "ServerMoveOld: TimeStamp expired. %f, CurrentTimeStamp: %f, Character: %ls", OldTimeStamp, ServerData->CurrentClientTimeStamp, *GetNameSafe(CharacterOwner));
 		return;
 	}
 
-	UE_LOG(LogNetPlayerMovement, Verbose, TEXT("Recovered move from OldTimeStamp %f, DeltaTime: %f"), OldTimeStamp, OldTimeStamp - ServerData->CurrentClientTimeStamp);
+	UE_LOGF(LogNetPlayerMovement, Verbose, "Recovered move from OldTimeStamp %f, DeltaTime: %f", OldTimeStamp, OldTimeStamp - ServerData->CurrentClientTimeStamp);
 
 	const UWorld* MyWorld = GetWorld();
 	const float DeltaTime = ServerData->GetServerMoveDeltaTime(OldTimeStamp, CharacterOwner->GetActorTimeDilation(*MyWorld));
@@ -9315,7 +9315,7 @@ void UCharacterMovementComponent::ServerMoveOld_Implementation
 	}
 	else
 	{
-		UE_LOG(LogNetPlayerMovement, Warning, TEXT("OldTimeStamp(%f) results in zero or negative actual DeltaTime(%f). Theoretical DeltaTime(%f)"),
+		UE_LOGF(LogNetPlayerMovement, Warning, "OldTimeStamp(%f) results in zero or negative actual DeltaTime(%f). Theoretical DeltaTime(%f)",
 			OldTimeStamp, DeltaTime, OldTimeStamp - ServerData->CurrentClientTimeStamp);
 	}
 }
@@ -9381,7 +9381,7 @@ bool UCharacterMovementComponent::VerifyClientTimeStamp(float TimeStamp, FNetwor
 	{
 		if (bTimeStampResetDetected)
 		{
-			UE_LOG(LogNetPlayerMovement, Log, TEXT("TimeStamp reset detected. CurrentTimeStamp: %f, new TimeStamp: %f"), ServerData.CurrentClientTimeStamp, TimeStamp);
+			UE_LOGF(LogNetPlayerMovement, Log, "TimeStamp reset detected. CurrentTimeStamp: %f, new TimeStamp: %f", ServerData.CurrentClientTimeStamp, TimeStamp);
 			if (World != nullptr)
 			{
 				LastTimeStampResetServerTime = World->GetTimeSeconds();
@@ -9399,7 +9399,7 @@ bool UCharacterMovementComponent::VerifyClientTimeStamp(float TimeStamp, FNetwor
 			// Let's sync up to this TimeStamp in order to resolve movement desyncs ASAP
 			// This will result in this move having a zero DeltaTime (TimeStamp - CurrentClientTimeStamp), so it will perform no movement but it will send a correction,
 			// and we should be able to process the next move that arrives, which will advance ServerTimeStamp and unlock the forced updates (see 'APlayerController::TickActor()').
-			UE_LOG(LogNetPlayerMovement, Log, TEXT("Received a new move after performing ForcedUpdates. Updating CurrentClientTimeStamp from %f to %f Pawn: %s"), ServerData.CurrentClientTimeStamp, TimeStamp, *GetNameSafe(CharacterOwner));
+			UE_LOGF(LogNetPlayerMovement, Log, "Received a new move after performing ForcedUpdates. Updating CurrentClientTimeStamp from %f to %f Pawn: %ls", ServerData.CurrentClientTimeStamp, TimeStamp, *GetNameSafe(CharacterOwner));
 			ServerData.CurrentClientTimeStamp = TimeStamp;
 			if (World != nullptr)
 			{
@@ -9408,7 +9408,7 @@ bool UCharacterMovementComponent::VerifyClientTimeStamp(float TimeStamp, FNetwor
 		}
 		else if (!bTimeStampResetDetected)
 		{
-			UE_LOG(LogNetPlayerMovement, VeryVerbose, TEXT("TimeStamp %f Accepted! CurrentTimeStamp: %f Pawn: %s"), TimeStamp, ServerData.CurrentClientTimeStamp, *GetNameSafe(CharacterOwner));
+			UE_LOGF(LogNetPlayerMovement, VeryVerbose, "TimeStamp %f Accepted! CurrentTimeStamp: %f Pawn: %ls", TimeStamp, ServerData.CurrentClientTimeStamp, *GetNameSafe(CharacterOwner));
 			ProcessClientTimeStampForTimeDiscrepancy(TimeStamp, ServerData);
 		}
 	}
@@ -9416,7 +9416,7 @@ bool UCharacterMovementComponent::VerifyClientTimeStamp(float TimeStamp, FNetwor
 	{
 		if (bTimeStampResetDetected)
 		{
-			UE_LOG(LogNetPlayerMovement, Log, TEXT("TimeStamp expired. Before TimeStamp Reset. CurrentTimeStamp: %f, TimeStamp: %f Pawn: %s"), ServerData.CurrentClientTimeStamp, TimeStamp, *GetNameSafe(CharacterOwner));
+			UE_LOGF(LogNetPlayerMovement, Log, "TimeStamp expired. Before TimeStamp Reset. CurrentTimeStamp: %f, TimeStamp: %f Pawn: %ls", ServerData.CurrentClientTimeStamp, TimeStamp, *GetNameSafe(CharacterOwner));
 		}
 		else
 		{
@@ -9497,7 +9497,7 @@ void UCharacterMovementComponent::ProcessClientTimeStampForTimeDiscrepancy(float
 		// Per-frame spew of time discrepancy-related values - useful for investigating state of time discrepancy tracking
 		if (CharacterMovementCVars::DebugTimeDiscrepancy > 0)
 		{
-			UE_LOG(LogNetPlayerMovement, Warning, TEXT("TimeDiscrepancyDetection: ClientError: %f, TimeDiscrepancy: %f, LifetimeRawTimeDiscrepancy: %f (Lifetime %f), Resolving: %d, ClientDelta: %f, ServerDelta: %f, ClientTimeStamp: %f"),
+			UE_LOGF(LogNetPlayerMovement, Warning, "TimeDiscrepancyDetection: ClientError: %f, TimeDiscrepancy: %f, LifetimeRawTimeDiscrepancy: %f (Lifetime %f), Resolving: %d, ClientDelta: %f, ServerDelta: %f, ClientTimeStamp: %f",
 				ClientError, ServerData.TimeDiscrepancy, ServerData.LifetimeRawTimeDiscrepancy, WorldTimeSeconds - ServerData.WorldCreationTime, ServerData.bResolvingTimeDiscrepancy, ClientDelta, ServerDelta, ClientTimeStamp);
 		}
 #endif // !(UE_BUILD_SHIPPING || UE_BUILD_TEST)
@@ -9611,7 +9611,7 @@ void UCharacterMovementComponent::ProcessClientTimeStampForTimeDiscrepancy(float
 			// Per-frame spew of time discrepancy resolution related values - useful for investigating state of time discrepancy tracking
 			if (CharacterMovementCVars::DebugTimeDiscrepancy > 1)
 			{
-				UE_LOG(LogNetPlayerMovement, Warning, TEXT("TimeDiscrepancyResolution: DeltaOverride: %f, TimeToPayBack: %f, BaseDelta: %f, ServerDeltaSinceLastMovementUpdate: %f, TimeDiscrepancyAccumulatedClientDeltasSinceLastServerTick: %f"),
+				UE_LOGF(LogNetPlayerMovement, Warning, "TimeDiscrepancyResolution: DeltaOverride: %f, TimeToPayBack: %f, BaseDelta: %f, ServerDeltaSinceLastMovementUpdate: %f, TimeDiscrepancyAccumulatedClientDeltasSinceLastServerTick: %f",
 					ServerData.TimeDiscrepancyResolutionMoveDeltaOverride, TimeToPayBack, BaseDeltaTime, ServerDeltaSinceLastMovementUpdate, ServerData.TimeDiscrepancyAccumulatedClientDeltasSinceLastServerTick);
 			}
 #endif // !(UE_BUILD_SHIPPING || UE_BUILD_TEST)
@@ -9675,7 +9675,7 @@ void UCharacterMovementComponent::OnClientTimeStampResetDetected()
 
 void UCharacterMovementComponent::OnTimeDiscrepancyDetected(float CurrentTimeDiscrepancy, float LifetimeRawTimeDiscrepancy, float Lifetime, float CurrentMoveError)
 {
-	UE_LOG(LogNetPlayerMovement, Verbose, TEXT("Movement Time Discrepancy detected between client-reported time and server on character %s. CurrentTimeDiscrepancy: %f, LifetimeRawTimeDiscrepancy: %f, Lifetime: %f, CurrentMoveError %f"), 
+	UE_LOGF(LogNetPlayerMovement, Verbose, "Movement Time Discrepancy detected between client-reported time and server on character %ls. CurrentTimeDiscrepancy: %f, LifetimeRawTimeDiscrepancy: %f, Lifetime: %f, CurrentMoveError %f", 
 		CharacterOwner ? *CharacterOwner->GetHumanReadableName() : TEXT("<UNKNOWN>"), 
 		CurrentTimeDiscrepancy, 
 		LifetimeRawTimeDiscrepancy, 
@@ -9696,7 +9696,7 @@ bool FCharacterNetworkSerializationPackedBits::NetSerialize(FArchive& Ar, class
 	if (NumBits > static_cast<uint32>(CharacterMovementCVars::NetPackedMovementMaxBits))
 	{
 		// Protect against bad data that could cause server to allocate way too much memory.
-		UE_LOG(LogNetPlayerMovement, Error, TEXT("FCharacterNetworkSerializationPackedBits::NetSerialize: Dropping move due to NumBits (%d) exceeding allowable limit (%d). See NetPackedMovementMaxBits."), NumBits, CharacterMovementCVars::NetPackedMovementMaxBits);
+		UE_LOGF(LogNetPlayerMovement, Error, "FCharacterNetworkSerializationPackedBits::NetSerialize: Dropping move due to NumBits (%d) exceeding allowable limit (%d). See NetPackedMovementMaxBits.", NumBits, CharacterMovementCVars::NetPackedMovementMaxBits);
 		return false;
 	}
 
@@ -9886,7 +9886,7 @@ void UCharacterMovementComponent::ServerMovePacked_ServerReceive(const FCharacte
 	if (NumBits > CharacterMovementCVars::NetPackedMovementMaxBits)
 	{
 		// Protect against bad data that could cause server to allocate way too much memory.
-		UE_LOG(LogNetPlayerMovement, Error, TEXT("ServerMovePacked_ServerReceive (%s): Dropping move due to NumBits (%d) exceeding allowable limit (%d). See NetPackedMovementMaxBits."), *GetNameSafe(GetOwner()), NumBits, CharacterMovementCVars::NetPackedMovementMaxBits);
+		UE_LOGF(LogNetPlayerMovement, Error, "ServerMovePacked_ServerReceive (%ls): Dropping move due to NumBits (%d) exceeding allowable limit (%d). See NetPackedMovementMaxBits.", *GetNameSafe(GetOwner()), NumBits, CharacterMovementCVars::NetPackedMovementMaxBits);
 		return;
 	}
 
@@ -9904,7 +9904,7 @@ void UCharacterMovementComponent::ServerMovePacked_ServerReceive(const FCharacte
 
 	if (ServerMoveBitReader.PackageMap == nullptr)
 	{
-		devCode(UE_LOG(LogNetPlayerMovement, Error, TEXT("ServerMovePacked_ServerReceive: Failed to find PackageMap for data serialization!")));
+		devCode(UE_LOGF(LogNetPlayerMovement, Error, "ServerMovePacked_ServerReceive: Failed to find PackageMap for data serialization!"));
 		return;
 	}
 
@@ -9913,7 +9913,7 @@ void UCharacterMovementComponent::ServerMovePacked_ServerReceive(const FCharacte
 	FCharacterNetworkMoveDataContainer& MoveDataContainer = GetNetworkMoveDataContainer();
 	if (!MoveDataContainer.Serialize(*this, ServerMoveBitReader, ServerMoveBitReader.PackageMap) || ServerMoveBitReader.IsError())
 	{
-		devCode(UE_LOG(LogNetPlayerMovement, Error, TEXT("ServerMovePacked_ServerReceive: Failed to serialize movement data!")));
+		devCode(UE_LOGF(LogNetPlayerMovement, Error, "ServerMovePacked_ServerReceive: Failed to serialize movement data!"));
 		return;
 	}
 
@@ -9993,11 +9993,11 @@ void UCharacterMovementComponent::ServerMove_PerformMovement(const FCharacterNet
 		// This is more severe if the timestamp has a large discrepancy and hasn't been recently reset.
 		if (ServerTimeStamp > 1.0f && FMath::Abs(ServerTimeStamp - ClientTimeStamp) > CharacterMovementCVars::NetServerMoveTimestampExpiredWarningThreshold)
 		{
-			UE_LOG(LogNetPlayerMovement, Warning, TEXT("ServerMove: TimeStamp expired: %f, CurrentTimeStamp: %f, Character: %s"), ClientTimeStamp, ServerTimeStamp, *GetNameSafe(CharacterOwner));
+			UE_LOGF(LogNetPlayerMovement, Warning, "ServerMove: TimeStamp expired: %f, CurrentTimeStamp: %f, Character: %ls", ClientTimeStamp, ServerTimeStamp, *GetNameSafe(CharacterOwner));
 		}
 		else
 		{
-			UE_LOG(LogNetPlayerMovement, Log, TEXT("ServerMove: TimeStamp expired: %f, CurrentTimeStamp: %f, Character: %s"), ClientTimeStamp, ServerTimeStamp, *GetNameSafe(CharacterOwner));
+			UE_LOGF(LogNetPlayerMovement, Log, "ServerMove: TimeStamp expired: %f, CurrentTimeStamp: %f, Character: %ls", ClientTimeStamp, ServerTimeStamp, *GetNameSafe(CharacterOwner));
 		}		
 		return;
 	}
@@ -10048,7 +10048,7 @@ void UCharacterMovementComponent::ServerMove_PerformMovement(const FCharacterNet
 			MoveAutonomous(ClientTimeStamp, DeltaTime, ClientMoveFlags, ClientAccel);
 		}
 
-		UE_CLOG(CharacterOwner && UpdatedComponent, LogNetPlayerMovement, VeryVerbose, TEXT("ServerMove Time %f Acceleration %s Velocity %s Position %s Rotation %s GravityDirection %s DeltaTime %f Mode %s MovementBase %s.%s (Dynamic:%d)"),
+		UE_CLOGF(CharacterOwner && UpdatedComponent, LogNetPlayerMovement, VeryVerbose, "ServerMove Time %f Acceleration %ls Velocity %ls Position %ls Rotation %ls GravityDirection %ls DeltaTime %f Mode %ls MovementBase %ls.%ls (Dynamic:%d)",
 			ClientTimeStamp, *ClientAccel.ToString(), *Velocity.ToString(), *UpdatedComponent->GetComponentLocation().ToString(), *UpdatedComponent->GetComponentRotation().ToCompactString(), *GravityDirection.ToCompactString(), DeltaTime, *GetMovementName(),
 			*GetNameSafe(MovementBaseInterfaceData.GetMovementBaseObject()), *CharacterOwner->GetBasedMovement().BoneName.ToString(), MovementBaseUtility::IsDynamicBase(&MovementBaseInterfaceData) ? 1 : 0);
 	}
@@ -10068,12 +10068,12 @@ void UCharacterMovementComponent::ServerMove(float TimeStamp, FVector_NetQuantiz
 	const FMovementBaseInterfaceData MovementBaseInterfaceData(ClientMovementBase);
 	if (MovementBaseUtility::IsDynamicBase(&MovementBaseInterfaceData))
 	{
-		//UE_LOG(LogCharacterMovement, Log, TEXT("ServerMove: base %s"), *ClientMovementBase->GetName());
+		//UE_LOGF(LogCharacterMovement, Log, "ServerMove: base %ls", *ClientMovementBase->GetName());
 		CharacterOwner->ServerMove(TimeStamp, InAccel, ClientLoc, CompressedMoveFlags, ClientRoll, View, ClientMovementBase, ClientBaseBoneName, ClientMovementMode);
 	}
 	else
 	{
-		//UE_LOG(LogCharacterMovement, Log, TEXT("ServerMoveNoBase"));
+		//UE_LOGF(LogCharacterMovement, Log, "ServerMoveNoBase");
 		CharacterOwner->ServerMoveNoBase(TimeStamp, InAccel, ClientLoc, CompressedMoveFlags, ClientRoll, View, ClientMovementMode);
 	}
 }
@@ -10107,11 +10107,11 @@ void UCharacterMovementComponent::ServerMove_Implementation(
 		// This is more severe if the timestamp has a large discrepancy and hasn't been recently reset.
 		if (ServerTimeStamp > 1.0f && FMath::Abs(ServerTimeStamp - TimeStamp) > CharacterMovementCVars::NetServerMoveTimestampExpiredWarningThreshold)
 		{
-			UE_LOG(LogNetPlayerMovement, Warning, TEXT("ServerMove: TimeStamp expired: %f, CurrentTimeStamp: %f, Character: %s"), TimeStamp, ServerTimeStamp, *GetNameSafe(CharacterOwner));
+			UE_LOGF(LogNetPlayerMovement, Warning, "ServerMove: TimeStamp expired: %f, CurrentTimeStamp: %f, Character: %ls", TimeStamp, ServerTimeStamp, *GetNameSafe(CharacterOwner));
 		}
 		else
 		{
-			UE_LOG(LogNetPlayerMovement, Log, TEXT("ServerMove: TimeStamp expired: %f, CurrentTimeStamp: %f, Character: %s"), TimeStamp, ServerTimeStamp, *GetNameSafe(CharacterOwner));
+			UE_LOGF(LogNetPlayerMovement, Log, "ServerMove: TimeStamp expired: %f, CurrentTimeStamp: %f, Character: %ls", TimeStamp, ServerTimeStamp, *GetNameSafe(CharacterOwner));
 		}		
 		return;
 	}
@@ -10167,7 +10167,7 @@ void UCharacterMovementComponent::ServerMove_Implementation(
 	}
 
 	FMovementBaseInterfaceData* MovementBaseInterfaceData = GetMovementBaseInterfaceData_Mutable();
-	UE_CLOG(CharacterOwner && UpdatedComponent, LogNetPlayerMovement, VeryVerbose, TEXT("ServerMove Time %f Acceleration %s Velocity %s Position %s Rotation %s DeltaTime %f Mode %s MovementBase %s.%s (Dynamic:%d)"),
+	UE_CLOGF(CharacterOwner && UpdatedComponent, LogNetPlayerMovement, VeryVerbose, "ServerMove Time %f Acceleration %ls Velocity %ls Position %ls Rotation %ls DeltaTime %f Mode %ls MovementBase %ls.%ls (Dynamic:%d)",
 			TimeStamp, *Accel.ToString(), *Velocity.ToString(), *UpdatedComponent->GetComponentLocation().ToString(), *UpdatedComponent->GetComponentRotation().ToCompactString(), DeltaTime, *GetMovementName(),
 			*GetNameSafe(MovementBaseInterfaceData ? MovementBaseInterfaceData->GetMovementBaseObject() : nullptr), *CharacterOwner->GetBasedMovement().BoneName.ToString(), MovementBaseUtility::IsDynamicBase(MovementBaseInterfaceData) ? 1 : 0);
 
@@ -10444,7 +10444,7 @@ void UCharacterMovementComponent::ServerMoveHandleClientError(float ClientTimeSt
 			const FVector LocDiff = UpdatedComponent->GetComponentLocation() - ClientLoc;
 			const UObject* MovementBaseObject = MovementBaseInterfaceData ? MovementBaseInterfaceData->GetMovementBaseObject() : nullptr;
 			const FString BaseString = MovementBaseObject ? MovementBaseObject->GetPathName(MovementBaseObject->GetOutermost()) : TEXT("None");
-			UE_LOG(LogNetPlayerMovement, Warning, TEXT("*** Server: Error for %s at Time=%.3f is %3.3f LocDiff(%s) ClientLoc(%s) ServerLoc(%s) Base: %s Bone: %s Accel(%s) Velocity(%s)"),
+			UE_LOGF(LogNetPlayerMovement, Warning, "*** Server: Error for %ls at Time=%.3f is %3.3f LocDiff(%ls) ClientLoc(%ls) ServerLoc(%ls) Base: %ls Bone: %ls Accel(%ls) Velocity(%ls)",
 				*GetNameSafe(CharacterOwner), ClientTimeStamp, LocDiff.Size(), *LocDiff.ToString(), *ClientLoc.ToString(), *UpdatedComponent->GetComponentLocation().ToString(), *BaseString, *ServerData->PendingAdjustment.NewBaseBoneName.ToString(), *Accel.ToString(), *Velocity.ToString());
 			const float DebugLifetime = CharacterMovementCVars::NetCorrectionLifetime;
 			DrawDebugCapsule(GetWorld(), UpdatedComponent->GetComponentLocation(), CharacterOwner->GetSimpleCollisionHalfHeight(), CharacterOwner->GetSimpleCollisionRadius(), FQuat::Identity, FColor(100, 255, 100), false, DebugLifetime);
@@ -10540,7 +10540,7 @@ bool UCharacterMovementComponent::ServerCheckClientError(float ClientTimeStamp,
 		{
 			if (RandomStream.FRand() < CharacterMovementCVars::NetForceClientAdjustmentPercent)
 			{
-				UE_LOG(LogNetPlayerMovement, VeryVerbose, TEXT("** ServerCheckClientError forced by p.NetForceClientAdjustmentPercent"));
+				UE_LOGF(LogNetPlayerMovement, VeryVerbose, "** ServerCheckClientError forced by p.NetForceClientAdjustmentPercent");
 				return true;
 			}
 		}
@@ -10551,7 +10551,7 @@ bool UCharacterMovementComponent::ServerCheckClientError(float ClientTimeStamp,
 #if !UE_BUILD_SHIPPING
 		if (CharacterMovementCVars::NetShowCorrections != 0)
 		{
-			UE_LOG(LogNetPlayerMovement, Warning, TEXT("*** Server: %s is set to ignore error checks and corrections."), *GetNameSafe(CharacterOwner));
+			UE_LOGF(LogNetPlayerMovement, Warning, "*** Server: %ls is set to ignore error checks and corrections.", *GetNameSafe(CharacterOwner));
 		}
 #endif // !UE_BUILD_SHIPPING
 	}
@@ -10621,12 +10621,12 @@ void UCharacterMovementComponent::ServerMoveDual(float TimeStamp0, FVector_NetQu
 	const FMovementBaseInterfaceData MovementBaseInterfaceData(ClientMovementBase);
 	if (MovementBaseUtility::IsDynamicBase(&MovementBaseInterfaceData))
 	{
-		//UE_LOG(LogCharacterMovement, Log, TEXT("ServerMoveDual: base %s"), *ClientMovementBase->GetName());
+		//UE_LOGF(LogCharacterMovement, Log, "ServerMoveDual: base %ls", *ClientMovementBase->GetName());
 		CharacterOwner->ServerMoveDual(TimeStamp0, InAccel0, PendingFlags, View0, TimeStamp, InAccel, ClientLoc, NewFlags, ClientRoll, View, ClientMovementBase, ClientBaseBoneName, ClientMovementMode);
 	}
 	else
 	{
-		//UE_LOG(LogCharacterMovement, Log, TEXT("ServerMoveDualNoBase"));
+		//UE_LOGF(LogCharacterMovement, Log, "ServerMoveDualNoBase");
 		CharacterOwner->ServerMoveDualNoBase(TimeStamp0, InAccel0, PendingFlags, View0, TimeStamp, InAccel, ClientLoc, NewFlags, ClientRoll, View, ClientMovementMode);
 	}
 }
@@ -10778,7 +10778,7 @@ void UCharacterMovementComponent::MoveResponsePacked_ClientReceive(const FCharac
 	if (!ensureMsgf(NumBits <= CharacterMovementCVars::NetPackedMovementMaxBits, TEXT("MoveResponsePacked_ClientReceive: NumBits (%d) exceeds CharacterMovementCVars::NetPackedMovementMaxBits (%d)"), NumBits, CharacterMovementCVars::NetPackedMovementMaxBits))
 	{
 		// Protect against bad data that could cause client to allocate way too much memory.
-		devCode(UE_LOG(LogNetPlayerMovement, Error, TEXT("MoveResponsePacked_ClientReceive: NumBits (%d) exceeds allowable limit!"), NumBits));
+		devCode(UE_LOGF(LogNetPlayerMovement, Error, "MoveResponsePacked_ClientReceive: NumBits (%d) exceeds allowable limit!", NumBits));
 		return;
 	}
 
@@ -10796,7 +10796,7 @@ void UCharacterMovementComponent::MoveResponsePacked_ClientReceive(const FCharac
 
 	if (MoveResponseBitReader.PackageMap == nullptr)
 	{
-		devCode(UE_LOG(LogNetPlayerMovement, Error, TEXT("MoveResponsePacked_ClientReceive: Failed to find PackageMap for data serialization!")));
+		devCode(UE_LOGF(LogNetPlayerMovement, Error, "MoveResponsePacked_ClientReceive: Failed to find PackageMap for data serialization!"));
 		return;
 	}
 
@@ -10805,7 +10805,7 @@ void UCharacterMovementComponent::MoveResponsePacked_ClientReceive(const FCharac
 	FCharacterMoveResponseDataContainer& ResponseDataContainer = GetMoveResponseDataContainer();
 	if (!ResponseDataContainer.Serialize(*this, MoveResponseBitReader, MoveResponseBitReader.PackageMap) || MoveResponseBitReader.IsError())
 	{
-		devCode(UE_LOG(LogNetPlayerMovement, Error, TEXT("MoveResponsePacked_ClientReceive: Failed to serialize response data!")));
+		devCode(UE_LOGF(LogNetPlayerMovement, Error, "MoveResponsePacked_ClientReceive: Failed to serialize response data!"));
 		return;
 	}
 
@@ -10839,7 +10839,7 @@ void UCharacterMovementComponent::ServerSendMoveResponse(const FClientAdjustment
 
 	if (MoveResponseBitWriter.PackageMap == nullptr)
 	{
-		UE_LOG(LogNetPlayerMovement, Error, TEXT("ServerSendMoveResponse: Failed to find a NetConnection/PackageMap for data serialization!"));
+		UE_LOGF(LogNetPlayerMovement, Error, "ServerSendMoveResponse: Failed to find a NetConnection/PackageMap for data serialization!");
 		return;
 	}
 
@@ -10850,7 +10850,7 @@ void UCharacterMovementComponent::ServerSendMoveResponse(const FClientAdjustment
 	// Serialize move struct into a bit stream
 	if (!ResponseDataContainer.Serialize(*this, MoveResponseBitWriter, MoveResponseBitWriter.PackageMap) || MoveResponseBitWriter.IsError())
 	{
-		UE_LOG(LogNetPlayerMovement, Error, TEXT("ServerSendMoveResponse: Failed to serialize out response data!"));
+		UE_LOGF(LogNetPlayerMovement, Error, "ServerSendMoveResponse: Failed to serialize out response data!");
 		return;
 	}
 
@@ -11249,13 +11249,13 @@ void UCharacterMovementComponent::ClientAdjustPosition_Implementation
 	{
 		if (bBaseRelativePosition)
 		{
-			UE_LOG(LogNetPlayerMovement, Warning, TEXT("ClientAdjustPosition_Implementation could not resolve the new relative movement base actor, ignoring server correction! Client currently at world location %s on base %s"),
+			UE_LOGF(LogNetPlayerMovement, Warning, "ClientAdjustPosition_Implementation could not resolve the new relative movement base actor, ignoring server correction! Client currently at world location %ls on base %ls",
 				*UpdatedComponent->GetComponentLocation().ToString(), *GetNameSafe(NewMovementBaseInterfaceData ? NewMovementBaseInterfaceData->GetMovementBaseObject() : nullptr));
 			return;
 		}
 		else
 		{
-			UE_LOG(LogNetPlayerMovement, Verbose, TEXT("ClientAdjustPosition_Implementation could not resolve the new absolute movement base actor, but WILL use the position!"));
+			UE_LOGF(LogNetPlayerMovement, Verbose, "ClientAdjustPosition_Implementation could not resolve the new absolute movement base actor, but WILL use the position!");
 		}
 	}
 	
@@ -11265,7 +11265,7 @@ void UCharacterMovementComponent::ClientAdjustPosition_Implementation
 	{
 		if( ClientData->LastAckedMove.IsValid() )
 		{
-			UE_LOG(LogNetPlayerMovement, Log,  TEXT("ClientAdjustPosition_Implementation could not find Move for TimeStamp: %f, LastAckedTimeStamp: %f, CurrentTimeStamp: %f"), TimeStamp, ClientData->LastAckedMove->TimeStamp, ClientData->CurrentTimeStamp);
+			UE_LOGF(LogNetPlayerMovement, Log,  "ClientAdjustPosition_Implementation could not find Move for TimeStamp: %f, LastAckedTimeStamp: %f, CurrentTimeStamp: %f", TimeStamp, ClientData->LastAckedMove->TimeStamp, ClientData->CurrentTimeStamp);
 		}
 		return;
 	}
@@ -11398,7 +11398,7 @@ void UCharacterMovementComponent::OnClientCorrectionReceived(FNetworkPredictionD
 		const FVector LocDiff = ClientLocAtCorrectedMove - NewLocation;
 		const UObject* MovementBaseObject = NewMovementBaseInterfaceData ? NewMovementBaseInterfaceData->GetMovementBaseObject() : nullptr;
 		const FString NewBaseString = MovementBaseUtility::IsMovementBaseDataValid(NewMovementBaseInterfaceData) ? MovementBaseObject->GetPathName(MovementBaseObject->GetOutermost()) : TEXT("None");
-		UE_LOG(LogNetPlayerMovement, Warning, TEXT("*** Client: Error for %s at Time=%.3f is %3.3f LocDiff(%s) ClientLoc(%s) ServerLoc(%s) NewBase: %s NewBone: %s ClientVel(%s) ServerVel(%s) SavedMoves %d"),
+		UE_LOGF(LogNetPlayerMovement, Warning, "*** Client: Error for %ls at Time=%.3f is %3.3f LocDiff(%ls) ClientLoc(%ls) ServerLoc(%ls) NewBase: %ls NewBone: %ls ClientVel(%ls) ServerVel(%ls) SavedMoves %d",
 			   *GetNameSafe(CharacterOwner), TimeStamp, LocDiff.Size(), *LocDiff.ToString(), *ClientLocAtCorrectedMove.ToString(), *NewLocation.ToString(), *NewBaseString, *NewBaseBoneName.ToString(), *Velocity.ToString(), *NewVelocity.ToString(), ClientData.SavedMoves.Num());
 		const float DebugLifetime = CharacterMovementCVars::NetCorrectionLifetime;
 		if (!LocDiff.IsNearlyZero())
@@ -11467,7 +11467,7 @@ void UCharacterMovementComponent::ClientAdjustRootMotionPosition_Implementation(
 	FRotator DecompressedRot(ServerRotation.X * 180.f, ServerRotation.Y * 180.f, ServerRotation.Z * 180.f);
 	CharacterOwner->SetActorRotation(DecompressedRot);
 	const FVector ServerLocation(FRepMovement::RebaseOntoLocalOrigin(ServerLoc, UpdatedComponent));
-	UE_LOG(LogRootMotion, Log,  TEXT("ClientAdjustRootMotionPosition_Implementation TimeStamp: %f, ServerMontageTrackPosition: %f, ServerLocation: %s, ServerRotation: %s, ServerVelZ: %f, ServerBase: %s"),
+	UE_LOGF(LogRootMotion, Log,  "ClientAdjustRootMotionPosition_Implementation TimeStamp: %f, ServerMontageTrackPosition: %f, ServerLocation: %ls, ServerRotation: %ls, ServerVelZ: %f, ServerBase: %ls",
 		TimeStamp, ServerMontageTrackPosition, *ServerLocation.ToCompactString(), *DecompressedRot.ToCompactString(), ServerVelZ, *GetNameSafe(ServerMovementBaseInterfaceData ? ServerMovementBaseInterfaceData->GetMovementBaseObject() : nullptr));
 
 #if 0
@@ -11479,7 +11479,7 @@ void UCharacterMovementComponent::ClientAdjustRootMotionPosition_Implementation(
 		const float DeltaTrackPosition = (ServerMontageTrackPosition - ClientData->LastAckedMove->RootMotionTrackPosition);
 		const float DeltaVelZ = (ServerVelZ - ClientData->LastAckedMove->SavedVelocity.Z);
 
-		UE_LOG(LogRootMotion, Log,  TEXT("\tErrors DeltaLocation: %s, DeltaRotation: %s, DeltaTrackPosition: %f"),
+		UE_LOGF(LogRootMotion, Log,  "\tErrors DeltaLocation: %ls, DeltaRotation: %ls, DeltaTrackPosition: %f",
 			*DeltaLocation.ToCompactString(), *DeltaRotation.ToCompactString(), DeltaTrackPosition );
 	}
 #endif
@@ -11491,7 +11491,7 @@ void UCharacterMovementComponent::ClientAdjustRootMotionPosition_Implementation(
 		FAnimMontageInstance * RootMotionMontageInstance = CharacterOwner->GetRootMotionAnimMontageInstance();
 		if (RootMotionMontageInstance && !RootMotionMontageInstance->IsRootMotionDisabled())
 		{
-			UE_LOG(LogRootMotion, Log, TEXT("\tServer disagrees with Client's track position!! ServerTrackPosition: %f, ClientTrackPosition: %f, DeltaTrackPosition: %f. TimeStamp: %f, Character: %s, Montage: %s"),
+			UE_LOGF(LogRootMotion, Log, "\tServer disagrees with Client's track position!! ServerTrackPosition: %f, ClientTrackPosition: %f, DeltaTrackPosition: %f. TimeStamp: %f, Character: %ls, Montage: %ls",
 					ServerMontageTrackPosition, ClientData->LastAckedMove->RootMotionTrackPosition, (ServerMontageTrackPosition - ClientData->LastAckedMove->RootMotionTrackPosition), TimeStamp, *GetNameSafe(CharacterOwner), *GetNameSafe(RootMotionMontageInstance->Montage));
 	
 			RootMotionMontageInstance->SetPosition(ServerMontageTrackPosition);
@@ -11556,7 +11556,7 @@ void UCharacterMovementComponent::ClientAdjustRootMotionSourcePosition_Implement
 	FRotator DecompressedRot(ServerRotation.X * 180.f, ServerRotation.Y * 180.f, ServerRotation.Z * 180.f);
 	CharacterOwner->SetActorRotation(DecompressedRot);
 	const FVector ServerLocation(FRepMovement::RebaseOntoLocalOrigin(ServerLoc, UpdatedComponent));
-	UE_LOG(LogRootMotion, Log,  TEXT("ClientAdjustRootMotionSourcePosition_Implementation TimeStamp: %f, NumRootMotionSources: %d, ServerLocation: %s, ServerRotation: %s, ServerVelZ: %f, ServerBase: %s"),
+	UE_LOGF(LogRootMotion, Log,  "ClientAdjustRootMotionSourcePosition_Implementation TimeStamp: %f, NumRootMotionSources: %d, ServerLocation: %ls, ServerRotation: %ls, ServerVelZ: %f, ServerBase: %ls",
 		TimeStamp, ServerRootMotion.RootMotionSources.Num(), *ServerLocation.ToCompactString(), *DecompressedRot.ToCompactString(), ServerVelZ, *GetNameSafe(ServerMovementBaseInterfaceData ? ServerMovementBaseInterfaceData->GetMovementBaseObject() : nullptr));
 
 	// Handle AnimRootMotion correction
@@ -11570,7 +11570,7 @@ void UCharacterMovementComponent::ClientAdjustRootMotionSourcePosition_Implement
 			const float DeltaTrackPosition = (ServerMontageTrackPosition - ClientData->LastAckedMove->RootMotionTrackPosition);
 			const float DeltaVelZ = (ServerVelZ - ClientData->LastAckedMove->SavedVelocity.Z);
 
-			UE_LOG(LogRootMotion, Log,  TEXT("\tErrors DeltaLocation: %s, DeltaRotation: %s, DeltaTrackPosition: %f"),
+			UE_LOGF(LogRootMotion, Log,  "\tErrors DeltaLocation: %ls, DeltaRotation: %ls, DeltaTrackPosition: %f",
 				*DeltaLocation.ToCompactString(), *DeltaRotation.ToCompactString(), DeltaTrackPosition );
 		}
 		#endif
@@ -11578,7 +11578,7 @@ void UCharacterMovementComponent::ClientAdjustRootMotionSourcePosition_Implement
 		// Server disagrees with Client on the Root Motion AnimMontage Track position.
 		if( CharacterOwner->bClientResimulateRootMotion || (ServerMontageTrackPosition != ClientData->LastAckedMove->RootMotionTrackPosition) )
 		{
-			UE_LOG(LogRootMotion, Log,  TEXT("\tServer disagrees with Client's track position!! ServerTrackPosition: %f, ClientTrackPosition: %f, DeltaTrackPosition: %f. TimeStamp: %f"),
+			UE_LOGF(LogRootMotion, Log,  "\tServer disagrees with Client's track position!! ServerTrackPosition: %f, ClientTrackPosition: %f, DeltaTrackPosition: %f. TimeStamp: %f",
 				ServerMontageTrackPosition, ClientData->LastAckedMove->RootMotionTrackPosition, (ServerMontageTrackPosition - ClientData->LastAckedMove->RootMotionTrackPosition), TimeStamp);
 
 			// Not much we can do there unfortunately, just jump to server's track position.
@@ -11602,7 +11602,7 @@ void UCharacterMovementComponent::ClientAdjustRootMotionSourcePosition_Implement
 	{
 		if (!CharacterOwner->bClientResimulateRootMotionSources)
 		{
-			UE_LOG(LogNetPlayerMovement, VeryVerbose, TEXT("ClientAdjustRootMotionSourcePosition called, server/LastAckedMove mismatch"));
+			UE_LOGF(LogNetPlayerMovement, VeryVerbose, "ClientAdjustRootMotionSourcePosition called, server/LastAckedMove mismatch");
 		}
 
 		CharacterOwner->SavedRootMotion = ServerRootMotion;
@@ -11640,7 +11640,7 @@ void UCharacterMovementComponent::ClientAckGoodMove_Implementation(float TimeSta
 	{
 		if( ClientData->LastAckedMove.IsValid() )
 		{
-			UE_LOG(LogNetPlayerMovement, Log, TEXT("ClientAckGoodMove_Implementation could not find Move for TimeStamp: %f, LastAckedTimeStamp: %f, CurrentTimeStamp: %f"), TimeStamp, ClientData->LastAckedMove->TimeStamp, ClientData->CurrentTimeStamp);
+			UE_LOGF(LogNetPlayerMovement, Log, "ClientAckGoodMove_Implementation could not find Move for TimeStamp: %f, LastAckedTimeStamp: %f, CurrentTimeStamp: %f", TimeStamp, ClientData->LastAckedMove->TimeStamp, ClientData->CurrentTimeStamp);
 		}
 		return;
 	}
@@ -11816,7 +11816,7 @@ void UCharacterMovementComponent::ApplyRepulsionForce(float DeltaSeconds)
 
 				if (!OverlapBody)
 				{
-					UE_LOG(LogCharacterMovement, Warning, TEXT("%s could not find overlap body for body index %d"), *GetName(), OverlapBodyIndex);
+					UE_LOGF(LogCharacterMovement, Warning, "%ls could not find overlap body for body index %d", *GetName(), OverlapBodyIndex);
 					continue;
 				}
 
@@ -12047,7 +12047,7 @@ void UCharacterMovementComponent::TickCharacterPose(float DeltaTime)
 			// Debugging
 			{
 				FAnimMontageInstance* RootMotionMontageInstance = CharacterOwner->GetRootMotionAnimMontageInstance();
-				UE_LOG(LogRootMotion, Log, TEXT("UCharacterMovementComponent::TickCharacterPose Role: %s, RootMotionMontage: %s, MontagePos: %f, DeltaTime: %f, ExtractedRootMotion: %s, AccumulatedRootMotion: %s")
+				UE_LOGF(LogRootMotion, Log, "UCharacterMovementComponent::TickCharacterPose Role: %ls, RootMotionMontage: %ls, MontagePos: %f, DeltaTime: %f, ExtractedRootMotion: %ls, AccumulatedRootMotion: %ls"
 					, *UEnum::GetValueAsString(TEXT("Engine.ENetRole"), CharacterOwner->GetLocalRole())
 					, *GetNameSafe(RootMotionMontageInstance ? RootMotionMontageInstance->Montage : NULL)
 					, RootMotionMontageInstance ? RootMotionMontageInstance->GetPosition() : -1.f
@@ -12328,27 +12328,27 @@ void UCharacterMovementComponent::ConvertRootMotionServerIDsToLocalIDs(const FRo
 
 	if (bDumpDebugInfo)
 	{
-		UE_LOG(LogRootMotion, Warning, TEXT("Dumping current mappings:"));
+		UE_LOGF(LogRootMotion, Warning, "Dumping current mappings:");
 		for (FRootMotionServerToLocalIDMapping& Mapping : RootMotionIDMappings)
 		{
-			UE_LOG(LogRootMotion, Warning, TEXT("- LocalID(%d) ServerID(%d)"), Mapping.LocalID, Mapping.ServerID);
+			UE_LOGF(LogRootMotion, Warning, "- LocalID(%d) ServerID(%d)", Mapping.LocalID, Mapping.ServerID);
 		}
 
-		UE_LOG(LogRootMotion, Warning, TEXT("Dumping local RootMotionSources:"));
+		UE_LOGF(LogRootMotion, Warning, "Dumping local RootMotionSources:");
 		for (const TSharedPtr<FRootMotionSource>& LocalRootMotionSource : LocalRootMotionToMatchWith.RootMotionSources)
 		{
 			if (LocalRootMotionSource.IsValid())
 			{
-				UE_LOG(LogRootMotion, Warning, TEXT("- LocalRootMotionSource(%s)"), *LocalRootMotionSource->ToSimpleString());
+				UE_LOGF(LogRootMotion, Warning, "- LocalRootMotionSource(%ls)", *LocalRootMotionSource->ToSimpleString());
 			}
 		}
 
-		UE_LOG(LogRootMotion, Warning, TEXT("Dumping server RootMotionSources:"));
+		UE_LOGF(LogRootMotion, Warning, "Dumping server RootMotionSources:");
 		for (TSharedPtr<FRootMotionSource>& ServerRootMotionSource : InOutServerRootMotion.RootMotionSources)
 		{
 			if (ServerRootMotionSource.IsValid())
 			{
-				UE_LOG(LogRootMotion, Warning, TEXT("- ServerRootMotionSource(%s)"), *ServerRootMotionSource->ToSimpleString());
+				UE_LOGF(LogRootMotion, Warning, "- ServerRootMotionSource(%ls)", *ServerRootMotionSource->ToSimpleString());
 			}
 		}
 	}
@@ -12446,7 +12446,7 @@ FSavedMovePtr FNetworkPredictionData_Client_Character::CreateSavedMove()
 {
 	if (SavedMoves.Num() >= MaxSavedMoveCount)
 	{
-		UE_LOG(LogNetPlayerMovement, Warning, TEXT("CreateSavedMove: Hit limit of %d saved moves (timing out or very bad ping?)"), SavedMoves.Num());
+		UE_LOGF(LogNetPlayerMovement, Warning, "CreateSavedMove: Hit limit of %d saved moves (timing out or very bad ping?)", SavedMoves.Num());
 		// Free all saved moves
 		for (int32 i=0; i < SavedMoves.Num(); i++)
 		{
@@ -12533,7 +12533,7 @@ void FNetworkPredictionData_Client_Character::AckMove(int32 AckedMoveIndex, UCha
 	{
 		// Keep reference to LastAckedMove
 		const FSavedMovePtr& AckedMove = SavedMoves[AckedMoveIndex];
-		UE_LOG(LogNetPlayerMovement, VeryVerbose, TEXT("AckedMove Index: %2d (%2d moves). TimeStamp: %f, CurrentTimeStamp: %f"), AckedMoveIndex, SavedMoves.Num(), AckedMove->TimeStamp, CurrentTimeStamp);
+		UE_LOGF(LogNetPlayerMovement, VeryVerbose, "AckedMove Index: %2d (%2d moves). TimeStamp: %f, CurrentTimeStamp: %f", AckedMoveIndex, SavedMoves.Num(), AckedMove->TimeStamp, CurrentTimeStamp);
 		if( LastAckedMove.IsValid() )
 		{
 			FreeMove(LastAckedMove);
@@ -12580,7 +12580,7 @@ FNetworkPredictionData_Server_Character::FNetworkPredictionData_Server_Character
 		MaxMoveDeltaTime = GameNetworkManager->MaxMoveDeltaTime;
 		if (GameNetworkManager->MaxMoveDeltaTime > GameNetworkManager->MAXCLIENTUPDATEINTERVAL)
 		{
-			UE_LOG(LogNetPlayerMovement, Warning, TEXT("GameNetworkManager::MaxMoveDeltaTime (%f) is greater than GameNetworkManager::MAXCLIENTUPDATEINTERVAL (%f)! Server will interfere with move deltas that large!"), GameNetworkManager->MaxMoveDeltaTime, GameNetworkManager->MAXCLIENTUPDATEINTERVAL);
+			UE_LOGF(LogNetPlayerMovement, Warning, "GameNetworkManager::MaxMoveDeltaTime (%f) is greater than GameNetworkManager::MAXCLIENTUPDATEINTERVAL (%f)! Server will interfere with move deltas that large!", GameNetworkManager->MaxMoveDeltaTime, GameNetworkManager->MAXCLIENTUPDATEINTERVAL);
 		}
 	}
 
@@ -12843,11 +12843,11 @@ void FSavedMove_Character::PostUpdate(ACharacter* Character, FSavedMove_Characte
 		{
 			if (Character->SavedRootMotion.HasActiveRootMotionSources())
 			{
-				UE_LOG(LogCharacterMovement, Log, TEXT("FSavedMove_Character::PostUpdate detected very high Velocity! (%s), but with active root motion sources (could be intentional)"), *SavedVelocity.ToString());
+				UE_LOGF(LogCharacterMovement, Log, "FSavedMove_Character::PostUpdate detected very high Velocity! (%ls), but with active root motion sources (could be intentional)", *SavedVelocity.ToString());
 			}
 			else
 			{
-				UE_LOG(LogCharacterMovement, Warning, TEXT("FSavedMove_Character::PostUpdate detected very high Velocity! (%s)"), *SavedVelocity.ToString());
+				UE_LOGF(LogCharacterMovement, Warning, "FSavedMove_Character::PostUpdate detected very high Velocity! (%ls)", *SavedVelocity.ToString());
 			}
 		}
 #endif
@@ -13208,7 +13208,7 @@ bool FSavedMove_Character::CanCombineWith(const FSavedMovePtr& NewMovePtr, AChar
 		const FVector RelativeLocationDelta = (StartAttachRelativeLocation - NewMove->StartAttachRelativeLocation);
 		if (!RelativeLocationDelta.IsNearlyZero(CharacterMovementCVars::NetMoveCombiningAttachedLocationTolerance))
 		{
-			//UE_LOG(LogCharacterMovement, Warning, TEXT("NoCombine: DeltaLocation(%s)"), *RelativeLocationDelta.ToString());
+			//UE_LOGF(LogCharacterMovement, Warning, "NoCombine: DeltaLocation(%ls)", *RelativeLocationDelta.ToString());
 			return false;
 		}
 		// For rotation, Yaw doesn't matter for capsules
@@ -13310,7 +13310,7 @@ void FSavedMove_Character::PrepMoveFor(ACharacter* Character)
 			// we're not able to properly re-prepare them, and should just keep whatever we currently have
 
 			// Apply any corrections/state from either last played move or last received from server (in ACharacter::SavedRootMotion)
-			UE_LOG(LogRootMotion, VeryVerbose, TEXT("SavedMove SavedRootMotion getting updated for SavedMove replays: %s"), *Character->GetName());
+			UE_LOGF(LogRootMotion, VeryVerbose, "SavedMove SavedRootMotion getting updated for SavedMove replays: %ls", *Character->GetName());
 			SavedRootMotion.UpdateStateFrom(Character->SavedRootMotion);
 			SavedRootMotion.CleanUpInvalidRootMotion(DeltaTime, *Character, *Character->GetCharacterMovement());
 			SavedRootMotion.PrepareRootMotion(DeltaTime, *Character, *Character->GetCharacterMovement());
@@ -13423,7 +13423,7 @@ void UCharacterMovementComponent::FlushServerMoves()
 			FSavedMovePtr NewMove = ClientData->PendingMove;
 			ClientData->PendingMove = nullptr;
 
-			UE_CLOG(CharacterOwner && UpdatedComponent, LogNetPlayerMovement, Verbose, TEXT("ClientMove (Flush) Time %f Acceleration %s Velocity %s Position %s DeltaTime %f Mode %s MovementBase %s.%s (Dynamic:%d) DualMove? %d"),
+			UE_CLOGF(CharacterOwner && UpdatedComponent, LogNetPlayerMovement, Verbose, "ClientMove (Flush) Time %f Acceleration %ls Velocity %ls Position %ls DeltaTime %f Mode %ls MovementBase %ls.%ls (Dynamic:%d) DualMove? %d",
 				NewMove->TimeStamp, *NewMove->Acceleration.ToString(), *Velocity.ToString(), *UpdatedComponent->GetComponentLocation().ToString(), NewMove->DeltaTime, *GetMovementName(),
 				*GetNameSafe(NewMove->EndMovementBaseInterfaceData.GetMovementBaseObject()), *NewMove->EndBoneName.ToString(), MovementBaseUtility::IsDynamicBase(&NewMove->EndMovementBaseInterfaceData) ? 1 : 0, ClientData->PendingMove.IsValid() ? 1 : 0);
 
diff --git a/Engine/Source/Runtime/Engine/Private/Materials/MaterialExpressionsToMIR.cpp b/Engine/Source/Runtime/Engine/Private/Materials/MaterialExpressionsToMIR.cpp
--- a/Engine/Source/Runtime/Engine/Private/Materials/MaterialExpressionsToMIR.cpp
+++ b/Engine/Source/Runtime/Engine/Private/Materials/MaterialExpressionsToMIR.cpp
@@ -1988,13 +1988,13 @@ void UMaterialExpressionFontSample::Build(MIR::FEmitter& Em)
 	// if we can't throw the error below, attempt to thwart the error by using the default font
 	if (!Font)
 	{
-		UE_LOG(LogMaterial, Log, TEXT("Using default font instead of real font!"));
+		UE_LOGF(LogMaterial, Log, "Using default font instead of real font!");
 		Font = GEngine->GetMediumFont();
 		FontTexturePage = 0;
 	}
 	else if (!Font->Textures.IsValidIndex(FontTexturePage))
 	{
-		UE_LOG(LogMaterial, Log, TEXT("Invalid font page %d. Max allowed is %d"), FontTexturePage, Font->Textures.Num());
+		UE_LOGF(LogMaterial, Log, "Invalid font page %d. Max allowed is %d", FontTexturePage, Font->Textures.Num());
 		FontTexturePage = 0;
 	}
 #endif
@@ -4416,7 +4416,7 @@ void UMaterialExpressionSkyLightEnvMapSample::Build(MIR::FEmitter& Em)
 {
 	if (Material->bIsSky)
 	{
-		UE_LOG(LogMaterial, Warning, TEXT("Using SkyLightEnvMapSample from a IsSky material can result in visual artifact. For instance, if the previous frame capture was super bright, it might leak onto a new frame, e.g. transtion from menu to game."));
+		UE_LOGF(LogMaterial, Warning, "Using SkyLightEnvMapSample from a IsSky material can result in visual artifact. For instance, if the previous frame capture was super bright, it might leak onto a new frame, e.g. transtion from menu to game.");
 	}
 
 	FValueRef DirectionValue = Em.InputDefaultFloat3(&Direction, FVector3f(0.0f, 0.0f, 1.0f));
diff --git a/Engine/Source/Runtime/StorageServerClient/Private/StorageServerPlatformFile.cpp b/Engine/Source/Runtime/StorageServerClient/Private/StorageServerPlatformFile.cpp
--- a/Engine/Source/Runtime/StorageServerClient/Private/StorageServerPlatformFile.cpp
+++ b/Engine/Source/Runtime/StorageServerClient/Private/StorageServerPlatformFile.cpp
@@ -208,7 +208,7 @@ bool FStorageServerFileSystemTOC::GetFileData(const FString& Path, FIoChunkId& O
 
 bool FStorageServerFileSystemTOC::IterateDirectory(const FString& Path, TFunctionRef<bool(const FIoChunkId&, const TCHAR*, int64 RawSize)> Callback)
 {
-	UE_LOG(LogStorageServerPlatformFile, Verbose, TEXT("IterateDirectory '%s'"), *Path);
+	UE_LOGF(LogStorageServerPlatformFile, Verbose, "IterateDirectory '%ls'", *Path);
 
 	FReadScopeLock _(TocLock);
 
@@ -237,7 +237,7 @@ bool FStorageServerFileSystemTOC::IterateDirectory(const FString& Path, TFunctio
 
 bool FStorageServerFileSystemTOC::IterateDirectoryRecursively(const FString& Path, TFunctionRef<bool(const FIoChunkId&, const TCHAR*, int64)> Callback)
 {
-	UE_LOG(LogStorageServerPlatformFile, Verbose, TEXT("IterateDirectoryRecursively '%s'"), *Path);
+	UE_LOGF(LogStorageServerPlatformFile, Verbose, "IterateDirectoryRecursively '%ls'", *Path);
 
 	FReadScopeLock _(TocLock);
 	FDirectory* Directory = Directories.FindRef(Path);
@@ -308,7 +308,7 @@ public:
 		IsCompleted.Reset();
 
 #if STORAGE_SERVER_START_CACHE_REPORT
-		UE_CLOG(Cache.Num() > 0, LogStorageServerPlatformFile, Log, TEXT("Engine startup precache size was %llu bytes, %d items. %d items actually retrieved (%d%%). %d cache misses (%d%%)"), CacheSize.load(), Cache.Num(), AccessedChunks.Num(), (AccessedChunks.Num()*100) / Cache.Num(), UncachedChunks.Num(), (UncachedChunks.Num()*100) / (AccessedChunks.Num()+UncachedChunks.Num()) );
+		UE_CLOGF(Cache.Num() > 0, LogStorageServerPlatformFile, Log, "Engine startup precache size was %llu bytes, %d items. %d items actually retrieved (%d%%). %d cache misses (%d%%)", CacheSize.load(), Cache.Num(), AccessedChunks.Num(), (AccessedChunks.Num()*100) / Cache.Num(), UncachedChunks.Num(), (UncachedChunks.Num()*100) / (AccessedChunks.Num()+UncachedChunks.Num()) );
 #endif
 	}
 
@@ -794,7 +794,7 @@ public:
 			}
 			else
 			{
-				UE_LOG(LogStorageServerPlatformFile, Warning, TEXT("Failed to obtain size of file '%s'"), *Filename);
+				UE_LOGF(LogStorageServerPlatformFile, Warning, "Failed to obtain size of file '%ls'", *Filename);
 				FileSize = 0;
 			}
 		}
@@ -982,7 +982,7 @@ TUniquePtr<FArchive> FStorageServerPlatformFile::TryFindProjectStoreMarkerFile(I
 		FString ProjectMarkerPath = ProjectStorePath / TEXT("ue.projectstore");
 		if (IFileHandle* ProjectStoreMarkerHandle = Inner->OpenRead(*ProjectMarkerPath); ProjectStoreMarkerHandle != nullptr)
 		{
-			UE_LOG(LogStorageServerPlatformFile, Display, TEXT("Found '%s'"), *ProjectMarkerPath);
+			UE_LOGF(LogStorageServerPlatformFile, Display, "Found '%ls'", *ProjectMarkerPath);
 			return TUniquePtr<FArchive>(new FArchiveFileReaderGeneric(ProjectStoreMarkerHandle, *ProjectMarkerPath, ProjectStoreMarkerHandle->Size()));
 		}
 	}
@@ -1085,16 +1085,16 @@ FAnsiString FStorageServerPlatformFile::MakeBaseURI()
 				{
 					return NewShareURI;
 				}
-				UE_LOG(LogStorageServerPlatformFile, Error, TEXT("Failed to to resolve or create workspace share path '%s' from %s"), *WorkspaceSharePath, *FString(QueryConnection.GetHostAddr()));
+				UE_LOGF(LogStorageServerPlatformFile, Error, "Failed to to resolve or create workspace share path '%ls' from %ls", *WorkspaceSharePath, *FString(QueryConnection.GetHostAddr()));
 			}
 			else
 			{
-				UE_LOG(LogStorageServerPlatformFile, Error, TEXT("Failed to get list of workspaces from %s. Status: %s"), *FString(QueryConnection.GetHostAddr()), *WorkspacesResponse.Status().ToString());
+				UE_LOGF(LogStorageServerPlatformFile, Error, "Failed to get list of workspaces from %ls. Status: %ls", *FString(QueryConnection.GetHostAddr()), *WorkspacesResponse.Status().ToString());
 			}
 		}
 		else
 		{
-			UE_LOG(LogStorageServerPlatformFile, Error, TEXT("Failed to connect to %s to get list of workspace shares"), *FString(QueryConnection.GetHostAddr()));
+			UE_LOGF(LogStorageServerPlatformFile, Error, "Failed to connect to %ls to get list of workspace shares", *FString(QueryConnection.GetHostAddr()));
 		}
 	}
 	if (!BaseURI.IsEmpty())
@@ -1212,19 +1212,19 @@ bool FStorageServerPlatformFile::ShouldBeUsed(IPlatformFile* Inner, const TCHAR*
 				{
 					HostPort = SerializedHostPort;
 				}
-				UE_LOG(LogStorageServerPlatformFile, Display, TEXT("Using connection settings from ue.projectstore: HostAddrs='%s' and HostPort='%d'"), *FString::Join(HostAddrs, TEXT("+")), HostPort);
+				UE_LOGF(LogStorageServerPlatformFile, Display, "Using connection settings from ue.projectstore: HostAddrs='%ls' and HostPort='%d'", *FString::Join(HostAddrs, TEXT("+")), HostPort);
 			}
 		}
 		else
 		{
-			UE_LOG(LogStorageServerPlatformFile, Error, TEXT("Failed to Deserialize ue.projectstore!'"));
+			UE_LOGF(LogStorageServerPlatformFile, Error, "Failed to Deserialize ue.projectstore!'");
 		}
 	}
 
 	FString Host;
 	if (FParse::Value(CmdLine, TEXT("-ZenStoreHost="), Host))
 	{
-		UE_LOG(LogStorageServerPlatformFile, Display, TEXT("Adding connection settings from command line: -ZenStoreHost='%s'"), *Host);
+		UE_LOGF(LogStorageServerPlatformFile, Display, "Adding connection settings from command line: -ZenStoreHost='%ls'", *Host);
 		if (!Host.ParseIntoArray(HostAddrs, TEXT("+"), true))
 		{
 			HostAddrs.Add(Host);
@@ -1232,7 +1232,7 @@ bool FStorageServerPlatformFile::ShouldBeUsed(IPlatformFile* Inner, const TCHAR*
 	}
 	if (FParse::Value(CmdLine, TEXT("-ZenStorePort="), HostPort))
 	{
-		UE_LOG(LogStorageServerPlatformFile, Display, TEXT("Using connection settings from command line: -ZenStorePort='%d'"), HostPort);
+		UE_LOGF(LogStorageServerPlatformFile, Display, "Using connection settings from command line: -ZenStorePort='%d'", HostPort);
 	}
 	// add hostnames as last resort
 	HostAddrs.Append(HostNames);
@@ -1317,12 +1317,12 @@ bool FStorageServerPlatformFile::Initialize(IPlatformFile* Inner, const TCHAR* C
 								FString Credentials = FString::Printf(TEXT("%s:%s"), *Username, *Password);
 								FString Encoded = FBase64::Encode(Credentials);
 								AuthHeaders.Add(FAnsiString::Printf("Authorization: Basic %s", TCHAR_TO_ANSI(*Encoded)));
-								UE_LOG(LogStorageServerPlatformFile, Display, TEXT("Using password authentication from ue.projectstore"));
+								UE_LOGF(LogStorageServerPlatformFile, Display, "Using password authentication from ue.projectstore");
 							}
 						}
 					}
 
-					UE_LOG(LogStorageServerPlatformFile, Display, TEXT("Using settings from ue.projectstore: ServerProject='%s' and ServerPlatform='%s'"), *ServerProject, *ServerPlatform);
+					UE_LOGF(LogStorageServerPlatformFile, Display, "Using settings from ue.projectstore: ServerProject='%ls' and ServerPlatform='%ls'", *ServerProject, *ServerPlatform);
 				}
 
 				const TArray<TSharedPtr<FJsonValue>>* RemapDirectoriesArrayPtr = nullptr;
@@ -1342,19 +1342,19 @@ bool FStorageServerPlatformFile::Initialize(IPlatformFile* Inner, const TCHAR* C
 	
 		if (FParse::Value(CmdLine, TEXT("-ZenStoreProject="), ServerProject))
 		{
-			UE_LOG(LogStorageServerPlatformFile, Display, TEXT("Using settings from command line: -ZenStoreProject='%s'"), *ServerProject);
+			UE_LOGF(LogStorageServerPlatformFile, Display, "Using settings from command line: -ZenStoreProject='%ls'", *ServerProject);
 		}
 		if (FParse::Value(CmdLine, TEXT("-ZenStorePlatform="), ServerPlatform))
 		{
-			UE_LOG(LogStorageServerPlatformFile, Display, TEXT("Using settings from command line: -ZenStorePlatform='%s'"), *ServerPlatform);
+			UE_LOGF(LogStorageServerPlatformFile, Display, "Using settings from command line: -ZenStorePlatform='%ls'", *ServerPlatform);
 		}
 		if (FParse::Value(CmdLine, TEXT("-ZenStoreBaseURI="), BaseURI))
 		{
-			UE_LOG(LogStorageServerPlatformFile, Display, TEXT("Using settings from command line: -ZenStoreBaseURI='%s'"), *BaseURI);
+			UE_LOGF(LogStorageServerPlatformFile, Display, "Using settings from command line: -ZenStoreBaseURI='%ls'", *BaseURI);
 		}
 		if (FParse::Value(CmdLine, TEXT("-ZenWorkspaceSharePath="), WorkspaceSharePath))
 		{
-			UE_LOG(LogStorageServerPlatformFile, Display, TEXT("Using settings from command line: -ZenWorkspaceSharePath='%s'"), *WorkspaceSharePath);
+			UE_LOGF(LogStorageServerPlatformFile, Display, "Using settings from command line: -ZenWorkspaceSharePath='%ls'", *WorkspaceSharePath);
 		}
 
 		if (UE::IsUsingZenPakFileStreaming())
@@ -1431,7 +1431,7 @@ void FStorageServerPlatformFile::InitializeConnection()
 		else
 		{
 			FStringView HostAddr = Connection->GetHostAddr();
-			UE_LOG(LogStorageServerPlatformFile, Fatal, TEXT("Failed to get file list from Zen at '%.*s'"), HostAddr.Len(), HostAddr.GetData());
+			UE_LOGF(LogStorageServerPlatformFile, Fatal, "Failed to get file list from Zen at '%.*ls'", HostAddr.Len(), HostAddr.GetData());
 		}
 	}
 	else if (bAbortOnConnectionFailure)
@@ -1451,7 +1451,7 @@ void FStorageServerPlatformFile::InitializeConnection()
 			FPlatformMisc::MessageBoxExt(EAppMsgType::Ok, *FailedConnectionText, *FailedConnectionTitle);
 		}
 
-		UE_LOG(LogStorageServerPlatformFile, Error, TEXT("Failed to initialize connection to any of [ %s ]"),
+		UE_LOGF(LogStorageServerPlatformFile, Error, "Failed to initialize connection to any of [ %ls ]",
 			*FString::Join(HostAddrs, TEXT(", ")));
 		Connection.Reset();
 		FPlatformMisc::RequestExit(true);
@@ -1460,7 +1460,7 @@ void FStorageServerPlatformFile::InitializeConnection()
 	{
 		// One reason the connection can fail is if the oplog has not yet been created during a new cook:
 		// we will get the message "oplog <PlatformName> not found in project <ProjectName>".
-		UE_LOG(LogStorageServerPlatformFile, Display, TEXT("Failed to initialize connection to any of [ %s ]"),
+		UE_LOGF(LogStorageServerPlatformFile, Display, "Failed to initialize connection to any of [ %ls ]",
 			*FString::Join(HostAddrs, TEXT(", ")));
 		Connection.Reset();
 	}
@@ -1477,7 +1477,7 @@ void FStorageServerPlatformFile::UpdateFileList()
 	if (!SendGetFileListMessage())
 	{
 		FStringView HostAddr = Connection->GetHostAddr();
-		UE_LOG(LogStorageServerPlatformFile, Fatal, TEXT("Failed to get file list from Zen at '%.*s'"), HostAddr.Len(), HostAddr.GetData());
+		UE_LOGF(LogStorageServerPlatformFile, Fatal, "Failed to get file list from Zen at '%.*ls'", HostAddr.Len(), HostAddr.GetData());
 	}
 }
 
@@ -1852,7 +1852,7 @@ FString FStorageServerPlatformFile::GetFilenameOnDisk(const TCHAR* Filename)
 	TStringBuilder<1024> StorageServerFilename;
 	if (MakeStorageServerPath(Filename, StorageServerFilename) && ServerToc.FileExists(*StorageServerFilename))
 	{
-		UE_LOG(LogStorageServerPlatformFile, Warning, TEXT("Attempting to get disk filename of remote file '%s'"), Filename);
+		UE_LOGF(LogStorageServerPlatformFile, Warning, "Attempting to get disk filename of remote file '%ls'", Filename);
 		return Filename;
 	}
 	return (LowerLevel && IsNonServerFilenameAllowed(Filename)) ? LowerLevel->GetFilenameOnDisk(Filename) : Filename;
@@ -2074,8 +2074,8 @@ bool FStorageServerPlatformFile::IsNonServerFilenameAllowed(FStringView InFilena
 				}
 			}
 
-			UE_CLOG(!bAllowed, LogStorageServerPlatformFile, VeryVerbose,
-				TEXT("Access to file '%.*s' is limited to server contents due to file extension being listed in ExcludedNonServerExtensions."),
+			UE_CLOGF(!bAllowed, LogStorageServerPlatformFile, VeryVerbose,
+				"Access to file '%.*ls' is limited to server contents due to file extension being listed in ExcludedNonServerExtensions.",
 				InFilename.Len(), InFilename.GetData())
 		}
 	}
@@ -2265,7 +2265,7 @@ void FStorageServerPlatformFile::OnCookOnTheFlyMessage(const UE::Cook::FCookOnTh
 	{
 		case UE::Cook::ECookOnTheFlyMessage::FilesAdded:
 		{
-			UE_LOG(LogCookOnTheFly, Verbose, TEXT("Received '%s' message"), LexToString(Message.GetHeader().MessageType));
+			UE_LOGF(LogCookOnTheFly, Verbose, "Received '%ls' message", LexToString(Message.GetHeader().MessageType));
 
 			TArray<FString> Filenames;
 			TArray<FIoChunkId> ChunkIds;
@@ -2280,7 +2280,7 @@ void FStorageServerPlatformFile::OnCookOnTheFlyMessage(const UE::Cook::FCookOnTh
 
 			for (int32 Idx = 0, Num = Filenames.Num(); Idx < Num; ++Idx)
 			{
-				UE_LOG(LogCookOnTheFly, Verbose, TEXT("Adding file '%s'"), *Filenames[Idx]);
+				UE_LOGF(LogCookOnTheFly, Verbose, "Adding file '%ls'", *Filenames[Idx]);
 				ServerToc.AddFile(ChunkIds[Idx], Filenames[Idx], STORAGE_SERVER_FILE_UNKOWN_SIZE);
 			}
 
