Reflection capture editor optimization for "r.ReflectionCaptureUpdateEveryFrame" CVar.  Iterator changed from actor to component iterator, 30x faster, saving 1.3 ms in a test level.  Hidden reflection captures no longer refresh every frame, and alternative CVar setting (r.ReflectionCaptureUpdateEveryFrame = 2) only updates selected reflection captures.  These two filtering options can allow the user to avoid single digit frame rates that otherwise result in a level with more than a couple reflection captures, making the CVar usable in those cases.

#Synced-CL 52603298

diff --git a/Engine/Source/Runtime/Engine/Private/Components/ReflectionCaptureComponent.cpp b/Engine/Source/Runtime/Engine/Private/Components/ReflectionCaptureComponent.cpp
--- a/Engine/Source/Runtime/Engine/Private/Components/ReflectionCaptureComponent.cpp
+++ b/Engine/Source/Runtime/Engine/Private/Components/ReflectionCaptureComponent.cpp
@@ -60,6 +60,9 @@ TAutoConsoleVariable<int32> CVarReflectionCaptureUpdateEveryFrame(
 	TEXT("r.ReflectionCaptureUpdateEveryFrame"),
 	0,
 	TEXT("When set, reflection captures will constantly be scheduled for update.\n")
+	TEXT("0:  Disabled.\n")
+	TEXT("1:  Update all every frame.\n")
+	TEXT("2:  Update selected every frame.\n")
 );
 
 TAutoConsoleVariable<int32> CVarReflectionCaptureRuntime(
@@ -1127,17 +1130,27 @@ void UReflectionCaptureComponent::UpdateReflectionCaptureContents(UWorld* WorldT
 		//guarantee that all render proxies are up to date before kicking off this render
 		WorldToUpdate->SendAllEndOfFrameUpdates();
 
-		if (CVarReflectionCaptureUpdateEveryFrame.GetValueOnGameThread())
+		int32 UpdateEveryFrameMode = CVarReflectionCaptureUpdateEveryFrame.GetValueOnGameThread();
+		if (UpdateEveryFrameMode)
 		{
-			for (FActorIterator It(WorldToUpdate); It; ++It)
+			ForEachObjectOfClass(UReflectionCaptureComponent::StaticClass(), [WorldToUpdate, UpdateEveryFrameMode](UObject* ComponentObject)
 			{
-				TInlineComponentArray<UReflectionCaptureComponent*> Components;
-				(*It)->GetComponents(Components);
-				for (int32 ComponentIndex = 0; ComponentIndex < Components.Num(); ComponentIndex++)
+				if (ComponentObject->GetWorld() == WorldToUpdate)
 				{
-					Components[ComponentIndex]->MarkDirtyForRecapture(); // Continuously refresh reflection captures
+					if (UReflectionCaptureComponent* Component = CastChecked<UReflectionCaptureComponent>(ComponentObject))
+					{
+						AActor* Owner = Component->GetOwner();
+
+#if WITH_EDITOR
+						// UpdateEveryFrameMode == 2 indicates only selected objects should be refreshed, mode 1 refreshes everything.
+						if (Owner && Component->IsVisibleInEditor() && !Owner->IsHiddenEd() && (UpdateEveryFrameMode == 1 || Component->IsSelectedInEditor() || Owner->IsSelectedInEditor()))
+#endif
+						{
+							Component->MarkDirtyForRecapture();		// Continuously refresh reflection captures
+						}
+					}
 				}
-			}
+			}, /*bIncludeDerivedClasses=*/ true, RF_ClassDefaultObject, EInternalObjectFlags::Garbage);
 		}
 
 		TArray<UReflectionCaptureComponent*> WorldCombinedCaptures;
