Update "stat fps" frame time indicator to use the same smoothing factor as "stat unit".

#Synced-CL 52628728

diff --git a/Engine/Source/Runtime/Engine/Private/UnrealEngine.cpp b/Engine/Source/Runtime/Engine/Private/UnrealEngine.cpp
--- a/Engine/Source/Runtime/Engine/Private/UnrealEngine.cpp
+++ b/Engine/Source/Runtime/Engine/Private/UnrealEngine.cpp
@@ -819,8 +819,8 @@ void CalculateFPSTimings()
 		}
 	}
 	
-	// A 3/4, 1/4 split gets close to a simple 10 frame moving average
-	GAverageMS = GAverageMS * 0.75f + FrameTimeMS * 0.25f;
+	// Exponential moving average with a 0.1 smoothing factor.
+	GAverageMS = GAverageMS * 0.9f + FrameTimeMS * 0.1f;
 
 	LastTime = CurrentTime;
 
