Iris - Log an error when RepMovementNetSerializer isn't reigstered as well as when layout is unexpected but the serializer is registered anyway.

#Synced-CL 52609747

diff --git a/Engine/Source/Runtime/Engine/Private/RepMovementNetSerializer.cpp b/Engine/Source/Runtime/Engine/Private/RepMovementNetSerializer.cpp
--- a/Engine/Source/Runtime/Engine/Private/RepMovementNetSerializer.cpp
+++ b/Engine/Source/Runtime/Engine/Private/RepMovementNetSerializer.cpp
@@ -5,6 +5,7 @@
 #include UE_INLINE_GENERATED_CPP_BY_NAME(RepMovementNetSerializer)
 
 #include "Engine/ReplicatedState.h"
+#include "Iris/Core/IrisLog.h"
 #include "Iris/Serialization/BitPacking.h"
 #include "Iris/Serialization/NetBitStreamReader.h"
 #include "Iris/Serialization/NetBitStreamUtil.h"
@@ -874,12 +875,14 @@ void FRepMovementNetSerializer::FNetSerializerRegistryDelegates::OnPreFreezeNetS
 	// If our quantized type doesn't meet the requirements of the serializers we're forwarding too then bail out.
 	if (!QuantizedTypeMeetRequirements())
 	{
+
+		UE_LOGF(LogIris, Error, "FRepMovementNetSerializer will not be registered");
 		return;
 	}
 
 	if (!IsRepMovementLayoutAsExpected())
 	{
-		return;
+		UE_LOGF(LogIris, Error, "FRepMovement layout is unexpected. Things may not replicate as desired.");
 	}
 
 	InitNetSerializer();
