[WorldPartition] Guard variable to stop querying non-spatial cells twice or more when querying streaming sources

#Synced-CL 52575010

diff --git a/Engine/Source/Runtime/Engine/Private/WorldPartition/RuntimeHashSet/WorldPartitionRuntimeHashSet.cpp b/Engine/Source/Runtime/Engine/Private/WorldPartition/RuntimeHashSet/WorldPartitionRuntimeHashSet.cpp
--- a/Engine/Source/Runtime/Engine/Private/WorldPartition/RuntimeHashSet/WorldPartitionRuntimeHashSet.cpp
+++ b/Engine/Source/Runtime/Engine/Private/WorldPartition/RuntimeHashSet/WorldPartitionRuntimeHashSet.cpp
@@ -14,8 +14,8 @@
 #include "WorldPartition/WorldPartitionSubsystem.h"
 #include "Algo/RemoveIf.h"
 #include "Algo/Transform.h"
-#include "Misc/HashBuilder.h"
 #include "Misc/ArchiveMD5.h"
+#include "Misc/HashBuilder.h"
 #include "UObject/UObjectGlobalsInternal.h"
 
 #include UE_INLINE_GENERATED_CPP_BY_NAME(WorldPartitionRuntimeHashSet)
@@ -610,9 +610,13 @@ void UWorldPartitionRuntimeHashSet::ForEachStreamingCellsQuery(const FWorldParti
 		return true;
 	};
 
-	auto ForEachNonSpatiallyLoadedCells = [&ShouldAddCell, &QuerySource, &Func](TArray<TObjectPtr<UWorldPartitionRuntimeCell>> InNonSpatiallyLoadedCells)
+	auto ForEachNonSpatiallyLoadedCells = [this, &ShouldAddCell, &QuerySource, &Func](const TArray<TObjectPtr<UWorldPartitionRuntimeCell>>& InNonSpatiallyLoadedCells)
 	{
-		for (UWorldPartitionRuntimeCell* Cell : InNonSpatiallyLoadedCells)
+		if (!bQueryNonSpatialCells)
+		{
+			return true;
+		}
+		for (const UWorldPartitionRuntimeCell* Cell : InNonSpatiallyLoadedCells)
 		{
 			if (ShouldAddCell(Cell, QuerySource))
 			{
diff --git a/Engine/Source/Runtime/Engine/Private/WorldPartition/WorldPartitionRuntimeSpatialHash.cpp b/Engine/Source/Runtime/Engine/Private/WorldPartition/WorldPartitionRuntimeSpatialHash.cpp
--- a/Engine/Source/Runtime/Engine/Private/WorldPartition/WorldPartitionRuntimeSpatialHash.cpp
+++ b/Engine/Source/Runtime/Engine/Private/WorldPartition/WorldPartitionRuntimeSpatialHash.cpp
@@ -295,7 +295,7 @@ void FSpatialHashStreamingGrid::DumpStateLog(FHierarchicalLogArchive& Ar) const
 }
 #endif
 
-void FSpatialHashStreamingGrid::GetCells(const FWorldPartitionStreamingQuerySource& QuerySource, TSet<const UWorldPartitionRuntimeCell*>& OutCells, bool bEnableZCulling, FWorldPartitionQueryCache* QueryCache) const
+void FSpatialHashStreamingGrid::GetCells(const FWorldPartitionStreamingQuerySource& QuerySource, TSet<const UWorldPartitionRuntimeCell*>& OutCells, bool bEnableZCulling, FWorldPartitionQueryCache* QueryCache, bool bQueryNonSpatialCells) const
 {
 	TRACE_CPUPROFILER_EVENT_SCOPE(FSpatialHashStreamingGrid::GetCells_QuerySource);
 
@@ -352,26 +352,29 @@ void FSpatialHashStreamingGrid::GetCells(const FWorldPartitionStreamingQuerySour
 		});
 	}
 
-	// Non Spatial (always included)
-	auto ForEachGridLevelCells = [&QuerySource, &OutCells, ShouldAddCell](const FSpatialHashStreamingGridLevel& InGridLevel)
+	// Non Spatial (LayersCells is the array containing non-spatial cells)
+	if (bQueryNonSpatialCells)
 	{
-		for (const FSpatialHashStreamingGridLayerCell& LayerCell : InGridLevel.LayerCells)
+		auto ForEachGridLevelCells = [&QuerySource, &OutCells, ShouldAddCell](const FSpatialHashStreamingGridLevel& InGridLevel)
 		{
-			for (const UWorldPartitionRuntimeCell* Cell : LayerCell.GridCells)
+			for (const FSpatialHashStreamingGridLayerCell& LayerCell : InGridLevel.LayerCells)
 			{
-				if (ShouldAddCell(Cell, QuerySource))
+				for (const UWorldPartitionRuntimeCell* Cell : LayerCell.GridCells)
 				{
-					OutCells.Add(Cell);
+					if (ShouldAddCell(Cell, QuerySource))
+					{
+						OutCells.Add(Cell);
+					}
 				}
 			}
-		}
-	};
+		};
 
-	const int32 TopLevel = GridLevels.Num() - 1;
-	ForEachGridLevelCells(GridLevels[TopLevel]);
-	if (InjectedGridLevels.IsValidIndex(TopLevel))
-	{
-		ForEachGridLevelCells(InjectedGridLevels[TopLevel]);
+		const int32 TopLevel = GridLevels.Num() - 1;
+		ForEachGridLevelCells(GridLevels[TopLevel]);
+		if (InjectedGridLevels.IsValidIndex(TopLevel))
+		{
+			ForEachGridLevelCells(InjectedGridLevels[TopLevel]);
+		}
 	}
 }
 
@@ -1792,7 +1795,7 @@ void UWorldPartitionRuntimeSpatialHash::ForEachStreamingCellsQuery(const FWorldP
 		if (IsCellRelevantFor(StreamingGrid.bClientOnlyVisible))
 		{
 			TSet<const UWorldPartitionRuntimeCell*> Cells;
-			StreamingGrid.GetCells(QuerySource, Cells, GetEffectiveEnableZCulling(bEnableZCulling), QueryCache);
+			StreamingGrid.GetCells(QuerySource, Cells, GetEffectiveEnableZCulling(bEnableZCulling), QueryCache, bQueryNonSpatialCells);
 
 			for (const UWorldPartitionRuntimeCell* Cell : Cells)
 			{
diff --git a/Engine/Source/Runtime/Engine/Private/WorldPartition/WorldPartitionStreamingPolicy.cpp b/Engine/Source/Runtime/Engine/Private/WorldPartition/WorldPartitionStreamingPolicy.cpp
--- a/Engine/Source/Runtime/Engine/Private/WorldPartition/WorldPartitionStreamingPolicy.cpp
+++ b/Engine/Source/Runtime/Engine/Private/WorldPartition/WorldPartitionStreamingPolicy.cpp
@@ -1088,6 +1088,9 @@ bool UWorldPartitionStreamingPolicy::IsStreamingCompleted(const TArray<FWorldPar
 	}
 
 	// Test spatially loaded cells using streaming sources (or provided streaming source)
+	// Non-spatial cells were already queried above; skip them for each per-source spatial query.
+	check(WorldPartition->RuntimeHash->bQueryNonSpatialCells); // Flag future re-entrancy issues
+	TGuardValue<bool> QueryNonSpatialCells(WorldPartition->RuntimeHash->bQueryNonSpatialCells, false);
 	TArrayView<const FWorldPartitionStreamingSource> QueriedStreamingSources = bTestProvidedStreamingSource ? *InStreamingSources : CurrentState.StreamingSources;
 	for (const FWorldPartitionStreamingSource& StreamingSource : QueriedStreamingSources)
 	{
diff --git a/Engine/Source/Runtime/Engine/Public/WorldPartition/WorldPartitionRuntimeHash.h b/Engine/Source/Runtime/Engine/Public/WorldPartition/WorldPartitionRuntimeHash.h
--- a/Engine/Source/Runtime/Engine/Public/WorldPartition/WorldPartitionRuntimeHash.h
+++ b/Engine/Source/Runtime/Engine/Public/WorldPartition/WorldPartitionRuntimeHash.h
@@ -192,6 +192,7 @@ class UWorldPartitionRuntimeHash : public UObject
 
 	friend class UWorldPartition;
 	friend class URuntimePartition;
+	friend class UWorldPartitionStreamingPolicy;
 
 #if WITH_EDITOR
 	virtual void SetDefaultValues() {}
@@ -297,6 +298,9 @@ protected:
 
 	TSet<TWeakObjectPtr<URuntimeHashExternalStreamingObjectBase>> InjectedExternalStreamingObjects;
 
+	// Variable set in scoped guard to skip non spatials cells in repeated queries
+	mutable bool bQueryNonSpatialCells = true;
+
 	virtual bool SupportsWorldAssetStreaming(const FName& InTargetGrid) { return false; }
 	virtual FGuid RegisterWorldAssetStreaming(const UWorldPartition::FRegisterWorldAssetStreamingParams& InParams) { return FGuid(); }
 	virtual bool UnregisterWorldAssetStreaming(const FGuid& InWorldAssetStreamingGuid) { return false; }
diff --git a/Engine/Source/Runtime/Engine/Public/WorldPartition/WorldPartitionRuntimeSpatialHash.h b/Engine/Source/Runtime/Engine/Public/WorldPartition/WorldPartitionRuntimeSpatialHash.h
--- a/Engine/Source/Runtime/Engine/Public/WorldPartition/WorldPartitionRuntimeSpatialHash.h
+++ b/Engine/Source/Runtime/Engine/Public/WorldPartition/WorldPartitionRuntimeSpatialHash.h
@@ -183,7 +183,7 @@ struct FSpatialHashStreamingGrid
 
 	// Used by PIE/Game
 	ENGINE_API int64 GetCellSize(int32 Level) const;
-	ENGINE_API void GetCells(const FWorldPartitionStreamingQuerySource& QuerySource, TSet<const UWorldPartitionRuntimeCell*>& OutCells, bool bEnableZCulling, FWorldPartitionQueryCache* QueryCache = nullptr) const;
+	ENGINE_API void GetCells(const FWorldPartitionStreamingQuerySource& QuerySource, TSet<const UWorldPartitionRuntimeCell*>& OutCells, bool bEnableZCulling, FWorldPartitionQueryCache* QueryCache = nullptr, bool bQueryNonSpatialCells = true) const;
 	ENGINE_API void GetCells(const TArray<FWorldPartitionStreamingSource>& Sources, UWorldPartitionRuntimeHash::FStreamingSourceCells& OutActivateCells, UWorldPartitionRuntimeHash::FStreamingSourceCells& OutLoadCells, bool bEnableZCulling, const FWorldPartitionStreamingContext& Context, float LoadingRangeScale = 1.f) const;
 	ENGINE_API void GetNonSpatiallyLoadedCells(TSet<const UWorldPartitionRuntimeCell*>& OutActivateCells, TSet<const UWorldPartitionRuntimeCell*>& OutLoadCells, const FWorldPartitionStreamingContext& Context) const;
 	ENGINE_API void Draw2D(const class UWorldPartitionRuntimeSpatialHash* Owner, const FBox2D& Region2D, const FBox2D& GridScreenBounds, TFunctionRef<FVector2D(const FVector2D&, bool)> WorldToScreen, FWorldPartitionDraw2DContext& DrawContext) const;
diff --git a/Engine/Source/Runtime/Engine/Public/WorldPartition/WorldPartitionStreamingSource.h b/Engine/Source/Runtime/Engine/Public/WorldPartition/WorldPartitionStreamingSource.h
--- a/Engine/Source/Runtime/Engine/Public/WorldPartition/WorldPartitionStreamingSource.h
+++ b/Engine/Source/Runtime/Engine/Public/WorldPartition/WorldPartitionStreamingSource.h
@@ -281,7 +281,7 @@ struct FWorldPartitionStreamingQuerySource
 	UPROPERTY(BlueprintReadWrite, EditAnywhere, Category = "Query")
 	bool bDataLayersOnly;
 
-	/* If False, Location/Radius will not be used to find the cells. Only AlwaysLoaded cells will be returned by the query. */
+	/* If False, Location/Radius will not be used to find spatial cells. Non-spatial cells are always returned by the query, but can be further filtered by bDataLayersOnly. */
 	UPROPERTY(BlueprintReadWrite, EditAnywhere, Category = "Query")
 	bool bSpatialQuery;
 
