Add bClientIgnoreMovementCorrections to UCharacterMovementComponent.
New flag that allows the client to skip applying server movement corrections in ClientAdjustPosition. Intended for use alongside bIgnoreClientMovementErrorChecksAndCorrection, which only prevents the server from generating corrections but does not protect against in-flight ones sent before the server received the flag-setting RPC.
#Synced-CL 52679658

diff --git a/Engine/Source/Runtime/Engine/Classes/GameFramework/CharacterMovementComponent.h b/Engine/Source/Runtime/Engine/Classes/GameFramework/CharacterMovementComponent.h
--- a/Engine/Source/Runtime/Engine/Classes/GameFramework/CharacterMovementComponent.h
+++ b/Engine/Source/Runtime/Engine/Classes/GameFramework/CharacterMovementComponent.h
@@ -969,16 +969,26 @@ public:
 	UPROPERTY(Transient)
 	uint8 bNetworkGravityDirectionChanged:1;
 
-	/** 
+	/**
 	 * If true, we should ignore server location difference checks for client error on this movement component.
 	 * This can be useful when character is moving at extreme speeds for a duration and you need it to look
 	 * smooth on clients without the server correcting the client. Make sure to disable when done, as this would
 	 * break this character's server-client movement correction.
-	 * @see bServerAcceptClientAuthoritativePosition, ServerCheckClientError()
+	 * Note: This only prevents the server from generating corrections. To also prevent the client from applying
+	 * in-flight corrections, use bClientIgnoreMovementCorrections alongside this flag.
+	 * @see bClientIgnoreMovementCorrections, bServerAcceptClientAuthoritativePosition, ServerCheckClientError()
 	 */
 	UPROPERTY(Category="Character Movement", EditAnywhere, BlueprintReadWrite)
 	uint8 bIgnoreClientMovementErrorChecksAndCorrection:1;
 
+	/**
+	 * When true, the client will not apply server movement corrections (skips ClientAdjustPosition).
+	 * Typically used together with bIgnoreClientMovementErrorChecksAndCorrection, which prevents the
+	 * server from generating corrections in the first place.
+	 */
+	UPROPERTY(Category="Character Movement", EditAnywhere, BlueprintReadWrite)
+	uint8 bClientIgnoreMovementCorrections:1;
+
 	/**
 	 * If true, and server does not detect client position error, server will copy the client movement location/velocity/etc after simulating the move.
 	 * This can be useful for short bursts of movement that are difficult to sync over the network.
diff --git a/Engine/Source/Runtime/Engine/Private/Components/CharacterMovementComponent.cpp b/Engine/Source/Runtime/Engine/Private/Components/CharacterMovementComponent.cpp
--- a/Engine/Source/Runtime/Engine/Private/Components/CharacterMovementComponent.cpp
+++ b/Engine/Source/Runtime/Engine/Private/Components/CharacterMovementComponent.cpp
@@ -793,6 +793,7 @@ UCharacterMovementComponent::UCharacterMovementComponent(const FObjectInitialize
 	bImpartBaseVelocityZ = true;
 	bImpartBaseAngularVelocity = true;
 	bIgnoreClientMovementErrorChecksAndCorrection = false;
+	bClientIgnoreMovementCorrections = false;
 	bServerAcceptClientAuthoritativePosition = false;
 	bAlwaysCheckFloor = true;
 
@@ -11239,6 +11240,23 @@ void UCharacterMovementComponent::ClientAdjustPosition_Implementation
 		return;
 	}
 
+	// Skip applying server corrections when flagged. AckMove() still prunes SavedMoves to prevent unbounded growth.
+	if (bClientIgnoreMovementCorrections)
+	{
+		FNetworkPredictionData_Client_Character* ClientData = GetPredictionData_Client_Character();
+		if (ClientData)
+		{
+			const int32 MoveIndex = ClientData->GetSavedMoveIndex(TimeStamp);
+			if (MoveIndex != INDEX_NONE)
+			{
+				ClientData->AckMove(MoveIndex, *this);
+			}
+		}
+
+		UE_LOG(LogNetPlayerMovement, Verbose, TEXT("ClientAdjustPosition_Implementation: Ignoring server correction. bClientIgnoreMovementCorrections is set. TimeStamp: %f"), TimeStamp);
+
+		return;
+	}
 
 	FNetworkPredictionData_Client_Character* ClientData = GetPredictionData_Client_Character();
 	check(ClientData);
