[MIR] Replace eager function call push/pop with generic on-demand subgraph mechanism

Replaces PushFunctionCall/PopFunctionCall and FAnalysisContext with a generic subgraph API on the emitter. The builder is now fully decoupled from function call/input expression types, and function call outputs are only evaluated when actually needed. Inlining of function calls now happens on demand and lazily, as if the function subgraph was pasted in the material. This means that Call inputs are not evaluated if they don't contribute to the end final material result.

All UE-specific function call semantics now live in MaterialExpressionsToMIR.cpp.

#rnx
#Synced-CL 52678682

diff --git a/Engine/Source/Runtime/Engine/Private/Materials/MaterialExpressionsToMIR.cpp b/Engine/Source/Runtime/Engine/Private/Materials/MaterialExpressionsToMIR.cpp
--- a/Engine/Source/Runtime/Engine/Private/Materials/MaterialExpressionsToMIR.cpp
+++ b/Engine/Source/Runtime/Engine/Private/Materials/MaterialExpressionsToMIR.cpp
@@ -329,15 +329,49 @@ static void BuildMaterialExpressionExternalCodeBase(MIR::FEmitter& Em, UMaterial
 
 static FName NAME_CameraVector("CameraVector");
 
-/* Constants */
-
 void UMaterialExpression::Build(MIR::FEmitter& Em)
 {
 	Em.Error(TEXT("Unsupported material expression."));
 } 
 
+void UMaterialExpressionMaterialFunctionCall::Build(MIR::FEmitter& Em)
+{
+	MIR::TTemporaryArray<MIR::FSubgraphOutputMapping> OutputMappings{ FunctionOutputs.Num() };
+	for (int i = 0; i < FunctionOutputs.Num(); ++i)
+	{
+		OutputMappings[i].Expression  = FunctionOutputs[i].ExpressionOutput.Get();
+		OutputMappings[i].OutputIndex = 0; // Note: UMaterialExpressionFunctionOutput have a single bypass expression output
+	}
+	Em.Subgraph(FunctionInputs.Num(), OutputMappings);
+}
+
 void UMaterialExpressionFunctionInput::Build(MIR::FEmitter& Em)
 {
+	// When inside a function call, find our input index and fetch the caller-provided value.
+	if (UMaterialExpressionMaterialFunctionCall* Call = Cast<UMaterialExpressionMaterialFunctionCall>(Em.GetSubgraphExpression()))
+	{
+		int32 InputIndex = Call->FunctionInputs.IndexOfByPredicate([this](const FFunctionExpressionInput& FI) { return FI.ExpressionInput.Get() == this; });
+		if (InputIndex != INDEX_NONE)
+		{
+			FValueRef CallerValue = bUsePreviewValueAsDefault ? Em.TrySubgraphInput(InputIndex) : Em.SubgraphInput(InputIndex);
+			UE_MIR_CHECKPOINT(Em);
+
+			if (CallerValue)
+			{
+				MIR::FType Type = MIR::FType::FromMaterialValueType(GetInputValueType(0));
+				Em.Output(0, Em.Cast(CallerValue, Type));
+				return;
+			}
+			// Optional input not connected, fall through to preview/default.
+		}
+		else
+		{
+			Em.Errorf(TEXT("Function input '%s' not found in call '%s'."), *InputName.ToString(), Call->MaterialFunction ? *Call->MaterialFunction->GetName() : TEXT("<null>"));
+			return;
+		}
+	}
+
+	// Not inside a call, or function input not found. Use the preview connection or a default constant.
 	FValueRef OutputValue = Em.TryInput(&Preview);
 	if (OutputValue)
 	{
diff --git a/Engine/Source/Runtime/Engine/Private/Materials/MaterialIREmitter.cpp b/Engine/Source/Runtime/Engine/Private/Materials/MaterialIREmitter.cpp
--- a/Engine/Source/Runtime/Engine/Private/Materials/MaterialIREmitter.cpp
+++ b/Engine/Source/Runtime/Engine/Private/Materials/MaterialIREmitter.cpp
@@ -140,7 +140,7 @@ struct FEmitter::FPrivate
 
 // Creates a copy of specified array using the Module allocator and returns it.
 template <typename T>
-TConstArrayView<T> MakeArrayCopy(FEmitter& Emitter, TConstArrayView<T> Array)
+static TConstArrayView<T> MakeArrayCopy(FEmitter& Emitter, TConstArrayView<T> Array)
 {
 	T* Data = (T*)FEmitter::FPrivate::Allocate(Emitter, sizeof(T) * Array.Num(), alignof(T));
 	FMemory::Memcpy(Data, Array.GetData(), sizeof(T) * Array.Num());
@@ -533,15 +533,14 @@ FType FEmitter::GetMaterialAggregateAttributeType(const UMaterialAggregate* Aggr
 
 FValueRef FEmitter::TryInput(const FExpressionInput* InInput)
 {
-	if (bExpressionRequestsInputsOnDemand)
+	TOptional<FValue*> Result = Internal::RequestInputValue(BuilderImpl, InInput);
+	if (!Result.IsSet())
 	{
-		if (!Internal::RequestInputValue(BuilderImpl, InInput))
-		{
-			State = EState::Pending;
-			return { FPoison::Get(), InInput };
-		}
+		check(bExpressionRequestsInputsOnDemand);
+		State = EState::Pending;
+		return { FPoison::Get(), InInput };
 	}
-	return FValueRef{ Internal::FetchValueFromExpressionInput(BuilderImpl, InInput), InInput };
+	return FValueRef{ Result.GetValue(), InInput };
 }
 
 FValueRef FEmitter::Input(const FExpressionInput* InInput)
@@ -878,6 +877,41 @@ FEmitter& FEmitter::OutputsWithComponentMask(FValueRef Value)
 	return *this;
 }
 
+/*------------------------------- Subgraph management ------------------------------*/
+
+void FEmitter::Subgraph(int32 NumInputs, TArrayView<FSubgraphOutputMapping> Outputs)
+{
+	Internal::DeclareSubgraph(BuilderImpl, NumInputs, Outputs);
+}
+
+FValueRef FEmitter::TrySubgraphInput(int32 InputIndex)
+{
+	TOptional<FValue*> Result = Internal::RequestSubgraphInputValue(BuilderImpl, InputIndex);
+	if (!Result.IsSet())
+	{
+		check(bExpressionRequestsInputsOnDemand);
+		State = EState::Pending;
+		return { FPoison::Get(), nullptr };
+	}
+	return FValueRef{ Result.GetValue(), nullptr };
+}
+
+FValueRef FEmitter::SubgraphInput(int32 InputIndex)
+{
+	FValueRef Value = TrySubgraphInput(InputIndex);
+	if (!Value)
+	{
+		Errorf(TEXT("Missing subgraph input %d."), InputIndex);
+		return Value.ToPoison();
+	}
+	return Value;
+}
+
+UMaterialExpression* FEmitter::GetSubgraphExpression()
+{
+	return Internal::GetSubgraphExpression(BuilderImpl);
+}
+
 /*------------------------------- Constants emission -------------------------------*/
 
 FValueRef FEmitter::ConstantFromShaderValue(const UE::Shader::FValue& InValue)
diff --git a/Engine/Source/Runtime/Engine/Private/Materials/MaterialIRInternal.h b/Engine/Source/Runtime/Engine/Private/Materials/MaterialIRInternal.h
--- a/Engine/Source/Runtime/Engine/Private/Materials/MaterialIRInternal.h
+++ b/Engine/Source/Runtime/Engine/Private/Materials/MaterialIRInternal.h
@@ -11,21 +11,26 @@ namespace MIR::Internal {
 // Returns the material value type of the specified UTexture or URuntimeVirtualTexture.
 EMaterialValueType GetTextureMaterialValueType(const UObject* TextureObject);
 
-// @massimo.tristano expose an internal "Context" struct between the builder and the emitter, and move these methods
-// to the context itself
+// Resolves the value for a given expression input, following through call scope boundaries.
+// Returns an empty TOptional if a dependency was pushed (caller should defer).
+// Returns TOptional with the value (possibly nullptr for disconnected inputs) when available.
+TOptional<FValue*> RequestInputValue(FMaterialIRModuleBuilderImpl* Builder, const FExpressionInput* Input);
 
-// Requests the value for the given input, pushing its source expression as a dependency if not yet built.
-// Returns true if the value is already available, false if a new dependency was pushed.
-bool RequestInputValue(FMaterialIRModuleBuilderImpl* Builder, const FExpressionInput* Input);
+// Flows a value into a given expression output.
+void BindValueToExpressionOutput(FMaterialIRModuleBuilderImpl* Builder, const FExpressionOutput* Output, FValue* Value);
 
-// Returns the value flowing into given expression input (previously set through `BindValueToExpressionInput`).
-FValue* FetchValueFromExpressionInput(FMaterialIRModuleBuilderImpl* Builder, const FExpressionInput* Input);
+// Registers the current expression as a subgraph with its own scope.
+// NumInputs is the number of input expressions in this subgraph.
+// Outputs[i] is the inner expression that produces output i.
+void DeclareSubgraph(FMaterialIRModuleBuilderImpl* Builder, int32 NumInputs, TConstArrayView<FSubgraphOutputMapping> Outputs);
 
-// Flows a value into given expression input.
-void BindValueToExpressionInput(FMaterialIRModuleBuilderImpl* Builder, const FExpressionInput* Input, FValue* Value);
+// Fetches the value from the subgraph expression's input pin at the given index, resolving in the parent scope.
+// Returns an empty TOptional if a dependency was pushed (caller should defer).
+// Returns nullptr if not inside a subgraph scope or if the index is invalid.
+TOptional<FValue*> RequestSubgraphInputValue(FMaterialIRModuleBuilderImpl* Builder, int32 InputIndex);
 
-// Flows a value int given expression output.
-void BindValueToExpressionOutput(FMaterialIRModuleBuilderImpl* Builder, const FExpressionOutput* Output, FValue* Value);
+// Returns the expression that declared the current subgraph scope, or null if at root level.
+UMaterialExpression* GetSubgraphExpression(FMaterialIRModuleBuilderImpl* Builder);
 
 } // namespace MIR::Internal
 
diff --git a/Engine/Source/Runtime/Engine/Private/Materials/MaterialIRModuleBuilder.cpp b/Engine/Source/Runtime/Engine/Private/Materials/MaterialIRModuleBuilder.cpp
--- a/Engine/Source/Runtime/Engine/Private/Materials/MaterialIRModuleBuilder.cpp
+++ b/Engine/Source/Runtime/Engine/Private/Materials/MaterialIRModuleBuilder.cpp
@@ -14,9 +14,7 @@
 #include "Materials/MaterialAggregate.h"
 #include "Materials/MaterialAttributeDefinitionMap.h"
 #include "Materials/MaterialExpressionAbsorptionMediumMaterialOutput.h"
-#include "Materials/MaterialExpressionFunctionInput.h"
 #include "Materials/MaterialExpressionFunctionOutput.h"
-#include "Materials/MaterialExpressionMaterialFunctionCall.h"
 #include "Materials/MaterialExpressionNamedReroute.h"
 #include "Materials/MaterialExpressionSingleLayerWaterMaterialOutput.h"
 #include "Materials/MaterialExpressionThinTranslucentMaterialOutput.h"
@@ -54,50 +52,94 @@ static inline bool IsMaterialPropertyCustomizedUV(EMaterialProperty Property)
 	return Property >= MP_CustomizedUVs0 && Property <= MP_LastCustomizedUVs;
 }
 
-struct FAnalysisContext
+// Identifies a subgraph scope. The (SubgraphExpression, Parent) pair uniquely identifies
+// every instantiation, including nested and duplicate subgraphs.
+struct FSubgraphContext
 {
-	UMaterialExpressionMaterialFunctionCall* Call{};
-	TSet<UMaterialExpression*> BuiltExpressions{};
-	TArray<UMaterialExpression*> ExpressionStack{};
-	TMap<const FExpressionInput*, MIR::FValue*> InputValues;
-	TMap<const FExpressionOutput*, MIR::FValue*> OutputValues;
+	// The expression that declared this subgraph scope (null for root).
+	UMaterialExpression* SubgraphExpression{};
 
-	MIR::FValue* GetInputValue(const FExpressionInput* Input)
-	{
-		MIR::FValue** Value = InputValues.Find(Input);
-		return Value ? *Value : nullptr;
-	}
+	// Parent scope from which this subgraph was instantiated (null for root).
+	FSubgraphContext* Parent{};
 
-	void SetInputValue(const FExpressionInput* Input, MIR::FValue* Value)
-	{
-		InputValues.Add(Input, Value);
-	}
+	// Number of inputs.
+	int32 NumInputs;
 
-	MIR::FValue* GetOutputValue(const FExpressionOutput* Output)
-	{
-		MIR::FValue** Value = OutputValues.Find(Output);
-		return Value ? *Value : nullptr;
-	}
-	
-	void SetOutputValue(const FExpressionOutput* Output, MIR::FValue* Value)
+	// One bit per input, tracks which parent-side connection insights have already been pushed.
+	uint64* InputInsightsBitField;
+
+	// Length of 'Outputs[]'.
+	int32 NumOutputs;
+
+	// The expressions and their output indices generating this subgraph outputs.
+	MIR::FSubgraphOutputMapping Outputs[];
+
+	// Allocates a context into the given allocator.
+	static FSubgraphContext* Create(FMemStackBase& Allocator, UMaterialExpression* InSubgraphExpression, FSubgraphContext* InParent, int32 NumInputs, TConstArrayView<MIR::FSubgraphOutputMapping> Outputs)
 	{
-		OutputValues.Add(Output, Value);
+		const int32 NumInputBitFieldPages = (NumInputs + 63) / 64;
+
+		FSubgraphContext* Ctx = (FSubgraphContext*) Allocator.Alloc(sizeof(FSubgraphContext) + Outputs.NumBytes(), alignof(FSubgraphContext));
+		Ctx->SubgraphExpression = InSubgraphExpression;
+		Ctx->Parent = InParent;
+		Ctx->NumInputs = NumInputs;
+		Ctx->InputInsightsBitField = NumInputs ? (uint64*) Allocator.Alloc(sizeof(uint64) * NumInputBitFieldPages, alignof(uint64)) : nullptr;
+		FMemory::Memzero(Ctx->InputInsightsBitField, sizeof(uint64) * NumInputBitFieldPages);
+		Ctx->NumOutputs = Outputs.Num();
+		FMemory::Memcpy(Ctx->Outputs, Outputs.GetData(), Outputs.NumBytes());
+		return Ctx;
 	}
 };
 
+// Main implementation struct for the MIR module builder.
 struct FMaterialIRModuleBuilderImpl
 {
+	// Public builder interface that owns this implementation.
 	FMaterialIRModuleBuilder* Builder;
+
+	// Target module being constructed.
 	FMaterialIRModule* Module;
+
+	// IR value emitter used by expression Build() methods.
 	MIR::FEmitter Emitter;
+
+	// Preshader manager for CPU-side uniform expressions.
 	FMaterialIRPreshader Preshader;
-	TArray<FAnalysisContext> AnalysisContextStack;
+
+	// Semantic analyzer for the IR value graph.
 	FMaterialIRValueAnalyzer ValueAnalyzer;
+
+	// Default zero-valued material attributes aggregate.
 	MIR::FValue* DefaultMaterialAggregate;
+
+	// Preview expression input for the material editor preview.
 	FColorMaterialInput PreviewInput;
 
+	// Bump allocator for FSubgraphContext instances (stable pointers, bulk deallocation).
+	FMemStackBase ContextAllocator;
+
+	// Root context (no subgraph, no parent), used for top-level expressions.
+	FSubgraphContext* RootContext{};
+
+	// Context currently being processed.
+	FSubgraphContext* CurrentContext{};
+
+	// Expression build stack, entries from any context interleave naturally.
+	TArray<TPair<UMaterialExpression*, FSubgraphContext*>> ExpressionStack;
+
+	// Set of expressions already built in their respective contexts.
+	TSet<TPair<UMaterialExpression*, FSubgraphContext*>> BuiltExpressions;
+
+	// Maps expression outputs to their produced IR values, keyed by context.
+	TMap<TPair<const FExpressionOutput*, FSubgraphContext*>, MIR::FValue*> OutputValues;
+
+	// Maps (SubgraphExpression, ParentContext) to the subgraph context for that instantiation.
+	TMap<TPair<UMaterialExpression*, FSubgraphContext*>, FSubgraphContext*> SubgraphContexts;
+
+	// Substrate material system translator data.
 	Substrate::FSubstrateTranslatorData SubstrateTranslatorData;
 
+	// Sets up everything needed before expression graph traversal can begin.
 	void Step_Initialize(FMaterialIRModuleBuilder* InBuilder, FMaterialIRModule* InModule)
 	{
 		this->Builder = InBuilder;
@@ -134,8 +176,9 @@ struct FMaterialIRModuleBuilderImpl
 		// Create an IR value to hold the material attributes aggregate default
 		DefaultMaterialAggregate = Emitter.Aggregate(MaterialAttributesAggregate::Get());
 
-		// Push the root analysis context
-		AnalysisContextStack.Emplace();
+		// Create the root subgraph context
+		RootContext = FSubgraphContext::Create(ContextAllocator, nullptr, nullptr, 0, {});
+		CurrentContext = RootContext;
 
 		// Set the preview input expression
 		PreviewInput.Expression = Builder->PreviewExpression;
@@ -172,29 +215,40 @@ struct FMaterialIRModuleBuilderImpl
 		SubstrateTranslatorData.ProcessSubstrateTreeAndTopology(MaterialShadingModels, MaterialGraphNodePreviewExpression, bDoClearFunctionStack);
 	}
 
-	// This function returns whether a given property is active in the translating material.
-	// Some properties may be deactivated in the material, such as CustomizedUVx properties, or properties that are inactive for the material's blend mode.
-	bool PropertyIsActiveInMaterial(const EMaterialProperty Property)
+	// Main build loop, converts the material expression graph into IR values.
+	bool Step_BuildMaterialExpressionsToIRGraph()
 	{
-		if (IsMaterialPropertyCustomizedUV(Property) && Property >= MP_CustomizedUVs0 + Emitter.BaseMaterial->NumCustomizedUVs)
-		{
-			return false;
-		}
-		if (!Emitter.MaterialInterface->IsPropertyActive(Property))
+		// Identify the material property output pins and push their value-producing
+		// expressions onto the expression stack to start crawling from them.
+		PushRootExpressionDependencies();
+
+		do
 		{
-			return false;
+			while (!ExpressionStack.IsEmpty())
+			{
+				FMemMark MemMark { FMemStack::Get() };
+				BuildTopMaterialExpression();
+			}
+
+			// Re-seeds root dependencies after each drain because DeclareSubgraph may have registered
+			// new subgraph contexts that change how root inputs resolve (two-pass resolution).
+			CurrentContext = RootContext;
+			PushRootExpressionDependencies();
 		}
-		return true;
-	}
+		while (Module->IsValid() && !ExpressionStack.IsEmpty());
 
-	void Step_PushRootExpressionDependencies()
+		return Module->IsValid();
+	}
+	
+	// Seeds the expression stack with the expressions linked to final material outputs.
+	void PushRootExpressionDependencies()
 	{
 		// If we are processing from the Material Attributes output node, push that dependency here.
 		if (Builder->MaterialInterface->GetMaterial()->bUseMaterialAttributes)
 		{
 			FMaterialInputDescription Input;
 			GetExpressionInputDescription(MP_MaterialAttributes, Input);
-			PushDependencyOnStack(Input.Input->Expression);
+			RequestInputValue(Input.Input);
 		}
 		else
 		{
@@ -212,7 +266,7 @@ struct FMaterialIRModuleBuilderImpl
 				GetExpressionInputDescription(Property, Input);
 
 				// Push the connected expression to this material attribute input as a dependency, if any.
-				PushDependencyOnStack(Input.Input->Expression);
+				RequestInputValue(Input.Input);
 			}
 		}
 
@@ -229,67 +283,33 @@ struct FMaterialIRModuleBuilderImpl
 		}
 	}
 
-	bool Step_BuildMaterialExpressionsToIRGraph()
-	{
-		while (true)
-		{
-			FMemMark MemMark { FMemStack::Get() };
-			FAnalysisContext& Context = AnalysisContextStack.Last();
-
-			if (!Context.ExpressionStack.IsEmpty())
-			{
-				// Some expression is on the expression stack of this context. Analyze it. This will
-				// have the effect of either building the expression or pushing its other expression
-				// dependencies onto the stack.
-				BuildTopMaterialExpression();
-			}
-			else if (Context.Call)
-			{
-				// There are no more expressions to analyze on the stack, this analysis context is complete.
-				// Context.Call isn't null so this context is for a function call, which has now been fully analyzed.
-				// Pop the callee context from the stack and resume analyzing the parent context (the caller).
-				PopFunctionCall();
-			}
-			else
-			{
-				// No other expressions on the stack to evaluate, nor this is a function
-				// call context but the root context. Nothing left to do so simply quit.
-				break;
-			}
-		}
-
-		return Module->IsValid();
-	}
-
+	// Produces the Substrate front material custom output, then re-processes BSDFs for the fully-simplified pass.
 	bool Step_BuildSubstrateFrontMaterialOutput()
 	{
 		if (SubstrateTranslatorData.FrontMaterialExpr)
 		{
 			// Save the the FrontMaterial to a custom output.
-			MIR::FValueRef SubstrateFrontMaterial = MIR::FValueRef(AnalysisContextStack.Last().GetOutputValue(SubstrateTranslatorData.FrontMaterialExpr->GetOutput(0))); // SUBSTRATE_TODO always assuming output 0 for now but this can be wrong for some material function
+			MIR::FValue** SubstrateValuePtr = OutputValues.Find( { SubstrateTranslatorData.FrontMaterialExpr->GetOutput(0), RootContext }); // SUBSTRATE_TODO always assuming output 0 for now but this can be wrong for some material function
+			MIR::FValueRef SubstrateFrontMaterial = MIR::FValueRef(SubstrateValuePtr ? *SubstrateValuePtr : nullptr);
 			Emitter.SetCustomOutputs(TEXTVIEW("SubstrateFrontMaterial"), { &SubstrateFrontMaterial, 1 }, MIR::EMaterialOutputFrequency::PerPixel);
 		}
 
-		// Remove all substrate nodes and function calls from all the contexts.
+		// Un-mark Substrate BSDF nodes so they can be re-processed for the fully simplified pass.
 		bool bNeedSubstrateExpressionProcessing = false;
-		for (FAnalysisContext& Context : AnalysisContextStack)
 		{
-			// Re add expression we need to process them again
-			for (UMaterialExpression* BuiltExpression : Context.BuiltExpressions)
+			// Collect entries to un-mark (can't modify TSet while iterating)
+			TArray<TPair<UMaterialExpression*, FSubgraphContext*>> ToUnmark;
+			for (const TPair<UMaterialExpression*, FSubgraphContext*>& Entry : BuiltExpressions)
 			{
-				if (BuiltExpression->IsA(UMaterialExpressionSubstrateBSDF::StaticClass()))
+				if (Entry.Key->IsA(UMaterialExpressionSubstrateBSDF::StaticClass()))
 				{
-					Context.ExpressionStack.Push(BuiltExpression);
+					ToUnmark.Add(Entry);
 					bNeedSubstrateExpressionProcessing = SubstrateTranslatorData.FrontMaterialExpr != nullptr;
 				}
 			}
-			// Remove them from built
-			for (UMaterialExpression* BuiltExpression : Context.ExpressionStack)
+			for (const TPair<UMaterialExpression*, FSubgraphContext*>& Entry : ToUnmark)
 			{
-				if (BuiltExpression->IsA(UMaterialExpressionSubstrateBSDF::StaticClass()))
-				{
-					Context.BuiltExpressions.Remove(BuiltExpression);
-				}
+				BuiltExpressions.Remove(Entry);
 			}
 		}
 
@@ -298,41 +318,22 @@ struct FMaterialIRModuleBuilderImpl
 		{
 			SubstrateTranslatorData.CurrentSubstrateCompilationContext = Substrate::ESubstrateCompilationContext::SCC_FullySimplified;
 
-			FAnalysisContext& FirstContext = AnalysisContextStack.Last();
-			FirstContext.ExpressionStack.Empty();
-			FirstContext.ExpressionStack.Push(SubstrateTranslatorData.FrontMaterialExpr);
+			// Clear the expression stack and push the front material expression in the root context.
+			ExpressionStack.Empty();
+			ExpressionStack.Push( { SubstrateTranslatorData.FrontMaterialExpr, RootContext });
+			CurrentContext = RootContext;
 
-			while (true)
+			while (!ExpressionStack.IsEmpty())
 			{
-				FMemMark MemMark{ FMemStack::Get() };
-				FAnalysisContext& Context = AnalysisContextStack.Last();
-
-				if (!Context.ExpressionStack.IsEmpty())
-				{
-					// Some expression is on the expression stack of this context. Analyze it. This will
-					// have the effect of either building the expression or pushing its other expression
-					// dependencies onto the stack.
-					BuildTopMaterialExpression();
-				}
-				else if (Context.Call)
-				{
-					// There are no more expressions to analyze on the stack, this analysis context is complete.
-					// Context.Call isn't null so this context is for a function call, which has now been fully analyzed.
-					// Pop the callee context from the stack and resume analyzing the parent context (the caller).
-					PopFunctionCall();
-				}
-				else
-				{
-					// No other expressions on the stack to evaluate, nor this is a function
-					// call context but the root context. Nothing left to do so simply quit.
-					break;
-				}
+				FMemMark MemMark { FMemStack::Get() };
+				BuildTopMaterialExpression();
 			}
 
 			if (SubstrateTranslatorData.CurrentSubstrateCompilationContext == Substrate::ESubstrateCompilationContext::SCC_FullySimplified)
 			{
 				// Save the output to the FrontMaterial, now fully simplified, to the custom output representing fully simplified.
-				MIR::FValueRef SubstrateFrontMaterial = MIR::FValueRef(AnalysisContextStack.Last().GetOutputValue(SubstrateTranslatorData.FrontMaterialExpr->GetOutput(0)));
+				MIR::FValue** SubstrateSimplifiedPtr = OutputValues.Find( { SubstrateTranslatorData.FrontMaterialExpr->GetOutput(0), RootContext });
+				MIR::FValueRef SubstrateFrontMaterial = MIR::FValueRef(SubstrateSimplifiedPtr ? *SubstrateSimplifiedPtr : nullptr);
 				Emitter.SetCustomOutputs(TEXTVIEW("SubstrateFrontMaterialFullySimplified"), { &SubstrateFrontMaterial, 1 }, MIR::EMaterialOutputFrequency::PerPixel);
 
 				SubstrateTranslatorData.CurrentSubstrateCompilationContext = Substrate::ESubstrateCompilationContext::SCC_Default;
@@ -346,17 +347,19 @@ struct FMaterialIRModuleBuilderImpl
 		return Module->IsValid();
 	}
 
+	// Builds one expression from the top of the stack, or defers if its dependencies aren't ready.
 	void BuildTopMaterialExpression()
 	{
 		static const FName OnDemandInputRequestName(TEXT("MIR_OnDemandInputRequest"));
 
-		FAnalysisContext& CurrContext = AnalysisContextStack.Last();
-		Emitter.Expression = CurrContext.ExpressionStack.Last();
-	
-		// If expression is clean, nothing to be done.
-		if (CurrContext.BuiltExpressions.Contains(Emitter.Expression))
+		const auto& Top = ExpressionStack.Last();
+		Emitter.Expression = Top.Key;
+		CurrentContext = Top.Value;
+
+		// If expression is already built in this context, nothing to be done.
+		if (BuiltExpressions.Contains( { Emitter.Expression, CurrentContext }))
 		{
-			CurrContext.ExpressionStack.Pop(EAllowShrinking::No);
+			ExpressionStack.Pop(EAllowShrinking::No);
 			return;
 		}
 
@@ -367,63 +370,49 @@ struct FMaterialIRModuleBuilderImpl
 		// Auto-request all inputs unless the expression opts out via MIR_OnDemandInputRequest meta flag.
 		if (!Emitter.bExpressionRequestsInputsOnDemand)
 		{
-			for (FExpressionInputIterator It{ Emitter.Expression }; It; ++It)
+			for (FExpressionInputIterator It { Emitter.Expression }; It; ++It)
 			{
-				(void)::MIR::Internal::RequestInputValue(this, It.Input);
+				RequestInputValue(It.Input);
 			}
 
-			// If top expression on the stack isn't this, some new dependency was pushed on it.
-			if (CurrContext.ExpressionStack.Last() != Emitter.Expression)
+			// If top expression on the stack changed, some new dependency was pushed. Defer.
+			if (ExpressionStack.Last().Key != Emitter.Expression || ExpressionStack.Last().Value != CurrentContext)
 			{
 				return;
 			}
 		}
-		
-		// @massimo.tristano The call mechanism should be moved to the emitter itself by pushing a "context".
-		if (auto Call = Cast<UMaterialExpressionMaterialFunctionCall>(Emitter.Expression))
-		{
-			// Pop and mark built before PushFunctionCall, which may Emplace onto AnalysisContextStack
-			// and invalidate the CurrContext reference.
-			CurrContext.ExpressionStack.Pop();
-			CurrContext.BuiltExpressions.Add(Emitter.Expression);
 
-			// Function calls are handled internally as they manipulate the analysis context stack.
-			PushFunctionCall(Call);
-		}
-		else
-		{
-			UpdateSubstrateNodeIdentifierStack(CurrContext);
+		UpdateSubstrateNodeIdentifierStack();
 
-			// Invoke the expression build function. This will perform semantic analysis, error reporting and
-			// emit IR values for its outputs (which will flow into connected expressions inputs).
-			Emitter.Expression->Build(Emitter);
+		// Invoke the expression build function. This will perform semantic analysis, error reporting and
+		// emit IR values for its outputs (which will flow into connected expressions inputs).
+		Emitter.Expression->Build(Emitter);
 
-			// If the expression has a pending dependency, defer building it.
-			if (Emitter.State == MIR::FEmitter::EState::Pending)
-			{
-				return;
-			}
+		// If the expression has a pending dependency, defer building it.
+		if (Emitter.State == MIR::FEmitter::EState::Pending)
+		{
+			return;
+		}
 
-			// Populate the insight information about this expression pins.
-			AddExpressionConnectionInsights(Emitter.Expression);
+		// Populate the insight information about this expression pins.
+		PushExpressionConnectionInsights(Emitter.Expression);
 
-			// Take the top expression out of the stack as ready for analysis. Also mark it as built.
-			CurrContext.ExpressionStack.Pop();
-			CurrContext.BuiltExpressions.Add(Emitter.Expression);
-		}
+		// Take the top expression out of the stack as ready for analysis. Also mark it as built.
+		ExpressionStack.Pop();
+		BuiltExpressions.Add( { Emitter.Expression, CurrentContext });
 	}
 
-	void UpdateSubstrateNodeIdentifierStack(FAnalysisContext& CurrContext)
+	// Maintains the Substrate node identity chain needed for unique BSDF GUIDs.
+	void UpdateSubstrateNodeIdentifierStack()
 	{
 		// Build the Substrate expressions unique GUID as if we had a SubstrateTreeStack.
 		// It needs to be done each time because expression are appended breath first while we need position Top-1 to be the parent
-		if (auto NewSubstrateExpr = Cast<UMaterialExpressionSubstrateBSDF>(Emitter.Expression))
+		if (UMaterialExpressionSubstrateBSDF* NewSubstrateExpr = Cast<UMaterialExpressionSubstrateBSDF>(Emitter.Expression))
 		{
 			Substrate::FSubstrateCompilationContext& SubstrateCtx = SubstrateTranslatorData.SubstrateCompilationContext[Substrate::ESubstrateCompilationContext::SCC_Default];
 			SubstrateCtx.SubstrateNodeIdentifierStack.Empty();
 
-			FAnalysisContext& SearchContext = CurrContext;
-			int32 SearchContextParentIndex = AnalysisContextStack.Num() - 2;
+			FSubgraphContext* SearchCtx = CurrentContext;
 
 			UMaterialExpression* TrackedNode = NewSubstrateExpr;
 			struct FWalkedTree
@@ -432,93 +421,98 @@ struct FMaterialIRModuleBuilderImpl
 				uint32 InputIndex = 0;
 			};
 			TArray<FWalkedTree> WalkedNodeStack;
+			// Walk the unified expression stack, filtering entries by the current search context.
 			// If WalkedNodeStack is empty then we are only dealing with a single BSDF, otherwise we chain all operator from the BSDF to the root
-			for (int32 Index = SearchContext.ExpressionStack.Num() - 1; Index >= 0; Index--)
+			for (int32 Index = ExpressionStack.Num() - 1; Index >= 0; Index--)
 			{
-				if (auto SubstrateHorizontalExpr = Cast<UMaterialExpressionSubstrateHorizontalMixing>(SearchContext.ExpressionStack[Index]))
+				UMaterialExpression* Expression = ExpressionStack[Index].Key;
+				FSubgraphContext* Context = ExpressionStack[Index].Value;
+				if (Context != SearchCtx)
+				{
+					continue;
+				}
+
+				if (UMaterialExpressionSubstrateHorizontalMixing* SubstrateHorizontalExpr = Cast<UMaterialExpressionSubstrateHorizontalMixing>(Expression))
 				{
 					if (SubstrateHorizontalExpr->Background.GetTracedInput().Expression == TrackedNode)
 					{
-						WalkedNodeStack.Push({ SubstrateHorizontalExpr , 0 });
+						WalkedNodeStack.Push( { SubstrateHorizontalExpr , 0 });
 						TrackedNode = SubstrateHorizontalExpr;
 					}
 					else if (SubstrateHorizontalExpr->Foreground.GetTracedInput().Expression == TrackedNode)
 					{
-						WalkedNodeStack.Push({ SubstrateHorizontalExpr , 1 });
+						WalkedNodeStack.Push( { SubstrateHorizontalExpr , 1 });
 						TrackedNode = SubstrateHorizontalExpr;
 					}
 					// Else, cannot error since it could be another mixing from another branch we never walk from there
 				}
-				else if (auto SubstrateDecalExpr = Cast<UMaterialExpressionSubstrateConvertToDecal>(SearchContext.ExpressionStack[Index]))
+				else if (UMaterialExpressionSubstrateConvertToDecal* SubstrateDecalExpr = Cast<UMaterialExpressionSubstrateConvertToDecal>(Expression))
 				{
 					if (SubstrateDecalExpr->DecalMaterial.GetTracedInput().Expression == TrackedNode)
 					{
-						WalkedNodeStack.Push({ SubstrateDecalExpr , 0 });
+						WalkedNodeStack.Push( { SubstrateDecalExpr , 0 });
 						TrackedNode = SubstrateDecalExpr;
 					}
 				}
-				else if (auto SubstrateVerticalExpr = Cast<UMaterialExpressionSubstrateVerticalLayering>(SearchContext.ExpressionStack[Index]))
+				else if (UMaterialExpressionSubstrateVerticalLayering* SubstrateVerticalExpr = Cast<UMaterialExpressionSubstrateVerticalLayering>(Expression))
 				{
 					if (SubstrateVerticalExpr->Top.GetTracedInput().Expression == TrackedNode)
 					{
-						WalkedNodeStack.Push({ SubstrateVerticalExpr , 0 });
+						WalkedNodeStack.Push( { SubstrateVerticalExpr , 0 });
 						TrackedNode = SubstrateVerticalExpr;
 					}
 					else if (SubstrateVerticalExpr->Base.GetTracedInput().Expression == TrackedNode)
 					{
-						WalkedNodeStack.Push({ SubstrateVerticalExpr , 1 });
+						WalkedNodeStack.Push( { SubstrateVerticalExpr , 1 });
 						TrackedNode = SubstrateVerticalExpr;
 					}
 				}
-				else if (auto SubstrateAddExpr = Cast<UMaterialExpressionSubstrateAdd>(SearchContext.ExpressionStack[Index]))
+				else if (UMaterialExpressionSubstrateAdd* SubstrateAddExpr = Cast<UMaterialExpressionSubstrateAdd>(Expression))
 				{
 					if (SubstrateAddExpr->A.GetTracedInput().Expression == TrackedNode)
 					{
-						WalkedNodeStack.Push({ SubstrateAddExpr , 0 });
+						WalkedNodeStack.Push( { SubstrateAddExpr , 0 });
 						TrackedNode = SubstrateAddExpr;
 					}
 					else if (SubstrateAddExpr->B.GetTracedInput().Expression == TrackedNode)
 					{
-						WalkedNodeStack.Push({ SubstrateAddExpr , 1 });
+						WalkedNodeStack.Push( { SubstrateAddExpr , 1 });
 						TrackedNode = SubstrateAddExpr;
 					}
 				}
-				else if (auto SubstrateWeightExpr = Cast<UMaterialExpressionSubstrateWeight>(SearchContext.ExpressionStack[Index]))
+				else if (UMaterialExpressionSubstrateWeight* SubstrateWeightExpr = Cast<UMaterialExpressionSubstrateWeight>(Expression))
 				{
 					if (SubstrateWeightExpr->A.GetTracedInput().Expression == TrackedNode)
 					{
-						WalkedNodeStack.Push({ SubstrateWeightExpr , 0 });
+						WalkedNodeStack.Push( { SubstrateWeightExpr , 0 });
 						TrackedNode = SubstrateWeightExpr;
 					}
 				}
-				else if (auto SubstrateSelectExpr = Cast<UMaterialExpressionSubstrateSelect>(SearchContext.ExpressionStack[Index]))
+				else if (UMaterialExpressionSubstrateSelect* SubstrateSelectExpr = Cast<UMaterialExpressionSubstrateSelect>(Expression))
 				{
 					if (SubstrateSelectExpr->A.GetTracedInput().Expression == TrackedNode)
 					{
-						WalkedNodeStack.Push({ SubstrateSelectExpr , 0 });
+						WalkedNodeStack.Push( { SubstrateSelectExpr , 0 });
 						TrackedNode = SubstrateSelectExpr;
 					}
 					else if (SubstrateSelectExpr->B.GetTracedInput().Expression == TrackedNode)
 					{
-						WalkedNodeStack.Push({ SubstrateSelectExpr , 1 });
+						WalkedNodeStack.Push( { SubstrateSelectExpr , 1 });
 						TrackedNode = SubstrateSelectExpr;
 					}
 				}
-				else if (auto FunctionOutput = Cast<UMaterialExpressionFunctionOutput>(SearchContext.ExpressionStack[Index]))
+				else if (UMaterialExpressionFunctionOutput* FunctionOutput = Cast<UMaterialExpressionFunctionOutput>(Expression))
 				{
 					// We need to select another context
-					if (SearchContext.Call == nullptr)
+					if (SearchCtx->SubgraphExpression == nullptr)
 					{
 						break; // we are done
 					}
 					else
 					{
-						// Update the tracked node to the output
-						TrackedNode = SearchContext.Call;
-
-						SearchContext = AnalysisContextStack[SearchContextParentIndex];
-						SearchContextParentIndex--;
-						Index = SearchContext.ExpressionStack.Num() - 1;
+						// Update the tracked node to the subgraph expression and search in the parent context
+						TrackedNode = SearchCtx->SubgraphExpression;
+						SearchCtx = SearchCtx->Parent;
 					}
 				}
 			}
@@ -529,7 +523,6 @@ struct FMaterialIRModuleBuilderImpl
 			{
 				FWalkedTree& SubstrateExpr = WalkedNodeStack[Index];
 
-
 				// Create an md5 hash for the parent, its input pin index and current node to represent the path.
 				uint32 IntputHashBuffer[9];
 				FGuid PreviousNodeGuid = SubstrateCtx.SubstrateNodeIdentifierStack.Top();
@@ -544,7 +537,7 @@ struct FMaterialIRModuleBuilderImpl
 				IntputHashBuffer[7] = NodeGuid.C;
 				IntputHashBuffer[8] = NodeGuid.D;
 
-				uint32 OutputHashBuffer[]{ 0, 0, 0, 0 };
+				uint32 OutputHashBuffer[] { 0, 0, 0, 0 };
 				FMD5 IdentifierStringHash;
 				IdentifierStringHash.Update((uint8*)IntputHashBuffer, sizeof(IntputHashBuffer));
 				IdentifierStringHash.Final((uint8*)&OutputHashBuffer);
@@ -558,131 +551,113 @@ struct FMaterialIRModuleBuilderImpl
 		}
 	}
 	
-	// Pushes the given expression onto the build stack if it hasn't been built yet.
-	// Returns true if already resolved (null or already built), false if pushed.
-	bool PushDependencyOnStack(UMaterialExpression* Expression)
+	// Resolves the value for a given expression input. If the expression has declared a subgraph
+	// scope (via Subgraph), resolves through to the mapped inner expression in the
+	// subgraph context. Pushes unbuilt dependencies and returns an empty TOptional (caller should
+	// defer). Returns TOptional with the value (possibly nullptr for disconnected inputs) when
+	// available.
+	TOptional<MIR::FValue*> RequestInputValue(const FExpressionInput* Input)
 	{
-		FAnalysisContext& CurrContext = AnalysisContextStack.Last();
-
-		// Ignore disconnected inputs and connected expressions already built.
-		if (Expression && !CurrContext.BuiltExpressions.Contains(Expression))
+		if (!Input->Expression)
 		{
-			CurrContext.ExpressionStack.Push(Expression);
-			return false;
+			return nullptr;
 		}
 
-		return true;
-	}
-
-	// Flows the value from connected output to specified input.
-	void FlowValueThroughConnection(const FExpressionInput* Input)
-	{
-		FAnalysisContext& CurrContext = AnalysisContextStack.Last();
+		UMaterialExpression* Expression = Input->Expression;
+		FExpressionOutput* Output = Input->GetConnectedOutput();
+		FSubgraphContext* Context = CurrentContext;
 
-		// Fetch the value flowing through connected output.
-		FExpressionOutput* ConnectedOutput = Input->GetConnectedOutput();
-		if (ConnectedOutput)
+		// If this expression has a registered subgraph scope, resolve to the mapped inner expression.
+		FSubgraphContext* SubgraphCtx;
+		if (MIR::Find(SubgraphContexts, { Expression, Context }, SubgraphCtx))
 		{
-			MIR::FValue** ValuePtr = CurrContext.OutputValues.Find(ConnectedOutput);
-			if (ValuePtr)
+			if (Input->OutputIndex < SubgraphCtx->NumOutputs)
 			{
-				// ...and flow it into this input.
-				CurrContext.InputValues.Add(Input, *ValuePtr);
-			}
-		}
-	}
-
-	// @massimo.tristano Function call could be optimized by having the caller flag which outputs are actually not used, so that we can skip
-	// pushing those function outputs dependencies
-	void PushFunctionCall(UMaterialExpressionMaterialFunctionCall* Call)
-	{
-		MIR::TTemporaryArray<MIR::FValue*> CallInputValues{ Call->FunctionInputs.Num() };
-
-		// Make sure each required function input is connected and has a value.  If so, cache the values flowing into this
-		// funcion call inside the auxiliary value array.  If the input is optional (bUsePreviewValueAsDefault set), we can
-		// ignore the missing value, and the downstream UMaterialExpressionFunctionInput build function will return the default.
-		for (int i = 0; i < Call->FunctionInputs.Num(); ++i)
-		{
-			FFunctionExpressionInput& FunctionInput = Call->FunctionInputs[i];
-			MIR::FType Type = MIR::FType::FromMaterialValueType(FunctionInput.ExpressionInput->GetInputValueType(0));
-
-			if (FunctionInput.ExpressionInput->bUsePreviewValueAsDefault)
-			{
-				MIR::FValue* Value = Emitter.TryInput(Call->GetInput(i));
-				CallInputValues[i] = Value ? Emitter.Cast(Value, Type) : nullptr;
+				const MIR::FSubgraphOutputMapping& OutputMapping = SubgraphCtx->Outputs[Input->OutputIndex];
+				if (OutputMapping.Expression)
+				{
+					Expression = OutputMapping.Expression;
+					Output     = OutputMapping.Expression->GetOutput(OutputMapping.OutputIndex);
+					Context    = SubgraphCtx;
+				}
 			}
 			else
 			{
-				MIR::FValue* Value = Emitter.Input(Call->GetInput(i));
-				CallInputValues[i] = Emitter.Cast(Value, Type);
+				// Subgraph scope exists but output index is out of range, stale connection.
+				Emitter.Errorf(TEXT("Subgraph expression '%s': invalid output index %d."), *Expression->GetName(), Input->OutputIndex);
 			}
 		}
 
-		// If some error occurred (e.g. some function input wasn't linked in) early out.
-		if (Emitter.IsInvalid())
+		// If the dependency hasn't been built yet, push it and defer.
+		if (!BuiltExpressions.Contains({Expression, Context}))
 		{
-			return;
+			ExpressionStack.Push({Expression, Context});
+			return {};
 		}
 
-		// Push a new analysis context on the stack dedicated to this function call.
-		AnalysisContextStack.Emplace();
-		FAnalysisContext& ParentContext = AnalysisContextStack[AnalysisContextStack.Num() - 2];
-		FAnalysisContext& NewContext = AnalysisContextStack[AnalysisContextStack.Num() - 1];
+		// Read the value directly from the producing expression's output.
+		MIR::FValue** ValuePtr = OutputValues.Find({Output, Context});
+		return ValuePtr ? *ValuePtr : nullptr;
+	}
+	
+	// Registers the current expression as a subgraph with its own scope.
+	void DeclareSubgraph(int32 NumInputs, TConstArrayView<MIR::FSubgraphOutputMapping> Outputs)
+	{
+		FSubgraphContext* Ctx = FSubgraphContext::Create(ContextAllocator, Emitter.Expression, CurrentContext, NumInputs, Outputs);
+		SubgraphContexts.Add({Emitter.Expression, CurrentContext}, Ctx);
+	}
 
-		// Set the function call. When the expressions stack in this new context is empty, this
-		// will be used to wire all values flowing inside the function outputs to the function call outputs.
-		NewContext.Call = Call;
+	// Fetches the value from the subgraph expression's input pin at the given index, resolving in the parent scope.
+	TOptional<MIR::FValue*> RequestSubgraphInputValue(int32 InputIndex)
+	{
+		FSubgraphContext* SubgraphCtx = CurrentContext;
+		if (!SubgraphCtx->SubgraphExpression)
+		{
+			return nullptr; // Not inside a subgraph scope.
+		}
 
-		// Forward values flowing into call inputs to called function inputs
-		for (int i = 0; i < Call->FunctionInputs.Num(); ++i)
+		const FExpressionInput* SubgraphInput = SubgraphCtx->SubgraphExpression->GetInput(InputIndex);
+		if (!SubgraphInput)
 		{
-			if (CallInputValues[i])
-			{
-				FFunctionExpressionInput& FunctionInput = Call->FunctionInputs[i];
+			return nullptr; // Invalid index.
+		}
 
-				// Bind the value flowing into the function call input to the function input
-				// expression (inside the function) in the new context.
-				NewContext.SetOutputValue(FunctionInput.ExpressionInput->GetOutput(0), CallInputValues[i]);
+		// Request the value in the parent scope.
+		CurrentContext = SubgraphCtx->Parent;
+		TOptional<MIR::FValue*> Result = RequestInputValue(SubgraphInput);
+		CurrentContext = SubgraphCtx;
 
-				// Mark the function input as built.
-				NewContext.BuiltExpressions.Add(FunctionInput.ExpressionInput.Get());
+		// Push a parent-side connection insight so the edge into the subgraph expression gets colored.
+		// Only for top-level subgraphs (parent is root), and only once per input (tracked via bitfield).
+		if (SubgraphCtx->Parent == RootContext && Result.IsSet() && Result.GetValue() && SubgraphInput->Expression)
+		{
+			// Test and set the bit for this input to avoid pushing duplicate insights
+			// when multiple inner expressions consume the same subgraph input.
+			uint64& BitFieldPage = SubgraphCtx->InputInsightsBitField[InputIndex / 64];
+			int32 BitFieldPageOffset = InputIndex % 64;
+			if ((BitFieldPage & (1ull << BitFieldPageOffset)) == 0)
+			{
+				BitFieldPage |= (1ull << BitFieldPageOffset);
+				PushConnectionInsight(SubgraphCtx->SubgraphExpression, InputIndex, SubgraphInput->Expression, SubgraphInput->OutputIndex, Result.GetValue()->Type);
 			}
 		}
 
-		// Finally push the function outputs to the expression evaluation stack in the new context.
-		for (FFunctionExpressionOutput& FunctionOutput : Call->FunctionOutputs)
-		{
-			NewContext.ExpressionStack.Push(FunctionOutput.ExpressionOutput.Get());
-		}
+		return Result;
 	}
 
-	void PopFunctionCall()
+	// Pushes the given expression onto the build stack if it hasn't been built yet.
+	// Returns true if already resolved (null or already built), false if pushed.
+	bool PushDependencyOnStack(UMaterialExpression* Expression)
 	{
-		// Pull the values flowing into the function outputs out of the current
-		// context and flow them into the Call outputs in the parent context so that
-		// analysis can continue from the call expression.
-		FAnalysisContext& ParentContext = AnalysisContextStack[AnalysisContextStack.Num() - 2];
-		FAnalysisContext& CurrContext = AnalysisContextStack[AnalysisContextStack.Num() - 1];
-		UMaterialExpressionMaterialFunctionCall* Call = CurrContext.Call;
-
-		for (int i = 0; i < Call->FunctionOutputs.Num(); ++i)
+		if (Expression && !BuiltExpressions.Contains({Expression, CurrentContext}))
 		{
-			FFunctionExpressionOutput& FunctionOutput = Call->FunctionOutputs[i];
-
-			// Get the value flowing into the function output inside the function in the current context.
-			MIR::FValue* Value = Emitter.Input(FunctionOutput.ExpressionOutput->GetInput(0));
-
-			// And flow it to the relative function *call* output in the parent context.
-			ParentContext.SetOutputValue(Call->GetOutput(i), Value);
+			ExpressionStack.Push({Expression, CurrentContext});
+			return false;
 		}
-
-		// Finally pop this context (the function call) to return to the caller.
-		AnalysisContextStack.Pop();
-
-		// Populate the insight information about this expression pins.
-		AddExpressionConnectionInsights(Call);
+		return true;
 	}
 
+	// Bridges the expression graph to the shader, assigns built values to material output properties.
 	bool Step_EmitSetMaterialPropertyInstructions()
 	{
 		// First, if the material is flagged to use the material attributes aggregate, read its value now
@@ -694,8 +669,7 @@ struct FMaterialIRModuleBuilderImpl
 			GetExpressionInputDescription(MP_MaterialAttributes, InputDesc);
 
 			// Fetch the value from the material attributes input.
-			FlowValueThroughConnection(InputDesc.Input);
-			MaterialAttributesValue = MIR::Internal::FetchValueFromExpressionInput(this, InputDesc.Input);
+			MaterialAttributesValue = RequestInputValue(InputDesc.Input).Get(nullptr);
 			if (!MaterialAttributesValue)
 			{
 				MaterialAttributesValue = DefaultMaterialAggregate;
@@ -743,8 +717,7 @@ struct FMaterialIRModuleBuilderImpl
 				else
 				{
 					// Otherwise grab the value from the individual attribute pin.
-					FlowValueThroughConnection(InputDesc.Input);
-					PropertyValue = MIR::Internal::FetchValueFromExpressionInput(this, InputDesc.Input);
+					PropertyValue = RequestInputValue(InputDesc.Input).Get(nullptr);
 				}
 			}
 			
@@ -861,6 +834,7 @@ struct FMaterialIRModuleBuilderImpl
 		return Module->IsValid();
 	}
 
+	// Synthesizes a Substrate slab for non-Substrate materials by converting legacy root properties.
 	MIR::FValue* BuildSubstrateFrontMaterialIRFromRootNodeHiddenConversion()
 	{
 		MIR::FValue* FrontMaterialPropertyValue = Emitter.SubstrateDefaultSlab();
@@ -951,6 +925,7 @@ struct FMaterialIRModuleBuilderImpl
 		return FrontMaterialPropertyValue;
 	}
 
+	// Records Substrate tree metadata so the shader compiler knows what BSDFs and features are in play.
 	void GenerateSubstrateMaterialCompilationOutputAndDefines()
 	{
 		// Write substrate tree specific code.
@@ -999,6 +974,7 @@ struct FMaterialIRModuleBuilderImpl
 		}
 	}
 
+	// Maps a material property to its source expression input, honoring preview expression overrides.
 	bool GetExpressionInputDescription(EMaterialProperty Property, FMaterialInputDescription& Input)
 	{
 		if (PreviewInput.IsConnected() && Property == MP_EmissiveColor)
@@ -1028,6 +1004,7 @@ struct FMaterialIRModuleBuilderImpl
 		}
 	}
 
+	// Semantic pass, propagates liveness, stage assignment, and side-effect info through the IR graph.
 	bool Step_AnalyzeIRGraph()
 	{
 		TArray<MIR::FValue*> ValueStack{};
@@ -1167,6 +1144,7 @@ struct FMaterialIRModuleBuilderImpl
 		}
 	}
 
+	// Post-analysis fixups that require the full IR graph to be finalized (buffer layout, preshader codegen).
 	bool Step_PerformModuleWideAnalysis()
 	{
 		// After the entire tree has been analyzed, we can run fixup, which allocates output buffer offsets and fixes up
@@ -1201,6 +1179,7 @@ struct FMaterialIRModuleBuilderImpl
 		return Module->IsValid();
 	}
 
+	// Prunes environment defines that don't apply to this material's actual configuration.
 	void ConsolidateEnvironmentDefines()
 	{
 		// Keep defines if a combined condition is met. Otherwise, remove them from the environemnt defines set.
@@ -1225,6 +1204,7 @@ struct FMaterialIRModuleBuilderImpl
 		Module->EnvironmentDefines = MoveTemp(ValueAnalyzer.EnvironmentDefines);
 	}
 
+	// Catches environment defines that conflict with the material's blend mode or feature set.
 	void ValidateEnvironmentDefines()
 	{
 		// Match various defines against the material configuration
@@ -1249,6 +1229,7 @@ struct FMaterialIRModuleBuilderImpl
 		}
 	}
 	
+	// Fills in the compilation output that downstream systems (shader compiler, renderer) depend on.
 	void FinishSetupMaterialCompilationOutput()
 	{
 		FMaterialCompilationOutput& CompilationOutput = Module->CompilationOutput;
@@ -1310,6 +1291,7 @@ struct FMaterialIRModuleBuilderImpl
 		}
 	}
 
+	// Linearizes the IR graph into ordered instruction lists, one per entry point.
 	void Step_LinkInstructions()
 	{
 		TArray<MIR::FInstruction*> InstructionStack{};
@@ -1390,12 +1372,14 @@ struct FMaterialIRModuleBuilderImpl
 		}
 	}
 
+	// Last step, produces debug/editor artifacts from the completed module.
 	void Step_FinalizeArtifacts()
 	{
 		// Finalize the material insights debug information
 		GenerateDebugInsights();
 	}
 
+	// Produces the debug insight payload consumed by the material editor for diagnostics display.
 	void GenerateDebugInsights()
 	{
 		if (!Builder->TargetInsights)
@@ -1427,20 +1411,62 @@ struct FMaterialIRModuleBuilderImpl
 	
 	/* Auxiliary functions */
 
+	// This function returns whether a given property is active in the translating material.
+	// Some properties may be deactivated in the material, such as CustomizedUVx properties, or properties that are inactive for the material's blend mode.
+	bool PropertyIsActiveInMaterial(const EMaterialProperty Property)
+	{
+		if (!Emitter.MaterialInterface->IsPropertyActive(Property))
+		{
+			return false;
+		}
+		if (IsMaterialPropertyCustomizedUV(Property) && Property >= MP_CustomizedUVs0 + Emitter.BaseMaterial->NumCustomizedUVs)
+		{
+			return false;
+		}
+		return true;
+	}
+
 	// Adds an expression connection insight to the MaterialInsights instance, if any.
-	void AddExpressionConnectionInsights(UMaterialExpression* Expression)
+	void PushExpressionConnectionInsights(UMaterialExpression* Expression)
 	{
-		if (!Builder->TargetInsights)
+		if (!Builder->TargetInsights || CurrentContext != RootContext)
+		{
+			return;
+		}
+
+		// Subgraph expressions (e.g. function calls) get their input pin insights from
+		// RequestSubgraphInputValue, ignore them here as they haven't been built yet.
+		bool bIsExpressionSubgraph = SubgraphContexts.Contains({ Expression, CurrentContext });
+		if (bIsExpressionSubgraph)
 		{
 			return;
 		}
 
-		// Update expression inputs insight.
+		// Go over all expression inputs and emit the connection insights for those that have a value coming in.
 		for (FExpressionInputIterator It{ Expression }; It; ++It)
 		{
-			if (MIR::FValue* Value = MIR::Internal::FetchValueFromExpressionInput(this, It.Input))
+			FExpressionOutput* Output = It.Input->GetConnectedOutput();
+			if (!Output)
 			{
-				PushConnectionInsight(Expression, It.Index, It->Expression, It->OutputIndex, Value->Type);
+				continue;
+			}
+
+			// Resolve through subgraph output mappings if the connected expression declares a subgraph.
+			FSubgraphContext* Context = CurrentContext;
+			FSubgraphContext* SubgraphCtx;
+			if (MIR::Find(SubgraphContexts, {It->Expression, Context}, SubgraphCtx)
+				&& It->OutputIndex < SubgraphCtx->NumOutputs
+				&& SubgraphCtx->Outputs[It->OutputIndex].Expression)
+			{
+				const MIR::FSubgraphOutputMapping& Mapping = SubgraphCtx->Outputs[It->OutputIndex];
+				Output  = Mapping.Expression->GetOutput(Mapping.OutputIndex);
+				Context = SubgraphCtx;
+			}
+
+			MIR::FValue** ValuePtr = OutputValues.Find({Output, Context});
+			if (ValuePtr)
+			{
+				PushConnectionInsight(Expression, It.Index, It->Expression, It->OutputIndex, (*ValuePtr)->Type);
 			}
 		}
 	}
@@ -1448,7 +1474,7 @@ struct FMaterialIRModuleBuilderImpl
 	// Adds a connection insight to the MaterialInsights instance, if any.
 	void PushConnectionInsight(const UObject* InputObject, int InputIndex, const UMaterialExpression* OutputExpression, int OutputIndex, MIR::FType Type)
 	{
-		if (!Builder->TargetInsights || Type.IsPoison())
+		if (!Builder->TargetInsights || Type.IsPoison() || !OutputExpression)
 		{
 			return;
 		}
@@ -1474,31 +1500,29 @@ struct FMaterialIRModuleBuilderImpl
 
 namespace MIR::Internal {
 
-	bool RequestInputValue(FMaterialIRModuleBuilderImpl* Builder, const FExpressionInput* Input)
+	TOptional<FValue*> RequestInputValue(FMaterialIRModuleBuilderImpl* Builder, const FExpressionInput* Input)
 	{
-		// Safe for disconnected inputs: PushDependencyOnStack returns true for null expressions,
-		// and FlowValueThroughConnection is a no-op when there is no connected output.
-		if (Builder->PushDependencyOnStack(Input->Expression))
-		{
-			Builder->FlowValueThroughConnection(Input);
-			return true;
-		}
-		return false;
+		return Builder->RequestInputValue(Input);
 	}
-
-	FValue* FetchValueFromExpressionInput(FMaterialIRModuleBuilderImpl* Builder, const FExpressionInput* Input)
+	
+	void BindValueToExpressionOutput(FMaterialIRModuleBuilderImpl* Builder, const FExpressionOutput* Output, FValue* Value)
 	{
-		return Builder->AnalysisContextStack.Last().GetInputValue(Input);
+		Builder->OutputValues.Add({Output, Builder->CurrentContext}, Value);
 	}
 
-	void BindValueToExpressionInput(FMaterialIRModuleBuilderImpl* Builder, const FExpressionInput* Input, FValue* Value)
+	void DeclareSubgraph(FMaterialIRModuleBuilderImpl* Builder, int32 NumInputs, TConstArrayView<FSubgraphOutputMapping> Outputs)
 	{
-		Builder->AnalysisContextStack.Last().SetInputValue(Input, Value);
+		Builder->DeclareSubgraph(NumInputs, Outputs);
+	}
+	
+	TOptional<FValue*> RequestSubgraphInputValue(FMaterialIRModuleBuilderImpl* Builder, int32 InputIndex)
+	{
+		return Builder->RequestSubgraphInputValue(InputIndex);
 	}
 
-	void BindValueToExpressionOutput(FMaterialIRModuleBuilderImpl* Builder, const FExpressionOutput* Output, FValue* Value)
+	UMaterialExpression* GetSubgraphExpression(FMaterialIRModuleBuilderImpl* Builder)
 	{
-		Builder->AnalysisContextStack.Last().SetOutputValue(Output, Value);
+		return Builder->CurrentContext->SubgraphExpression;
 	}
 
 } // namespace MIR::Internal
@@ -1512,11 +1536,7 @@ bool FMaterialIRModuleBuilder::Build(FMaterialIRModule* TargetModule)
 	// Initialize the module to a blank slate, initialize the builder auxiliary data
 	// and the emitter for IR values emission.
 	Impl.Step_Initialize(this, TargetModule);
-
-	// Identify the material property output pins and push their value-producing
-	// expressions onto the analysis context stack to start crawling from them.
-	Impl.Step_PushRootExpressionDependencies();
-
+	
 	// Main step. It crawls the expression graph and calls the Build() function on each
 	// visited expression in order to emit the IR values that implement that expression
 	// semantics. At the end of this step the IR values graph has been built, but is still
diff --git a/Engine/Source/Runtime/Engine/Public/Materials/MaterialExpressionFunctionInput.h b/Engine/Source/Runtime/Engine/Public/Materials/MaterialExpressionFunctionInput.h
--- a/Engine/Source/Runtime/Engine/Public/Materials/MaterialExpressionFunctionInput.h
+++ b/Engine/Source/Runtime/Engine/Public/Materials/MaterialExpressionFunctionInput.h
@@ -43,7 +43,7 @@ enum EBlendInputRelevance : int
 	Bottom
 };
 
-UCLASS(hidecategories=object, MinimalAPI)
+UCLASS(hidecategories=object, MinimalAPI, meta=(MIR_OnDemandInputRequest))
 class UMaterialExpressionFunctionInput : public UMaterialExpression
 {
 	GENERATED_UCLASS_BODY()
diff --git a/Engine/Source/Runtime/Engine/Public/Materials/MaterialExpressionMaterialFunctionCall.h b/Engine/Source/Runtime/Engine/Public/Materials/MaterialExpressionMaterialFunctionCall.h
--- a/Engine/Source/Runtime/Engine/Public/Materials/MaterialExpressionMaterialFunctionCall.h
+++ b/Engine/Source/Runtime/Engine/Public/Materials/MaterialExpressionMaterialFunctionCall.h
@@ -76,7 +76,7 @@ struct FFunctionExpressionOutput
 
 };
 
-UCLASS(hidecategories=object, MinimalAPI)
+UCLASS(hidecategories=object, MinimalAPI, meta=(MIR_OnDemandInputRequest))
 class UMaterialExpressionMaterialFunctionCall : public UMaterialExpression
 {
 	GENERATED_UCLASS_BODY()
@@ -120,6 +120,7 @@ class UMaterialExpressionMaterialFunctionCall : public UMaterialExpression
 
 	//~ Begin UMaterialExpression Interface
 #if WITH_EDITOR
+	ENGINE_API virtual void Build(MIR::FEmitter& Emitter) override;
 	ENGINE_API virtual int32 Compile(class FMaterialCompiler* Compiler, int32 OutputIndex) override;
 	ENGINE_API virtual void GetCaption(TArray<FString>& OutCaptions) const override;
 	ENGINE_API virtual TArrayView<FExpressionInput*> GetInputsView() override;
diff --git a/Engine/Source/Runtime/Engine/Public/Materials/MaterialIRCommon.h b/Engine/Source/Runtime/Engine/Public/Materials/MaterialIRCommon.h
--- a/Engine/Source/Runtime/Engine/Public/Materials/MaterialIRCommon.h
+++ b/Engine/Source/Runtime/Engine/Public/Materials/MaterialIRCommon.h
@@ -62,6 +62,7 @@ namespace MIR
 	struct FSetMaterialOutput;
 	struct FType;
 	struct FValue;
+	struct FSubgraphOutputMapping;
 }
 
 namespace MIR
diff --git a/Engine/Source/Runtime/Engine/Public/Materials/MaterialIREmitter.h b/Engine/Source/Runtime/Engine/Public/Materials/MaterialIREmitter.h
--- a/Engine/Source/Runtime/Engine/Public/Materials/MaterialIREmitter.h
+++ b/Engine/Source/Runtime/Engine/Public/Materials/MaterialIREmitter.h
@@ -199,6 +199,17 @@ struct FTextureSampleAttributes
 	bool bIsAdaptive : 1 { false };
 };
 
+//
+struct FSubgraphOutputMapping
+{
+	//
+	UMaterialExpression* Expression;
+	
+	//
+	int32 OutputIndex;
+};
+
+// @massimo.tristano all Substrate should be moved out to a separate module
 struct FSubstrateShadingModelsDesc
 {
 	FValueRef BaseColor;
@@ -352,19 +363,23 @@ public:
 
 	// Retrieves the value from the given input, if present. A null reference otherwise.
 	// Note: it does not report an error if the value is missing.
+	// If the expression is marked MIR_OnDemandInputRequest and the dependency hasn't been
+	// built yet, sets Pending and returns poison, use UE_MIR_CHECKPOINT to defer.
 	// It is safe to pass returned value to other Emitter functions as argument even if null.
 	// For instance, the following is safe and idiomatic
-	//     
+	//
 	//    FValueRef Value = Emitter.CheckIsArithmetic(Emitter.TryInput(&A));
 	//    UE_MIR_CHECKPOINT(Emitter);
 	//    if (Value) { /* handle case when value is provided */ }
-	// 
+	//
 	// Keep in mind that a reference can be null (its Value is null) or invalid (which means
 	// that it's null OR poison).
 	FValueRef TryInput(const FExpressionInput* Input);
 
 	// Retrieves the value flowing into given input.
 	// If the value is missing, it reports an error and returns poison.
+	// If the expression is marked MIR_OnDemandInputRequest and the dependency hasn't been
+	// built yet, sets Pending and returns poison, use UE_MIR_CHECKPOINT to defer.
 	FValueRef Input(const FExpressionInput* Input);
 
 	// If the input has no value flowing in (e.g., because the input is disconnected),
@@ -445,6 +460,29 @@ public:
 
 	// Flows given `Value` out of all outputs in the current expression.
 	FEmitter& OutputsWithComponentMask(FValueRef Value);
+
+	/*------------------------------- Subgraph management ------------------------------*/
+	
+	// Registers the current expression as a subgraph.
+	// - NumInputs is the number of input pins on this expression, used for connection insight deduplication.
+	// - Outputs[i] is the inner expression that produces output i. When something later connects to
+	//   this expression's output, resolution redirects to the mapped inner expression in the subgraph scope.
+	void Subgraph(int32 NumInputs, TArrayView<FSubgraphOutputMapping> Outputs);
+
+	// Retrieves the value from the subgraph expression's input pin at the given index, if present.
+	// A null reference otherwise. Does not report an error if the value is missing.
+	// If the expression is marked MIR_OnDemandInputRequest and the dependency hasn't been
+	// built yet, sets Pending and returns poison, use UE_MIR_CHECKPOINT to defer.
+	FValueRef TrySubgraphInput(int32 InputIndex);
+
+	// Retrieves the value from the subgraph expression's input pin at the given index.
+	// If the value is missing, it reports an error and returns poison.
+	// If the expression is marked MIR_OnDemandInputRequest and the dependency hasn't been
+	// built yet, sets Pending and returns poison, use UE_MIR_CHECKPOINT to defer.
+	FValueRef SubgraphInput(int32 InputIndex);
+
+	// Returns the expression that declared the current subgraph scope, or null if at root level.
+	UMaterialExpression* GetSubgraphExpression();
 	
 	/*------------------------------- Constants emission -------------------------------*/
 
@@ -978,7 +1016,6 @@ private:
 	friend FMaterialIRModuleBuilderImpl;
 };
 
-
 //------------------------------------------------------------------------
 //                        IMPLEMENTATION DETAILS
 //------------------------------------------------------------------------
