FSSS - added power control on the multiple scattering blur amount.

#tests PC editor
#rnx

#Synced-CL 52675696

diff --git a/Engine/Shaders/Private/HeightFogPixelShader.usf b/Engine/Shaders/Private/HeightFogPixelShader.usf
--- a/Engine/Shaders/Private/HeightFogPixelShader.usf
+++ b/Engine/Shaders/Private/HeightFogPixelShader.usf
@@ -223,7 +223,8 @@ void ExponentialPixelMain(
 	const float SceneDepthCentimeters   = SceneDepth;
 	FFogSSSParameters FSSSParams = EvaluateFSSSParameter(FogTransmittanceToScene, SceneDepthCentimeters);
 
-	const float SceneColorToScatterAmount = saturate(FSSSParams.Coverage * FSSSParams.Albedo * FogStruct.FSSSSceneColorScatteringAmountScale); // *(1.0f / (4.0f * PI));		// Approximation, should use eq. 4 from paper?
+	// Amount of scene color to scater, ignoring the phase function. It is assumed the phase function is part of FSSSSceneColorScatteringAmountScale.
+	const float SceneColorToScatterAmount = pow(saturate(FSSSParams.Coverage * FSSSParams.Albedo * FogStruct.FSSSSceneColorScatteringAmountScale), FogStruct.FSSSSceneColorScatteringAmountPower);
 
 	// This is point sampling scene color when using full resolution. At half resolution this is average of 2x2 texels.
 	const float3 SceneColorToScatter = FogTransmittanceToScene * SceneColorToScatterAmount * Texture2DSampleLevel(InputSceneColorTexture, InputSceneColorTextureSampler, TexCoordSceneTex, 0).rgb;
diff --git a/Engine/Source/Runtime/Engine/Classes/Components/ExponentialHeightFogComponent.h b/Engine/Source/Runtime/Engine/Classes/Components/ExponentialHeightFogComponent.h
--- a/Engine/Source/Runtime/Engine/Classes/Components/ExponentialHeightFogComponent.h
+++ b/Engine/Source/Runtime/Engine/Classes/Components/ExponentialHeightFogComponent.h
@@ -201,10 +201,18 @@ class UExponentialHeightFogComponent : public USceneComponent
 	 * The amount of scene color injected into the fog screen space scattering technique.
 	 * When 0, only the fog luminance will be scattered.
 	 * When 1, fog and scene luminance will be scattered.
+	 * When above 1, more scene color can be injected to increase the brightness of the multiple scattering effect, even when the fog is not too intense. We cannot scatter more than the scene color however.
 	 */
 	UPROPERTY(EditAnywhere, BlueprintReadOnly, Category = FogScreenSpaceScattering, meta = (DisplayName = "Scene Color Scattering Amount Scale", UIMin = "0", UIMax = "1", ClampMin = "0"))
 	float FSSSSceneColorScatteringAmountScale;
 
+	/**
+	 * Apply a power over the amount of scattered scene luminance after it is evaluated using Scene Color Scattering Amount Scale.
+	 * A value smaller than 1 can show more scattering where the fog is less intense.
+	 */
+	UPROPERTY(EditAnywhere, BlueprintReadOnly, Category = FogScreenSpaceScattering, meta = (DisplayName = "Scene Color Scattering Amount Power", UIMin = "0.01", UIMax = "2", ClampMin = "0"))
+	float FSSSSceneColorScatteringAmountPower;
+
 	/**
 	 * The scale of the spread for the fog multiple scattering. This basically scales mip level sampled from the blurred screen space FSSS texture.
 	 * Increasing will make scene and thicker fog look fuzzier.
diff --git a/Engine/Source/Runtime/Engine/Private/Components/ExponentialHeightFogComponent.cpp b/Engine/Source/Runtime/Engine/Private/Components/ExponentialHeightFogComponent.cpp
--- a/Engine/Source/Runtime/Engine/Private/Components/ExponentialHeightFogComponent.cpp
+++ b/Engine/Source/Runtime/Engine/Private/Components/ExponentialHeightFogComponent.cpp
@@ -42,6 +42,7 @@ constexpr bool UseExponentialHeightFogStateStream = true;
 	UEXPONENTIALHEIGHTFOG_VARIABLE(float, , SkyLightCaptureAffectsHeightFogRoughness) \
 	UEXPONENTIALHEIGHTFOG_VARIABLE(bool, b, EnableFSSS) \
 	UEXPONENTIALHEIGHTFOG_VARIABLE(float, , FSSSSceneColorScatteringAmountScale) \
+	UEXPONENTIALHEIGHTFOG_VARIABLE(float, , FSSSSceneColorScatteringAmountPower) \
 	UEXPONENTIALHEIGHTFOG_VARIABLE(float, , FSSSSpreadScale) \
 	UEXPONENTIALHEIGHTFOG_VARIABLE(float, , FSSSBlurControl) \
 	UEXPONENTIALHEIGHTFOG_VARIABLES_WITH_SETFUNCTION
@@ -113,6 +114,7 @@ UExponentialHeightFogComponent::UExponentialHeightFogComponent(const FObjectInit
 
 	bEnableFSSS = false;
 	FSSSSceneColorScatteringAmountScale = 1.0f;
+	FSSSSceneColorScatteringAmountPower = 1.0f;
 	FSSSSpreadScale = 0.1f;
 	FSSSBlurControl = 0.5f;
 }
diff --git a/Engine/Source/Runtime/Engine/Public/StateStream/ExponentialHeightFogStateStream.h b/Engine/Source/Runtime/Engine/Public/StateStream/ExponentialHeightFogStateStream.h
--- a/Engine/Source/Runtime/Engine/Public/StateStream/ExponentialHeightFogStateStream.h
+++ b/Engine/Source/Runtime/Engine/Public/StateStream/ExponentialHeightFogStateStream.h
@@ -95,6 +95,8 @@ struct FExponentialHeightFogDynamicState
 	UPROPERTY()
 	float FSSSSceneColorScatteringAmountScale = 1.0f;
 	UPROPERTY()
+	float FSSSSceneColorScatteringAmountPower = 1.0f;
+	UPROPERTY()
 	float FSSSSpreadScale = 0.1f;
 	UPROPERTY()
 	float FSSSBlurControl = 0.5f;
diff --git a/Engine/Source/Runtime/Renderer/Private/FogRendering.cpp b/Engine/Source/Runtime/Renderer/Private/FogRendering.cpp
--- a/Engine/Source/Runtime/Renderer/Private/FogRendering.cpp
+++ b/Engine/Source/Runtime/Renderer/Private/FogRendering.cpp
@@ -205,6 +205,7 @@ void SetupFogUniformParameters(FRDGBuilder& GraphBuilder, const FViewInfo& View,
 
 	OutParameters.EnableFSSS							= 0.0f;
 	OutParameters.FSSSSceneColorScatteringAmountScale	= 0.0f;
+	OutParameters.FSSSSceneColorScatteringAmountPower	= 0.0f;
 	OutParameters.FSSSSpreadScale						= 0.0f;
 	OutParameters.FSSSBlurControl						= 0.0f;
 	if(Scene && Scene->HasAnyExponentialHeightFog())
@@ -212,6 +213,7 @@ void SetupFogUniformParameters(FRDGBuilder& GraphBuilder, const FViewInfo& View,
 		const FExponentialHeightFogSceneInfo& HFog = Scene->ExponentialFogs[0];
 		OutParameters.EnableFSSS							= !HFog.bEnableFSSS ? 0.0f : 1.0f;
 		OutParameters.FSSSSceneColorScatteringAmountScale	= !HFog.bEnableFSSS ? 0.0f : HFog.FSSSSceneColorScatteringAmountScale;
+		OutParameters.FSSSSceneColorScatteringAmountPower	= !HFog.bEnableFSSS ? 0.0f : HFog.FSSSSceneColorScatteringAmountPower;
 		OutParameters.FSSSSpreadScale						= !HFog.bEnableFSSS ? 0.0f : HFog.FSSSSpreadScale;
 		OutParameters.FSSSBlurControl						= !HFog.bEnableFSSS ? 0.0f : HFog.FSSSBlurControl;
 	}
diff --git a/Engine/Source/Runtime/Renderer/Private/FogRendering.h b/Engine/Source/Runtime/Renderer/Private/FogRendering.h
--- a/Engine/Source/Runtime/Renderer/Private/FogRendering.h
+++ b/Engine/Source/Runtime/Renderer/Private/FogRendering.h
@@ -34,6 +34,7 @@ BEGIN_GLOBAL_SHADER_PARAMETER_STRUCT(FFogUniformParameters,)
 	SHADER_PARAMETER(float, VolumetricFogPhaseG)
 	SHADER_PARAMETER(float, EnableFSSS)
 	SHADER_PARAMETER(float, FSSSSceneColorScatteringAmountScale)
+	SHADER_PARAMETER(float, FSSSSceneColorScatteringAmountPower)
 	SHADER_PARAMETER(float, FSSSSpreadScale)
 	SHADER_PARAMETER(float, FSSSBlurControl)
 	SHADER_PARAMETER_TEXTURE(TextureCube, FogInscatteringColorCubemap)
diff --git a/Engine/Source/Runtime/Renderer/Private/SceneCore.cpp b/Engine/Source/Runtime/Renderer/Private/SceneCore.cpp
--- a/Engine/Source/Runtime/Renderer/Private/SceneCore.cpp
+++ b/Engine/Source/Runtime/Renderer/Private/SceneCore.cpp
@@ -440,6 +440,7 @@ FExponentialHeightFogSceneInfo::FExponentialHeightFogSceneInfo(uint64 InId, cons
 
 	bEnableFSSS = State.bEnableFSSS;
 	FSSSSceneColorScatteringAmountScale = State.FSSSSceneColorScatteringAmountScale;
+	FSSSSceneColorScatteringAmountPower = State.FSSSSceneColorScatteringAmountPower;
 	FSSSSpreadScale = State.FSSSSpreadScale;
 	FSSSBlurControl = State.FSSSBlurControl;
 
diff --git a/Engine/Source/Runtime/Renderer/Private/SceneCore.h b/Engine/Source/Runtime/Renderer/Private/SceneCore.h
--- a/Engine/Source/Runtime/Renderer/Private/SceneCore.h
+++ b/Engine/Source/Runtime/Renderer/Private/SceneCore.h
@@ -192,6 +192,7 @@ public:
 
 	bool bEnableFSSS;
 	float FSSSSceneColorScatteringAmountScale;
+	float FSSSSceneColorScatteringAmountPower;
 	float FSSSSpreadScale;
 	float FSSSBlurControl;
 
