- Merge environment defines from the value analyzer into the module instead of moving them, preserving defines that were added during AnalyzeInStage

#tests Create a material with ParticleRandom as color use in a NS and observe the color is different shades of grey

#Synced-CL 52684591

diff --git a/Engine/Source/Runtime/Engine/Private/Materials/MaterialIRModuleBuilder.cpp b/Engine/Source/Runtime/Engine/Private/Materials/MaterialIRModuleBuilder.cpp
--- a/Engine/Source/Runtime/Engine/Private/Materials/MaterialIRModuleBuilder.cpp
+++ b/Engine/Source/Runtime/Engine/Private/Materials/MaterialIRModuleBuilder.cpp
@@ -1200,8 +1200,11 @@ struct FMaterialIRModuleBuilderImpl
 		KeepDefineConditionally(TEXT("NEEDS_PER_INSTANCE_RANDOM_PS"), bUsedWithInstancedStaticMeshes);
 		KeepDefineConditionally(TEXT("USES_PER_INSTANCE_FADE_AMOUNT"), bUsedWithInstancedStaticMeshes);
 
-		// Move final environment defines from analyzer into output module
-		Module->EnvironmentDefines = MoveTemp(ValueAnalyzer.EnvironmentDefines);
+		// Merge defines from analyzer into module (don't discard defines added during AnalyzeInStage)
+		for (const FName& Define : ValueAnalyzer.EnvironmentDefines)
+		{
+			Module->EnvironmentDefines.Add(Define);
+		}
 	}
 
 	// Catches environment defines that conflict with the material's blend mode or feature set.
