New material translator: Fix refraction output for some single layer water content.
The Swizzle could fail with the default float output.

#rnx
#Synced-CL 52672995

diff --git a/Engine/Source/Runtime/Engine/Private/Materials/MaterialIRModuleBuilder.cpp b/Engine/Source/Runtime/Engine/Private/Materials/MaterialIRModuleBuilder.cpp
--- a/Engine/Source/Runtime/Engine/Private/Materials/MaterialIRModuleBuilder.cpp
+++ b/Engine/Source/Runtime/Engine/Private/Materials/MaterialIRModuleBuilder.cpp
@@ -805,9 +805,8 @@ struct FMaterialIRModuleBuilderImpl
 			// Append refraction depth bias to refraction for translucent and single layer water materials
 			else if (Property == MP_Refraction && (IsTranslucentBlendMode(Builder->BlendMode) || Builder->MaterialInterface->GetShadingModels().HasShadingModel(MSM_SingleLayerWater)))
 			{
-				MIR::FValueRef UserRefraction = Emitter.Swizzle(PropertyValue, MIR::FSwizzleMask::XY());
-				MIR::FValueRef RefractionDepthBias = Emitter.CastToFloat(Emitter.NamedPrimitiveUniform(FName(TEXT("RefractionDepthBias")), Emitter.ConstantFloat(Emitter.BaseMaterial->RefractionDepthBias)), 1);
-				PropertyValue = Emitter.Vector3(UserRefraction, RefractionDepthBias);
+				MIR::FValueRef RefractionDepthBias = Emitter.NamedPrimitiveUniform(FName(TEXT("RefractionDepthBias")), Emitter.ConstantFloat(Emitter.BaseMaterial->RefractionDepthBias));
+				PropertyValue = Emitter.Vector3(Emitter.CastToFloat(PropertyValue, 2), RefractionDepthBias);
 			}
 			// Pack SubsurfaceColor (float3) with SubsurfaceProfile ID (float1) into float4,
 			// matching the old translator's behavior (HLSLMaterialTranslator.cpp:1623-1636).
