- Add NMT support for UMaterialExpressionDecalDerivative & UMaterialExpressionRgbToHsv

#Synced-CL 52684591

diff --git a/Engine/Source/Runtime/Engine/Private/Materials/MaterialExpressionsToMIR.cpp b/Engine/Source/Runtime/Engine/Private/Materials/MaterialExpressionsToMIR.cpp
--- a/Engine/Source/Runtime/Engine/Private/Materials/MaterialExpressionsToMIR.cpp
+++ b/Engine/Source/Runtime/Engine/Private/Materials/MaterialExpressionsToMIR.cpp
@@ -4339,6 +4339,73 @@ void UMaterialExpressionReflectionVectorWS::Build(MIR::FEmitter& Em)
 	}
 }
 
+void UMaterialExpressionRgbToHsv::Build(MIR::FEmitter& Em)
+{
+	FValueRef InputValue = Em.CastToFloatKind(Em.Input(&Input));
+	UE_MIR_CHECKPOINT(Em);
+
+	MIR::FPrimitive Primitive = *InputValue->Type.AsPrimitive();
+	if ((!Primitive.IsRowVector() && !Primitive.IsColumnVector()) || Primitive.NumComponents() < 3)
+	{
+		Em.Error(InputValue, TEXT("Expected a 3D or 4D vector."));
+		return;
+	}
+
+	FValueRef R = Em.Subscript(InputValue, 0);
+	FValueRef G = Em.Subscript(InputValue, 1);
+	FValueRef B = Em.Subscript(InputValue, 2);
+
+	FValueRef NegOne		= Em.ConstantFloat(-1.0f);
+	FValueRef TwoThirds		= Em.Divide(Em.ConstantFloat(2.0f), Em.ConstantFloat(3.0f));
+	FValueRef Zero			= Em.ConstantFloat(0.0f);
+	FValueRef NegOneThird	= Em.Divide(Em.ConstantFloat(-1.0f), Em.ConstantFloat(3.0f));
+	FValueRef Six			= Em.ConstantFloat(6.0f);
+	FValueRef Epsilon		= Em.ConstantFloat(1e-10f);
+
+	// P = (G < B) ? float4(B, G, -1.0, 2.0/3.0) : float4(G, B, 0.0, -1.0/3.0)
+	FValueRef P = Em.Branch(Em.LessThan(G, B), Em.Vector4(B, G, NegOne, TwoThirds), Em.Vector4(G, B, Zero, NegOneThird));
+
+	FValueRef Px = Em.Subscript(P, 0);
+	FValueRef Py = Em.Subscript(P, 1);
+	FValueRef Pz = Em.Subscript(P, 2);
+	FValueRef Pw = Em.Subscript(P, 3);
+
+	// Q = (R < P.x) ? float4(P.xyw, R) : float4(R, P.yzx)
+	FValueRef Q = Em.Branch(Em.LessThan(R, Px), Em.Vector4(Px, Py, Pw, R), Em.Vector4(R, Py, Pz, Px));
+
+	FValueRef Qx = Em.Subscript(Q, 0);
+	FValueRef Qy = Em.Subscript(Q, 1);
+	FValueRef Qz = Em.Subscript(Q, 2);
+	FValueRef Qw = Em.Subscript(Q, 3);
+
+	// Chroma = Q.x - min(Q.w, Q.y)
+	FValueRef Chroma = Em.Subtract(Qx, Em.Min(Qw, Qy));
+
+	// Hue = abs((Q.w - Q.y) / (6.0 * Chroma + 1e-10) + Q.z)
+	FValueRef Hue = Em.Abs(Em.Add( Em.Divide(Em.Subtract(Qw, Qy), Em.Add(Em.Multiply(Six, Chroma), Epsilon)), Qz));
+
+	// S = Chroma / (Q.x + 1e-10)
+	FValueRef S = Em.Divide(Chroma, Em.Add(Qx, Epsilon));
+
+	// Result = float3(Hue, S, Q.x) or float4(Hue, S, Q.x, input.w)
+	FValueRef Result;
+	if (Primitive.NumComponents() == 4)
+	{
+		Result = Em.Vector4(Hue, S, Qx, Em.Subscript(InputValue, 3));
+	}
+	else
+	{
+		Result = Em.Vector3(Hue, S, Qx);
+	}
+
+	if (Primitive.IsColumnVector())
+	{
+		Result = Em.Transpose(Result);
+	}
+
+	Em.Output(0, Result);
+}
+
 void UMaterialExpressionRotateAboutAxis::Build(MIR::FEmitter& Em)
 {
 	FValueRef Angle = Em.Multiply(Em.Subscript(Em.Input(&RotationAngle), 0), Em.ConstantFloat(2.0f * (float)UE_PI / Period));
@@ -5127,6 +5194,14 @@ void UMaterialExpressionAbsorptionMediumMaterialOutput::Build(MIR::FEmitter& Em)
 	Em.SetCustomOutputs(TEXTVIEW("AbsorptionMediumMaterialOutput"), OutputValues, MIR::EMaterialOutputFrequency::PerPixel);
 }
 
+void UMaterialExpressionDecalDerivative::Build(MIR::FEmitter& Em)
+{
+	Em.CheckMaterialDomainIsAnyOf({ MD_DeferredDecal });
+
+	Em.Output(0, Em.Extern<MIR::FExternSimpleHLSL>({ MIR::EExternSimpleType::Float2, TEXT("ComputeDecalDDX(Parameters)") }));
+	Em.Output(1, Em.Extern<MIR::FExternSimpleHLSL>({ MIR::EExternSimpleType::Float2, TEXT("ComputeDecalDDY(Parameters)") }));
+}
+
 void UMaterialExpressionDecalMipmapLevel::Build(MIR::FEmitter& Em)
 {
 	Em.CheckMaterialDomainIsAnyOf({ MD_DeferredDecal });
diff --git a/Engine/Source/Runtime/Engine/Public/Materials/MaterialExpressionDecalDerivative.h b/Engine/Source/Runtime/Engine/Public/Materials/MaterialExpressionDecalDerivative.h
--- a/Engine/Source/Runtime/Engine/Public/Materials/MaterialExpressionDecalDerivative.h
+++ b/Engine/Source/Runtime/Engine/Public/Materials/MaterialExpressionDecalDerivative.h
@@ -13,12 +13,11 @@ class UMaterialExpressionDecalDerivative : public UMaterialExpression
 {
 	GENERATED_UCLASS_BODY()
 
-
 	//~ Begin UMaterialExpression Interface
 #if WITH_EDITOR
+	virtual void Build(MIR::FEmitter& Em) override;
 	virtual int32 Compile(class FMaterialCompiler* Compiler, int32 OutputIndex) override;
 	virtual void GetCaption(TArray<FString>& OutCaptions) const override;
-
 #endif
 	//~ End UMaterialExpression Interface
 };
diff --git a/Engine/Source/Runtime/Engine/Public/Materials/MaterialExpressionRgbToHsv.h b/Engine/Source/Runtime/Engine/Public/Materials/MaterialExpressionRgbToHsv.h
--- a/Engine/Source/Runtime/Engine/Public/Materials/MaterialExpressionRgbToHsv.h
+++ b/Engine/Source/Runtime/Engine/Public/Materials/MaterialExpressionRgbToHsv.h
@@ -16,11 +16,11 @@ public:
 
 	//~ Begin UMaterialExpressionMaterialX Interface
 #if WITH_EDITOR
+	virtual void Build(MIR::FEmitter& Em) override;
 	virtual int32 Compile(class FMaterialCompiler* Compiler, int32 OutputIndex) override;
 	virtual void GetCaption(TArray<FString>& OutCaptions) const override;
 	virtual void GetExpressionToolTip(TArray<FString>& OutToolTip) override;
 	virtual FText GetKeywords() const override {return FText::FromString(TEXT("rgbtohsv"));}
-
 #endif
 	//~ End UMaterialExpressionMaterialX Interface
 };
