Fixed motion vectors on MASS ISKM components.

#Synced-CL 52687437

diff --git a/Engine/Plugins/Runtime/MassGameplay/Source/MassRepresentation/Private/MassRepresentationProcessor.cpp b/Engine/Plugins/Runtime/MassGameplay/Source/MassRepresentation/Private/MassRepresentationProcessor.cpp
--- a/Engine/Plugins/Runtime/MassGameplay/Source/MassRepresentation/Private/MassRepresentationProcessor.cpp
+++ b/Engine/Plugins/Runtime/MassGameplay/Source/MassRepresentation/Private/MassRepresentationProcessor.cpp
@@ -248,6 +248,9 @@ void UMassRepresentationProcessor::UpdateRepresentation(FMassExecutionContext& C
 					}
  
 					Representation.CurrentRepresentation = EMassRepresentationType::StaticMeshInstance;
+					// Seed PrevTransform so the first frame of ISM rendering has a valid
+					// previous transform (avoids motion vector pop from stale data).
+					Representation.PrevTransform = TransformFragment.GetTransform();
 					break;
 				case EMassRepresentationType::SkinnedMeshInstance:
 					if (!bDoKeepActorExtraFrame || !UE::Mass::Representation::IsValidActorRepresentation(Representation.PrevRepresentation))
@@ -256,6 +259,7 @@ void UMassRepresentationProcessor::UpdateRepresentation(FMassExecutionContext& C
 					}
 
 					Representation.CurrentRepresentation = EMassRepresentationType::SkinnedMeshInstance;
+					Representation.PrevTransform = TransformFragment.GetTransform();
 					break;
 				case EMassRepresentationType::None:
 					if (!Actor || ActorInfo.IsOwnedByMass())
diff --git a/Engine/Plugins/Runtime/MassGameplay/Source/MassRepresentation/Private/MassVisualizationComponent.cpp b/Engine/Plugins/Runtime/MassGameplay/Source/MassRepresentation/Private/MassVisualizationComponent.cpp
--- a/Engine/Plugins/Runtime/MassGameplay/Source/MassRepresentation/Private/MassVisualizationComponent.cpp
+++ b/Engine/Plugins/Runtime/MassGameplay/Source/MassRepresentation/Private/MassVisualizationComponent.cpp
@@ -1337,6 +1337,7 @@ void UMassVisualizationComponent::ConstructSkinnedMeshComponents()
 				InstancedSkinnedMeshComponent->SetCastShadow(InstancedSkinnedMeshDesc.bCastShadows);
 				InstancedSkinnedMeshComponent->Mobility = InstancedSkinnedMeshDesc.Mobility;
 				InstancedSkinnedMeshComponent->SetReceivesDecals(false);
+				InstancedSkinnedMeshComponent->SetHasPerInstancePrevTransforms(true);
 				InstancedSkinnedMeshComponent->RegisterComponent();
 
 				if (SharedData == nullptr)
diff --git a/Engine/Source/Runtime/Engine/Classes/Components/InstancedSkinnedMeshComponent.h b/Engine/Source/Runtime/Engine/Classes/Components/InstancedSkinnedMeshComponent.h
--- a/Engine/Source/Runtime/Engine/Classes/Components/InstancedSkinnedMeshComponent.h
+++ b/Engine/Source/Runtime/Engine/Classes/Components/InstancedSkinnedMeshComponent.h
@@ -262,6 +262,10 @@ public:
 	UFUNCTION(BlueprintCallable, Category = "Components|InstancedStaticMesh")
 	ENGINE_API bool SetInstanceTransforms(FPrimitiveInstanceId InstanceId, const FTransform& InstanceTransform, const FTransform& InstancePrevTransform, bool bWorldSpace=false);
 
+	/** Enables or disables per-instance previous transforms for motion blur. Use SetInstancePrevTransform or SetInstanceTransforms to provide previous transforms. */
+	UFUNCTION(BlueprintCallable, Category = "Components|InstancedSkinnedMesh")
+	ENGINE_API void SetHasPerInstancePrevTransforms(bool bInHasPreviousTransforms);
+
 	/** Get the animation index for the instance specified. Returns True on success. */
 	UFUNCTION(BlueprintCallable, Category = "Components|InstancedSkinnedMesh")
 	ENGINE_API bool GetInstanceAnimationIndex(FPrimitiveInstanceId InstanceId, int32& OutAnimationIndex) const;
diff --git a/Engine/Source/Runtime/Engine/Private/Components/InstancedSkinnedMeshComponent.cpp b/Engine/Source/Runtime/Engine/Private/Components/InstancedSkinnedMeshComponent.cpp
--- a/Engine/Source/Runtime/Engine/Private/Components/InstancedSkinnedMeshComponent.cpp
+++ b/Engine/Source/Runtime/Engine/Private/Components/InstancedSkinnedMeshComponent.cpp
@@ -1440,12 +1440,38 @@ bool UInstancedSkinnedMeshComponent::SetInstanceTransform(FPrimitiveInstanceId I
 		return false;
 	}
 
-	int32 InstanceIndex = InstanceDataManager.IdToIndex(InstanceId);
+	const int32 InstanceIndex = InstanceDataManager.IdToIndex(InstanceId);
 	InstanceData[InstanceIndex].Transform = FTransform3f(bWorldSpace ? InstanceTransform.GetRelativeTransform(GetComponentTransform()) : InstanceTransform);
 	InstanceDataManager.TransformChanged(InstanceIndex);
 	return true;
 }
 
+void UInstancedSkinnedMeshComponent::SetHasPerInstancePrevTransforms(bool bInHasPreviousTransforms)
+{
+	if (bHasPreviousTransforms == bInHasPreviousTransforms)
+	{
+		return;
+	}
+
+	bHasPreviousTransforms = bInHasPreviousTransforms;
+
+	if (bHasPreviousTransforms)
+	{
+		PerInstancePrevTransform.SetNum(InstanceData.Num());
+
+		for (int32 Index = 0; Index < InstanceData.Num(); Index++)
+		{
+			PerInstancePrevTransform[Index] = InstanceData[Index].Transform;
+		}
+	}
+	else
+	{
+		PerInstancePrevTransform.Empty();
+	}
+
+	MarkRenderInstancesDirty();
+}
+
 bool UInstancedSkinnedMeshComponent::SetInstancePrevTransform(FPrimitiveInstanceId InstanceId, const FTransform& InstanceTransform, bool bWorldSpace)
 {
 	if (!InstanceDataManager.IsValidId(InstanceId) || !bHasPreviousTransforms || PerInstancePrevTransform.Num() != InstanceData.Num())
