- Implement MIR Build function for RuntimeVirtualTextureOutput material expression

#Synced-CL 52684591

diff --git a/Engine/Source/Runtime/Engine/Private/Materials/MaterialExpressionsToMIR.cpp b/Engine/Source/Runtime/Engine/Private/Materials/MaterialExpressionsToMIR.cpp
--- a/Engine/Source/Runtime/Engine/Private/Materials/MaterialExpressionsToMIR.cpp
+++ b/Engine/Source/Runtime/Engine/Private/Materials/MaterialExpressionsToMIR.cpp
@@ -2621,6 +2621,57 @@ void UMaterialExpressionRuntimeVirtualTextureSample::Build(MIR::FEmitter& Em)
 	}
 }
 
+void UMaterialExpressionRuntimeVirtualTextureOutput::Build(MIR::FEmitter& Em)
+{
+	// Get RVT Values
+	FValueRef BaseColorValue	= Em.TryInput(&BaseColor);
+	FValueRef SpecularValue		= Em.TryInput(&Specular);
+	FValueRef RoughnessValue	= Em.TryInput(&Roughness);
+	FValueRef NormalValue		= Em.TryInput(&Normal);
+	FValueRef WorldHeightValue	= Em.TryInput(&WorldHeight);
+	FValueRef OpacityValue		= Em.TryInput(&Opacity);
+	FValueRef MaskValue			= Em.TryInput(&Mask);
+	FValueRef DisplacementValue	= Em.TryInput(&Displacement);
+	FValueRef Mask4Value		= Em.TryInput(&Mask4);
+
+	// Generate output mask (Note opacity does not have a type bit)
+	uint8 AttributeMask = 0;
+	if (BaseColorValue)    AttributeMask |= 1 << (uint8)ERuntimeVirtualTextureAttributeType::BaseColor;
+	if (SpecularValue)     AttributeMask |= 1 << (uint8)ERuntimeVirtualTextureAttributeType::Specular;
+	if (RoughnessValue)    AttributeMask |= 1 << (uint8)ERuntimeVirtualTextureAttributeType::Roughness;
+	if (NormalValue)       AttributeMask |= 1 << (uint8)ERuntimeVirtualTextureAttributeType::Normal;
+	if (WorldHeightValue)  AttributeMask |= 1 << (uint8)ERuntimeVirtualTextureAttributeType::WorldHeight;
+	if (MaskValue)         AttributeMask |= 1 << (uint8)ERuntimeVirtualTextureAttributeType::Mask;
+	if (DisplacementValue) AttributeMask |= 1 << (uint8)ERuntimeVirtualTextureAttributeType::Displacement;
+	if (Mask4Value)        AttributeMask |= 1 << (uint8)ERuntimeVirtualTextureAttributeType::Mask4;
+
+	// Generate output value array to match VirtualTextureMaterial.usf GetVirtualTextureOutputX order
+	FValueRef OutputValues[] =
+	{
+		BaseColorValue		? Em.CastToFloat(BaseColorValue, 3)		: Em.ConstantFloat3({0.f, 0.f, 0.f}),		// 0: BaseColor
+		SpecularValue		? Em.CastToFloat(SpecularValue, 1)		: Em.ConstantFloat(0.5f),					// 1: Specular
+		RoughnessValue		? Em.CastToFloat(RoughnessValue, 1)		: Em.ConstantFloat(0.5f),					// 2: Roughness
+		NormalValue			? Em.CastToFloat(NormalValue, 3)		: Em.ConstantFloat3({0.f, 0.f, 1.f}),		// 3: Normal
+		WorldHeightValue	? Em.CastToFloat(WorldHeightValue, 1)	: Em.ConstantFloat(0.f),					// 4: WorldHeight
+		OpacityValue		? Em.CastToFloat(OpacityValue, 1)		: Em.ConstantFloat(1.f),					// 5: Opacity
+		MaskValue			? Em.CastToFloat(MaskValue, 1)			: Em.ConstantFloat(1.f),					// 6: Mask
+		DisplacementValue	? Em.CastToFloat(DisplacementValue, 1)	: Em.ConstantFloat(0.f),					// 7: Displacement
+		Mask4Value			? Em.CastToFloat(Mask4Value, 4)			: Em.ConstantFloat4(FLinearColor(0.f, 0.f, 0.f, 0.f)),	// 8: Mask4
+	};
+
+	Em.SetCustomOutputs(
+		TEXTVIEW("VirtualTextureOutput"),
+		OutputValues,
+		MIR::EMaterialOutputFrequency::PerPixel,
+		[AttributeMask](FMaterialIRModule& Module)
+		{
+			FMaterialCompilationOutput& Out = Module.GetCompilationOutput();
+			Out.bHasRuntimeVirtualTextureOutputNode |= (AttributeMask != 0);
+			Out.RuntimeVirtualTextureOutputAttributeMask |= AttributeMask;
+		}
+	);
+}
+
 void UMaterialExpressionTime::Build(MIR::FEmitter& Em)
 {
 	// When pausing the game is ignored for this time expression, use real-time instead of game-time.
diff --git a/Engine/Source/Runtime/Engine/Public/Materials/MaterialExpressionRuntimeVirtualTextureOutput.h b/Engine/Source/Runtime/Engine/Public/Materials/MaterialExpressionRuntimeVirtualTextureOutput.h
--- a/Engine/Source/Runtime/Engine/Public/Materials/MaterialExpressionRuntimeVirtualTextureOutput.h
+++ b/Engine/Source/Runtime/Engine/Public/Materials/MaterialExpressionRuntimeVirtualTextureOutput.h
@@ -54,7 +54,7 @@ public:
 	//~ Begin UMaterialExpression Interface
 	virtual int32 Compile(class FMaterialCompiler* Compiler, int32 OutputIndex) override;
 	virtual void GetCaption(TArray<FString>& OutCaptions) const override;
-
+	virtual void Build(MIR::FEmitter& Em) override;
 	//~ End UMaterialExpression Interface
 #endif
 
