Mark FEdGraphPin as containing _weak_ references and not _strong_. Strong causes a problem with GC because FEdGraphPin will build an empty GC schema.

#Synced-CL 52678434

diff --git a/Engine/Source/Runtime/Engine/Classes/EdGraph/EdGraphPin.h b/Engine/Source/Runtime/Engine/Classes/EdGraph/EdGraphPin.h
--- a/Engine/Source/Runtime/Engine/Classes/EdGraph/EdGraphPin.h
+++ b/Engine/Source/Runtime/Engine/Classes/EdGraph/EdGraphPin.h
@@ -232,7 +232,7 @@ struct TStructOpsTypeTraits< FEdGraphPinType > : public TStructOpsTypeTraitsBase
 		WithSerializer = true,
 		WithPostSerialize = true
 	};
-	static constexpr EPropertyObjectReferenceType WithSerializerObjectReferences = EPropertyObjectReferenceType::Strong;
+	static constexpr EPropertyObjectReferenceType WithSerializerObjectReferences = EPropertyObjectReferenceType::Weak;
 };
 
 UENUM()
