[Engine] Add debugging info to the ensure around invalid FieldIds.

#tests PIE Mainstage greenroom repro in the song selection screen.
#Synced-CL 52692120

diff --git a/Engine/Source/Runtime/Engine/Private/FieldNotification/FieldNotificationLibrary.cpp b/Engine/Source/Runtime/Engine/Private/FieldNotification/FieldNotificationLibrary.cpp
--- a/Engine/Source/Runtime/Engine/Private/FieldNotification/FieldNotificationLibrary.cpp
+++ b/Engine/Source/Runtime/Engine/Private/FieldNotification/FieldNotificationLibrary.cpp
@@ -84,7 +84,15 @@ bool SetPropertyValueAndBroadcast(void* SourceValuePtr, void* TargetValuePtr, FP
 			for (FFieldNotificationId FieldNotificationId : FieldIds)
 			{
 				UE::FieldNotification::FFieldId FieldId = NotifyFieldSelf->GetFieldNotificationDescriptor().GetField(Self->GetClass(), FieldNotificationId.GetFieldName());
-				if (ensureMsgf(FieldId.IsValid(), TEXT("The field should was validated at compilation.")))
+				if (ensureMsgf(FieldId.IsValid(),
+					TEXT("BroadcastFieldValueChanged: Field '%s' could not be resolved in the FieldNotificationDescriptor for class '%s' (object: '%s') "
+						 "while setting TargetProperty '%s'. "
+						 "This usually means the field was removed or renamed after the Blueprint node was compiled. "
+						 "Re-open and recompile the Blueprint that calls this node, or verify the field still exists on the class."),
+					*FieldNotificationId.GetFieldName().ToString(),
+					*Self->GetClass()->GetPathName(),
+					*Self->GetName(),
+					*TargetProperty->GetPathName()))
 				{
 					NotifyFieldSelf->BroadcastFieldValueChanged(FieldId);
 				}
