- Add missing early return after error in named reroute expression to prevent crash from proceeding with invalid declaration

#Synced-CL 52684591

diff --git a/Engine/Source/Runtime/Engine/Private/Materials/MaterialExpressionsToMIR.cpp b/Engine/Source/Runtime/Engine/Private/Materials/MaterialExpressionsToMIR.cpp
--- a/Engine/Source/Runtime/Engine/Private/Materials/MaterialExpressionsToMIR.cpp
+++ b/Engine/Source/Runtime/Engine/Private/Materials/MaterialExpressionsToMIR.cpp
@@ -3075,6 +3075,7 @@ void UMaterialExpressionNamedRerouteUsage::Build(MIR::FEmitter& Em)
 	if (!IsDeclarationValid())
 	{
 		Em.Error(TEXT("Named reroute expression does not have a valid declaration."));
+		return;
 	}
 	Em.Output(0, Em.TryInput(&Declaration->Input));
 }
