- Add MCT_TextureExternal support to the Material IR
- Add logical binary operator (AND/OR) translation
- Map MCT_TextureMeshPaint to virtual texture in MIR type conversion

#tests Force the new translator on and open ET, and open the materials which previously failed such as M_SceneDepthOcclusion

#Synced-CL 52684591

diff --git a/Engine/Source/Runtime/Engine/Private/Materials/MaterialIRToHLSLTranslator.cpp b/Engine/Source/Runtime/Engine/Private/Materials/MaterialIRToHLSLTranslator.cpp
--- a/Engine/Source/Runtime/Engine/Private/Materials/MaterialIRToHLSLTranslator.cpp
+++ b/Engine/Source/Runtime/Engine/Private/Materials/MaterialIRToHLSLTranslator.cpp
@@ -262,6 +262,7 @@ static FPrinter& operator<<(FPrinter& Printer, const MIR::FType& Type)
 			case MIR::ETypeKind::Texture2DArray: Printer << TEXTVIEW("Texture2DArray"); break;
 			case MIR::ETypeKind::TextureCubeArray: Printer << TEXTVIEW("TextureCubeArray"); break;
 			case MIR::ETypeKind::Texture3D: Printer << TEXTVIEW("Texture3D"); break;
+			case MIR::ETypeKind::TextureExternal: Printer << TEXTVIEW("TextureExternal"); break;
 			case MIR::ETypeKind::VirtualTexture: Printer << TEXTVIEW("VirtualTexturePhysical"); break;
 			case MIR::ETypeKind::RuntimeVirtualTexture: Printer << TEXTVIEW("VirtualTexturePhysical"); break;
 			case MIR::ETypeKind::SubstrateData: Printer << TEXTVIEW("FSubstrateData"); break;
@@ -1121,6 +1122,8 @@ struct FEntryPointCodeGen
 			FStringView OpString;
 			switch (Operator->Op)
 			{
+				case MIR::BO_And: OpString = TEXTVIEW("&&"); break;
+				case MIR::BO_Or: OpString = TEXTVIEW("||"); break;
 				case MIR::BO_Add: OpString = TEXTVIEW("+"); break;
 				case MIR::BO_Divide: OpString = TEXTVIEW("/"); break;
 				case MIR::BO_Modulo: OpString = TEXTVIEW("%"); break;
@@ -1438,6 +1441,10 @@ struct FEntryPointCodeGen
 		{
 			Printer << TEXTVIEW("VolumeTexture");
 		}
+		else if (bForResourceDeclarations && Type.Is(MIR::ETypeKind::TextureExternal))
+		{
+			Printer << TEXTVIEW("ExternalTexture");
+		}
 		else
 		{
 			Printer << Type;
diff --git a/Engine/Source/Runtime/Engine/Private/Materials/MaterialIRTypes.cpp b/Engine/Source/Runtime/Engine/Private/Materials/MaterialIRTypes.cpp
--- a/Engine/Source/Runtime/Engine/Private/Materials/MaterialIRTypes.cpp
+++ b/Engine/Source/Runtime/Engine/Private/Materials/MaterialIRTypes.cpp
@@ -23,6 +23,7 @@ const TCHAR* TypeKindToString(ETypeKind Kind)
 		CASE(Texture2DArray)
 		CASE(TextureCubeArray)
 		CASE(Texture3D)
+		CASE(TextureExternal)
 		CASE(VirtualTexture)
 		CASE(RuntimeVirtualTexture)
 		CASE(ParameterCollection)
@@ -111,6 +112,8 @@ FType FType::FromMaterialValueType(EMaterialValueType Type)
 		case MCT_Texture2DArray: return FType::MakeTexture2DArray();
 		case MCT_TextureCubeArray: return FType::MakeTextureCubeArray();
 		case MCT_VolumeTexture: return FType::MakeTexture3D();
+		case MCT_TextureExternal: return FType::MakeTextureExternal();
+		case MCT_TextureMeshPaint: return FType::MakeVirtualTexture();
 		case MCT_ShadingModel: return FType::MakeIntScalar();
 		case MCT_StaticBool: return FType::MakeBoolScalar();
 		case MCT_UInt1: return FType::MakeIntScalar();
@@ -227,6 +230,13 @@ FType FType::MakeTexture3D()
 	return Type;
 }
 
+FType FType::MakeTextureExternal()
+{
+	FType Type{};
+	Type.Kind = ETypeKind::TextureExternal;
+	return Type;
+}
+
 FType FType::MakeVirtualTexture()
 {
 	FType Type{};
diff --git a/Engine/Source/Runtime/Engine/Private/Materials/MaterialIRValueAnalyzer.cpp b/Engine/Source/Runtime/Engine/Private/Materials/MaterialIRValueAnalyzer.cpp
--- a/Engine/Source/Runtime/Engine/Private/Materials/MaterialIRValueAnalyzer.cpp
+++ b/Engine/Source/Runtime/Engine/Private/Materials/MaterialIRValueAnalyzer.cpp
@@ -300,6 +300,20 @@ static int32 RegisterTextureUniform(FMaterialIRValueAnalyzer& Analyzer, FName Na
 	return Analyzer.CompilationOutput->UniformExpressionSet.FindOrAddTextureParameter(ParamType, TextureParameterInfo);
 }
 
+static int32 RegisterTextureExternalUniform(FMaterialIRValueAnalyzer& Analyzer, UTexture* TextureObject)
+{
+	// It should not be possible to have nullptr here as we fail before translation
+	check(TextureObject);
+
+	const FMaterialExternalTextureParameterInfo ExternalTextureInfo
+	{
+		.ExternalTextureGuid = TextureObject->GetExternalTextureGuid(),
+		.SourceTextureIndex = Analyzer.MaterialInterface->GetDefaultTextureIdx(TextureObject),
+	};
+
+	return Analyzer.CompilationOutput->UniformExpressionSet.FindOrAddExternalTextureParameter(ExternalTextureInfo);
+}
+
 static void TextureUniform_Analyze(FMaterialIRValueAnalyzer& Analyzer, MIR::FTextureUniform* TextureUniform)
 {
 	// We need to handle downgrading virtual texture to texture 2D when virtual texturing is disabled in the project setting
@@ -308,7 +322,14 @@ static void TextureUniform_Analyze(FMaterialIRValueAnalyzer& Analyzer, MIR::FTex
 	{
 		TextureValueType = MCT_Texture2D;
 	}
-	TextureUniform->Analysis_UniformIndex = RegisterTextureUniform(Analyzer, FName(TextureUniform->Name), TextureUniform->Texture, TextureValueType, UINT8_MAX);
+	if (TextureValueType == MCT_TextureExternal)
+	{
+		TextureUniform->Analysis_UniformIndex = RegisterTextureExternalUniform(Analyzer, TextureUniform->Texture);
+	}
+	else
+	{
+		TextureUniform->Analysis_UniformIndex = RegisterTextureUniform(Analyzer, FName(TextureUniform->Name), TextureUniform->Texture, TextureValueType, UINT8_MAX);
+	}
 }
 
 static void VirtualTextureUniform_Analyze(FMaterialIRValueAnalyzer& Analyzer, MIR::FVirtualTextureUniform* VirtualTextureUniform)
diff --git a/Engine/Source/Runtime/Engine/Public/Materials/MaterialIRTypes.h b/Engine/Source/Runtime/Engine/Public/Materials/MaterialIRTypes.h
--- a/Engine/Source/Runtime/Engine/Public/Materials/MaterialIRTypes.h
+++ b/Engine/Source/Runtime/Engine/Public/Materials/MaterialIRTypes.h
@@ -26,6 +26,7 @@ enum class ETypeKind : uint8
 	Texture2DArray,			// A 2D texture array object.
 	TextureCubeArray,		// A cube texture array object.
 	Texture3D,				// A 3D (volume) texture object.
+	TextureExternal,		// An external texture
 	VirtualTexture,			// A virtual texture object (UTexture - MCT_TextureVirtual).
 	RuntimeVirtualTexture,	// A runtime virtual texture object (URuntimeVirtualTexture).
 	ParameterCollection,	// Material parameter collection object.
@@ -209,6 +210,9 @@ struct FType
 	// Returns the Texture3D object type.
 	static FType MakeTexture3D();
 
+	// Returns the texture external object type.
+	static FType MakeTextureExternal();
+
 	// Returns the VirtualTexture object type (for MCT_TextureVirtual, backed by UTexture).
 	static FType MakeVirtualTexture();
 
@@ -276,7 +280,7 @@ struct FType
 	bool IsMatrix() const { return Is(ETypeKind::Primitive) && Primitive.IsMatrix(); }
 
 	// Returns whether this type is any texture object stored in FTextureUniform (includes virtual textures, excludes RuntimeVirtualTexture).
-	bool IsTexture() const { return Kind == ETypeKind::Texture2D || Kind == ETypeKind::TextureCube || Kind == ETypeKind::Texture2DArray || Kind == ETypeKind::TextureCubeArray || Kind == ETypeKind::Texture3D || Kind == ETypeKind::VirtualTexture; }
+	bool IsTexture() const { return Kind == ETypeKind::Texture2D || Kind == ETypeKind::TextureCube || Kind == ETypeKind::Texture2DArray || Kind == ETypeKind::TextureCubeArray || Kind == ETypeKind::Texture3D || Kind == ETypeKind::VirtualTexture || Kind == ETypeKind::TextureExternal; }
 
 	// Returns whether this type is a VirtualTexture object (MCT_TextureVirtual, backed by UTexture).
 	bool IsVirtualTexture() const { return Kind == ETypeKind::VirtualTexture; }
