SimpleStreamableAssetManager - Fix double-decrement of UsedAssetIndices during asset slot reclamation

Issue:
CL 52399213 introduced asset slot reclamation in RemoveRenderAssetElements: when all objects referencing an asset are removed, the asset slot is reclaimed immediately rather than waiting for the asset itself to be destroyed. However, the condition only checked Num() == 0 without verifying the slot was still marked as used. When an asset is explicitly removed via Pending_RemovedAssetRecords in batch N (which calls Empty() on the sparse array and clears AssetUsedIndices), a stale FAssetRecord from a later object unregister in batch N+1 would see Num() == 0 and trigger the reclamation path again, double-decrementing UsedAssetIndices.

Over time this causes the counter to drift below the actual count, which can delay growth of the AssetUsedIndices bit array and eventually lead to a check-crash when FindAndSetFirstZeroBit returns INDEX_NONE.

Fix:
Add an AssetUsedIndices[AssetIndex] guard to the reclamation condition so it only fires when the slot is still actively owned. Also improved a related comment in AddRenderAssetElements.

#Synced-CL 52670680

diff --git a/Engine/Source/Runtime/Engine/Private/Streaming/SimpleStreamableAssetManager.cpp b/Engine/Source/Runtime/Engine/Private/Streaming/SimpleStreamableAssetManager.cpp
--- a/Engine/Source/Runtime/Engine/Private/Streaming/SimpleStreamableAssetManager.cpp
+++ b/Engine/Source/Runtime/Engine/Private/Streaming/SimpleStreamableAssetManager.cpp
@@ -556,7 +556,7 @@ void FSimpleStreamableAssetManager::AddRenderAssetElements(const TArrayView<FReg
 		}
 		else if (!AssetIndexToSharedIndex[AssetIndex])
 		{
-			// Re-populate after a previous zombie cleanup reclaimed this slot
+			// Re-populate after a previous slot reclamation freed this index
 			AssetIndexToSharedIndex[AssetIndex] = RegisteredAsset.SimpleStreamableAssetManagerIndex;
 		}
 		
@@ -596,7 +596,12 @@ void FSimpleStreamableAssetManager::RemoveRenderAssetElements(int32 ObjectRegist
 		// When all objects referencing an asset are removed, fully
 		// reclaim the asset slot to prevent unbounded memory growth when
 		// assets outlive the objects that reference them.
-		if (AssetIndexToBounds4Index[AssetIndex].Num() == 0)
+		// The AssetUsedIndices guard is required: if an asset was already
+		// removed via Pending_RemovedAssetRecords in an earlier batch, the
+		// sparse array is already Empty (Num()==0) and the slot freed. A
+		// stale FAssetRecord from a later object unregister would otherwise
+		// double-decrement UsedAssetIndices.
+		if (AssetIndexToBounds4Index[AssetIndex].Num() == 0 && AssetUsedIndices[AssetIndex])
 		{
 			AssetIndexToBounds4Index[AssetIndex].Empty();
 			AssetUsedIndices[AssetIndex] = false;
diff --git a/Engine/Source/Runtime/Engine/Private/Tests/Streaming/SSAMCoreTests.cpp b/Engine/Source/Runtime/Engine/Private/Tests/Streaming/SSAMCoreTests.cpp
--- a/Engine/Source/Runtime/Engine/Private/Tests/Streaming/SSAMCoreTests.cpp
+++ b/Engine/Source/Runtime/Engine/Private/Tests/Streaming/SSAMCoreTests.cpp
@@ -40,7 +40,7 @@
 // │ 5      │ Lifecycle.EmptyProcess                       │ Process() with no pending records → safe no-op         │ N/A                      │
 // │ 5      │ Lifecycle.IndexRecyclingNoLeak               │ Freed indices reused, memory stays bounded             │ Primary (proxy)          │
 // │ 5      │ Lifecycle.RapidCreateDestroy                 │ 50 create/destroy cycles → no leak, no lingering bounds│ Primary (proxy)          │
-// │ 5      │ Lifecycle.MemoryBloatOnLevelTransition       │ (DISABLED) UE-363476: inner sparse arrays freed        │ Primary (proxy)          │
+// │ 5      │ Lifecycle.MemoryBloatOnLevelTransition       │ UE-363476: inner sparse arrays freed                   │ Primary (proxy)          │
 // │ 5      │ Lifecycle.BulkRegisterUnregister             │ 100 components, 3 assets → bulk add/remove correctness │ Primary (proxy)          │
 // ├────────┼──────────────────────────────────────────────┼────────────────────────────────────────────────────────┼──────────────────────────┤
 // │ 6      │ Concurrency.ParallelRegistration             │ ParallelFor + AddPrimitive via FPrimitiveSceneDesc     │ FPrimitiveSceneDesc (GT) │
