GitHub 12962 : Fix for missing reflection captures in PIE if lighting scenario is loaded after reflection capture components

#Synced-CL 52679570

diff --git a/Engine/Source/Runtime/Engine/Private/Components/ReflectionCaptureComponent.cpp b/Engine/Source/Runtime/Engine/Private/Components/ReflectionCaptureComponent.cpp
--- a/Engine/Source/Runtime/Engine/Private/Components/ReflectionCaptureComponent.cpp
+++ b/Engine/Source/Runtime/Engine/Private/Components/ReflectionCaptureComponent.cpp
@@ -159,7 +159,8 @@ void UReflectionCaptureComponent::PropagateLightingScenarioChange()
 	{
 		ReregisterComponent();
 	}
-	else
+
+	if (!IsMobilePlatform(ShaderPlatform) || MobileForwardEnableClusteredReflections(ShaderPlatform))
 	{
 		// GetMapBuildData has changed, re-upload
 		MarkDirtyForRecaptureOrUpload();
