Texture Cube Array upgrades.  Mixed source format and gamma correction options are supported, with logic ported from Texture2DArray.cpp to convert the source data to a common format where necessary.  UX filters out cube maps with a mismatched size from being assigned, and if some invalid items end up being assigned anyway (through code, script, or change in size of a referenced resource), displays an actionable error message.

The original goal was just to implement filtering, but the tag values in the cached FAssetData didn't include the source format, so there was no way for the filtering logic to detect mismatched formats.  The alternative of properly supporting mixed source formats was a functionality upgrade.

However, I later discovered the source format missing was due to a separate bug, where asset drag drops use stale cached FAssetData from disk, rather than refreshed tag data from the registry.  I noticed this when investigating an additional issue where drags would randomly fail due to the "IsSourceValid" tag set in the RequiredAssetDataTags metadata being missing from the FAssetData's tags.  "IsSourceValid" is initialized in the dynamic tags in the registry, but not stored to disk.

The registry tags may be loaded after the cached FAssetData have already been permanently cached in the drag op.  Registry tags are cleaned up at the whim of the garbage collector, so  dragging the same object a second time would usually work, because garbage collection hadn't cleaned up the registry tags yet.  Hovering over an asset also forces the registry tags to be generated for the properties dialog pop up, so reproducing the bug required dragging it BEFORE that dialog had a chance to appear.  The timing isn't that tight (around half a second), so users will frequently encounter the bug.  The fix is to defer initialize a refreshed copy of the FAssetData in ExtractAssetDataFromDrag.

#Synced-CL 52720478

diff --git a/Engine/Source/Editor/DetailCustomizations/Private/DetailCustomizations.cpp b/Engine/Source/Editor/DetailCustomizations/Private/DetailCustomizations.cpp
--- a/Engine/Source/Editor/DetailCustomizations/Private/DetailCustomizations.cpp
+++ b/Engine/Source/Editor/DetailCustomizations/Private/DetailCustomizations.cpp
@@ -158,6 +158,7 @@
 #include "TemplateStringStructCustomization.h"
 #include "Templates/SharedPointer.h"
 #include "TextCustomization.h"
+#include "TextureCubeArrayDetails.h"
 #include "TimecodeDetailsCustomization.h"
 #include "TimespanStructCustomization.h"
 #include "UObject/UnrealNames.h"
@@ -494,6 +495,7 @@ void FDetailCustomizationsModule::RegisterObjectCustomizations()
 	RegisterCustomClassLayout("MaterialExpressionGetMaterialAttributes", FOnGetDetailCustomizationInstance::CreateStatic(&FMaterialAttributePropertyDetails::MakeInstance));
 	RegisterCustomClassLayout("MaterialExpressionSetMaterialAttributes", FOnGetDetailCustomizationInstance::CreateStatic(&FMaterialAttributePropertyDetails::MakeInstance));
 	RegisterCustomClassLayout("MaterialExpressionTextureBase", FOnGetDetailCustomizationInstance::CreateStatic(&FMaterialExpressionTextureBaseDetails::MakeInstance));
+	RegisterCustomClassLayout("TextureCubeArray", FOnGetDetailCustomizationInstance::CreateStatic(&FTextureCubeArrayDetails::MakeInstance));
 	RegisterCustomClassLayout("MaterialInstanceDynamic", FOnGetDetailCustomizationInstance::CreateStatic(&FMaterialInstanceDynamicDetails::MakeInstance));
 	RegisterCustomClassLayout("MaterialParameterCollection", FOnGetDetailCustomizationInstance::CreateStatic(&FMaterialParameterCollectionDetails::MakeInstance));
 	RegisterCustomClassLayout("SkeletalMeshLODSettings", FOnGetDetailCustomizationInstance::CreateStatic(&FSkeletalMeshLODSettingsDetails::MakeInstance));
diff --git a/Engine/Source/Editor/DetailCustomizations/Private/TextureCubeArrayDetails.cpp b/Engine/Source/Editor/DetailCustomizations/Private/TextureCubeArrayDetails.cpp
new file mode 100644
--- /dev/null
+++ b/Engine/Source/Editor/DetailCustomizations/Private/TextureCubeArrayDetails.cpp
@@ -0,0 +1,290 @@
+// Copyright Epic Games, Inc. All Rights Reserved.
+
+#include "TextureCubeArrayDetails.h"
+
+#include "DetailCategoryBuilder.h"
+#include "DetailLayoutBuilder.h"
+#include "DetailWidgetRow.h"
+#include "IDetailChildrenBuilder.h"
+#include "PropertyCustomizationHelpers.h"
+#include "Engine/TextureCubeArray.h"
+#include "Engine/TextureCube.h"
+#include "Widgets/Layout/SBox.h"
+#include "SWarningOrErrorBox.h"
+
+#define LOCTEXT_NAMESPACE "TextureCubeArrayDetails"
+
+TSharedRef<IDetailCustomization> FTextureCubeArrayDetails::MakeInstance()
+{
+	return MakeShareable(new FTextureCubeArrayDetails);
+}
+
+void FTextureCubeArrayDetails::CustomizeDetails(IDetailLayoutBuilder& DetailBuilder)
+{
+	TArray<TWeakObjectPtr<UObject>> Objects;
+	DetailBuilder.GetObjectsBeingCustomized(Objects);
+
+	if (Objects.Num() != 1)
+	{
+		return;
+	}
+
+	TextureCubeArrayPtr = Cast<UTextureCubeArray>(Objects[0].Get());
+	if (!TextureCubeArrayPtr.IsValid())
+	{
+		return;
+	}
+
+	IDetailCategoryBuilder& SourceCubeCategory = DetailBuilder.EditCategory("SourceCube");
+
+	// Error banner for size/layout mismatches
+	SourceCubeCategory.AddCustomRow(LOCTEXT("MismatchWarningFilter", "Mismatch Warning"), false)
+		.WholeRowContent()
+		[
+			SNew(SBox)
+			.Padding(FMargin(0.f, 4.f))
+			[
+				SNew(SWarningOrErrorBox)
+				.MessageStyle(EMessageStyle::Error)
+				.Message_Lambda([this]() -> FText { return GetErrorText(); })
+			]
+		]
+		.Visibility(TAttribute<EVisibility>::CreateLambda([this]() -> EVisibility
+		{
+			return GetErrorVisibility();
+		}));
+
+	// Custom array builder with filtered asset pickers per element
+	TSharedPtr<IPropertyHandle> SourceTexturesHandle = DetailBuilder.GetProperty(
+		GET_MEMBER_NAME_CHECKED(UTextureCubeArray, SourceTextures));
+
+	if (SourceTexturesHandle.IsValid() && SourceTexturesHandle->IsValidHandle())
+	{
+		TSharedRef<FDetailArrayBuilder> ArrayBuilder = MakeShareable(new FDetailArrayBuilder(SourceTexturesHandle.ToSharedRef()));
+
+		ArrayBuilder->OnGenerateArrayElementWidget(FOnGenerateArrayElementWidget::CreateSP(this, &FTextureCubeArrayDetails::OnGenerateArrayElementWidget));
+
+		SourceCubeCategory.AddCustomBuilder(ArrayBuilder, false);
+	}
+}
+
+void FTextureCubeArrayDetails::OnGenerateArrayElementWidget(
+	TSharedRef<IPropertyHandle> ElementHandle, int32 ElementIndex, IDetailChildrenBuilder& ChildrenBuilder)
+{
+	IDetailPropertyRow& PropertyRow = ChildrenBuilder.AddProperty(ElementHandle);
+
+	TSharedPtr<SWidget> NameWidget;
+	TSharedPtr<SWidget> ValueWidget;
+	FDetailWidgetRow Row;
+	PropertyRow.GetDefaultWidgets(NameWidget, ValueWidget, Row);
+
+	SAssignNew(ValueWidget, SObjectPropertyEntryBox)
+		.PropertyHandle(ElementHandle)
+		.AllowedClass(UTextureCube::StaticClass())
+		.AllowClear(true)
+		.DisplayThumbnail(true)
+		.OnShouldFilterAsset(FOnShouldFilterAsset::CreateSP(this, &FTextureCubeArrayDetails::ShouldFilterCubeAsset, ElementIndex));
+
+	PropertyRow.CustomWidget()
+		.NameContent()
+		.MinDesiredWidth(Row.NameWidget.MinWidth)
+		.MaxDesiredWidth(Row.NameWidget.MaxWidth)
+		[
+			NameWidget.ToSharedRef()
+		]
+		.ValueContent()
+		.MinDesiredWidth(Row.ValueWidget.MinWidth)
+		.MaxDesiredWidth(Row.ValueWidget.MaxWidth)
+		[
+			ValueWidget.ToSharedRef()
+		];
+}
+
+const FTextureSource* FindFirstValidSource(const UTextureCubeArray* Tex)
+{
+	const FTextureSource* FirstValidSource = nullptr;
+	for (const TObjectPtr<UTextureCube>& SourceTexture : Tex->SourceTextures)
+	{
+		if (SourceTexture && SourceTexture->Source.IsValid())
+		{
+			FirstValidSource = &SourceTexture->Source;
+			break;
+		}
+	}
+	return FirstValidSource;
+}
+
+bool FTextureCubeArrayDetails::ShouldFilterCubeAsset(const FAssetData& AssetData, int32 ElementIndex) const
+{
+	if (!TextureCubeArrayPtr.IsValid())
+	{
+		return false;
+	}
+
+	const TArray<TObjectPtr<UTextureCube>>& SourceTextures = TextureCubeArrayPtr->SourceTextures;
+
+	// Find the first valid source that is NOT the element being edited.  If there are no other valid entries,
+	// the user is free to pick anything — there's nothing else for it to conflict with.  This logic is designed
+	// to specifically allow edits on the first element to achieve a valid state matching later elements, which
+	// wouldn't work if we always chose the first valid element as ReferenceSource.  You could work around this
+	// limitation by resetting the element before editing, but that's not very intuitive if you don't know how
+	// the filtering code works under the hood.
+	const FTextureSource* ReferenceSource = nullptr;
+	for (int32 i = 0; i < SourceTextures.Num(); ++i)
+	{
+		if (i != ElementIndex && SourceTextures[i] && SourceTextures[i]->Source.IsValid())
+		{
+			ReferenceSource = &SourceTextures[i]->Source;
+			break;
+		}
+	}
+
+	if (!ReferenceSource)
+	{
+		return false;
+	}
+
+	// Parse the "Dimensions" asset registry tag (format: "512x512x6" or "512x512")
+	FString DimensionsStr;
+	if (!AssetData.GetTagValue(FName("Dimensions"), DimensionsStr))
+	{
+		return false;		// No tag available, allow it and let PostEditChange handle validation
+	}
+
+	// Parse SizeX and SizeY from the dimensions string
+	int32 CandidateSizeX = 0;
+	int32 CandidateSizeY = 0;
+
+	TArray<FString> Parts;
+	DimensionsStr.ParseIntoArray(Parts, TEXT("x"), true);
+	if (Parts.Num() >= 2)
+	{
+		CandidateSizeX = FCString::Atoi(*Parts[0]);
+		CandidateSizeY = FCString::Atoi(*Parts[1]);
+	}
+
+	// Filter out if sizes don't match
+	if (CandidateSizeX != ReferenceSource->GetSizeX() || CandidateSizeY != ReferenceSource->GetSizeY())
+	{
+		return true;	// Filter out
+	}
+
+	return false;	// Allow
+}
+
+EVisibility FTextureCubeArrayDetails::GetErrorVisibility() const
+{
+	if (!TextureCubeArrayPtr.IsValid())
+	{
+		return EVisibility::Collapsed;
+	}
+
+	const FTextureSource* FirstValidSource;
+	if (!(FirstValidSource = FindFirstValidSource(TextureCubeArrayPtr.Get())))
+	{
+		return EVisibility::Collapsed;		// Nothing valid in the array, no mismatch errors possible
+	}
+
+	for (const TObjectPtr<UTextureCube>& Tex : TextureCubeArrayPtr->SourceTextures)
+	{
+		if (!Tex || !Tex->Source.IsValid())
+		{
+			continue;
+		}
+
+		if (Tex->Source.GetSizeX() != FirstValidSource->GetSizeX() ||
+			Tex->Source.GetSizeY() != FirstValidSource->GetSizeY() ||
+			Tex->Source.GetNumSlices() != FirstValidSource->GetNumSlices())
+		{
+			return EVisibility::Visible;
+		}
+	}
+
+	return EVisibility::Collapsed;
+}
+
+FText FTextureCubeArrayDetails::GetErrorText() const
+{
+	if (!TextureCubeArrayPtr.IsValid())
+	{
+		return FText::GetEmpty();
+	}
+
+	const TArray<TObjectPtr<UTextureCube>>& SourceTextures = TextureCubeArrayPtr->SourceTextures;
+
+	// Find first valid texture as reference
+	const UTextureCube* FirstValid = nullptr;
+	for (const TObjectPtr<UTextureCube>& Tex : SourceTextures)
+	{
+		if (Tex && Tex->Source.IsValid())
+		{
+			FirstValid = Tex;
+			break;
+		}
+	}
+
+	if (!FirstValid)
+	{
+		return FText::GetEmpty();
+	}
+
+	const int32 ExpectedSizeX = (int32)FirstValid->Source.GetSizeX();
+	const int32 ExpectedSizeY = (int32)FirstValid->Source.GetSizeY();
+	const int32 ExpectedSlices = FirstValid->Source.GetNumSlices();
+	const TCHAR* ExpectedType = ExpectedSlices == 6 ? TEXT("Cubemap") : TEXT("LongLat");
+
+	// Get first mismatched texture, and indices of additional mismatches.  The dialog could get huge and become difficult
+	// to read, if we displayed full information for each mismatch where there are many, and in common cases, it's most
+	// likely there will be only one mismatch.
+	const UTextureCube* FirstMismatch = nullptr;
+	TStringBuilder<256> MismatchedIndices;
+	for (int32 SourceTextureIndex = 0; SourceTextureIndex < SourceTextures.Num(); SourceTextureIndex++)
+	{
+		const TObjectPtr<UTextureCube>& Tex = SourceTextures[SourceTextureIndex];
+		if (!Tex || !Tex->Source.IsValid())
+		{
+			continue;
+		}
+
+		if (Tex->Source.GetSizeX() != ExpectedSizeX ||
+			Tex->Source.GetSizeY() != ExpectedSizeY ||
+			Tex->Source.GetNumSlices() != ExpectedSlices)
+		{
+			if (!FirstMismatch)
+			{
+				FirstMismatch = Tex;
+			}
+
+			if (MismatchedIndices.Len() == 0)
+			{
+				MismatchedIndices.Appendf(TEXT("%d"), SourceTextureIndex);
+			}
+			else
+			{
+				MismatchedIndices.Appendf(TEXT(", %d"), SourceTextureIndex);
+			}
+		}
+	}
+
+	// FirstMismatch can be null, in cases where GetErrorText is called before GetErrorVisibility has a chance to hide the error control.
+	// Check for that and return.
+	if (!FirstMismatch)
+	{
+		return FText::GetEmpty();
+	}
+
+	return FText::Format(
+		LOCTEXT("SizeMismatchError",
+			"Texture size or layout mismatch:\nExpected:  {0}x{1} ({2})  '{3}'\nMismatch:  {4}x{5} ({6})  '{7}'.\nMismatch indices: {8}"),
+		FText::AsNumber(ExpectedSizeX),
+		FText::AsNumber(ExpectedSizeY),
+		FText::FromString(ExpectedType),
+		FText::FromString(FirstValid->GetName()),
+		FText::AsNumber(FirstMismatch->Source.GetSizeX()),
+		FText::AsNumber(FirstMismatch->Source.GetSizeY()),
+		FText::FromString(FirstMismatch->Source.GetNumSlices() == 6 ? TEXT("Cubemap") : TEXT("LongLat")),
+		FText::FromString(FirstMismatch->GetName()),
+		FText::FromString(MismatchedIndices.GetData()));
+}
+
+#undef LOCTEXT_NAMESPACE
diff --git a/Engine/Source/Editor/DetailCustomizations/Private/TextureCubeArrayDetails.h b/Engine/Source/Editor/DetailCustomizations/Private/TextureCubeArrayDetails.h
new file mode 100644
--- /dev/null
+++ b/Engine/Source/Editor/DetailCustomizations/Private/TextureCubeArrayDetails.h
@@ -0,0 +1,36 @@
+// Copyright Epic Games, Inc. All Rights Reserved.
+
+#pragma once
+
+#include "IDetailCustomization.h"
+#include "Layout/Visibility.h"
+#include "UObject/WeakObjectPtr.h"
+
+class IDetailLayoutBuilder;
+class IDetailChildrenBuilder;
+class IPropertyHandle;
+class UTextureCubeArray;
+struct FAssetData;
+
+class FTextureCubeArrayDetails : public IDetailCustomization
+{
+public:
+	static TSharedRef<IDetailCustomization> MakeInstance();
+
+	virtual void CustomizeDetails(IDetailLayoutBuilder& DetailBuilder) override;
+
+private:
+	/** Filter for array element asset pickers - blocks cubemaps that don't match the first entry's size/format. */
+	bool ShouldFilterCubeAsset(const FAssetData& AssetData, int32 ElementIndex) const;
+
+	/** Generates custom widget per array element with the filtered asset picker. */
+	void OnGenerateArrayElementWidget(TSharedRef<IPropertyHandle> ElementHandle, int32 ElementIndex, IDetailChildrenBuilder& ChildrenBuilder);
+
+	/** Dynamic error text describing which textures mismatch. */
+	FText GetErrorText() const;
+
+	/** Visible only when there's a size/format mismatch in the array. */
+	EVisibility GetErrorVisibility() const;
+
+	TWeakObjectPtr<UTextureCubeArray> TextureCubeArrayPtr;
+};
diff --git a/Engine/Source/Editor/UnrealEd/Private/AssetSelection.cpp b/Engine/Source/Editor/UnrealEd/Private/AssetSelection.cpp
--- a/Engine/Source/Editor/UnrealEd/Private/AssetSelection.cpp
+++ b/Engine/Source/Editor/UnrealEd/Private/AssetSelection.cpp
@@ -805,7 +805,26 @@ namespace AssetUtil
 		else if (Operation->IsOfType<FAssetDragDropOp>())
 		{
 			TSharedPtr<const FAssetDragDropOp> DragDropOp = StaticCastSharedPtr<const FAssetDragDropOp>(Operation);
-			DroppedAssetData.Append(DragDropOp->GetAssets());
+
+			if (!DragDropOp->RefreshedAssetData.IsSet())
+			{
+				TArray<FAssetData> Refreshed;
+				if (IAssetRegistry* Registry = IAssetRegistry::Get())
+				{
+					for (const FAssetData& OriginalData : DragDropOp->GetAssets())
+					{
+						FAssetData RegistryData = Registry->GetAssetByObjectPath(OriginalData.GetSoftObjectPath());
+						Refreshed.Add(RegistryData.IsValid() ? MoveTemp(RegistryData) : OriginalData);
+					}
+				}
+				else
+				{
+					Refreshed = DragDropOp->GetAssets();
+				}
+				DragDropOp->RefreshedAssetData = MoveTemp(Refreshed);
+			}
+
+			DroppedAssetData.Append(DragDropOp->RefreshedAssetData.GetValue());
 		}
 
 		return DroppedAssetData;
diff --git a/Engine/Source/Editor/UnrealEd/Public/DragAndDrop/AssetDragDropOp.h b/Engine/Source/Editor/UnrealEd/Public/DragAndDrop/AssetDragDropOp.h
--- a/Engine/Source/Editor/UnrealEd/Public/DragAndDrop/AssetDragDropOp.h
+++ b/Engine/Source/Editor/UnrealEd/Public/DragAndDrop/AssetDragDropOp.h
@@ -4,6 +4,7 @@
 
 #include "CoreMinimal.h"
 #include "AssetRegistry/AssetData.h"
+#include "AssetSelection.h"
 #include "Input/DragAndDrop.h"
 #include "Layout/Visibility.h"
 #include "DragAndDrop/DecoratedDragDropOp.h"
@@ -93,6 +94,16 @@ private:
 	/** Data for the assets this item represents */
 	TArray<FAssetData> AssetData;
 
+	/**
+	 * Cached asset data with refreshed tags from the asset registry.  The FAssetData captured at drag start may have stale
+	 * on-disk (Fixed) tags that are missing dynamically computed values (e.g. IsSourceValid).  After the asset is loaded
+	 * during drag-over, the registry has fresh Loose tags.  Populated once per drag op by ExtractAssetDataFromDrag.
+	 */
+	mutable TOptional<TArray<FAssetData>> RefreshedAssetData;
+
+	/** Allow access to RefreshedAssetData from ExtractAssetDataFromDrag. */
+	friend TArray<FAssetData> AssetUtil::ExtractAssetDataFromDrag(const TSharedPtr<const FDragDropOperation>& Operation);
+
 	/** Data for the asset paths this item represents */
 	TArray<FString> AssetPaths;
 
diff --git a/Engine/Source/Runtime/Engine/Classes/Engine/TextureCubeArray.h b/Engine/Source/Runtime/Engine/Classes/Engine/TextureCubeArray.h
--- a/Engine/Source/Runtime/Engine/Classes/Engine/TextureCubeArray.h
+++ b/Engine/Source/Runtime/Engine/Classes/Engine/TextureCubeArray.h
@@ -52,6 +52,7 @@ public:
 	ENGINE_API bool UpdateSourceFromSourceTextures(bool bCreatingNewTexture = true);
 	ENGINE_API void InvalidateTextureSource();
 	ENGINE_API bool CheckArrayTexturesCompatibility();
+	ENGINE_API bool CheckLastTextureCompatibility();
 #endif // WITH_EDITOR
 	virtual void UpdateResourceWithParams(EUpdateResourceFlags InFlags) override;
 	virtual EMaterialValueType GetMaterialType() const override { return MCT_TextureCubeArray; }
diff --git a/Engine/Source/Runtime/Engine/Private/TextureCubeArray.cpp b/Engine/Source/Runtime/Engine/Private/TextureCubeArray.cpp
--- a/Engine/Source/Runtime/Engine/Private/TextureCubeArray.cpp
+++ b/Engine/Source/Runtime/Engine/Private/TextureCubeArray.cpp
@@ -10,6 +10,7 @@
 #include "Engine/TextureCube.h"
 #include "EngineUtils.h"
 #include "EngineLogs.h"
+#include "ImageCoreUtils.h"
 #include "ImageUtils.h"
 #include "Misc/CoreStats.h"
 #include "Misc/ScopedSlowTask.h"
@@ -626,7 +627,6 @@ ENGINE_API bool UTextureCubeArray::CheckArrayTexturesCompatibility()
 
 	const FTextureSource& TextureSource = SourceTextures[0]->Source;
 	const FString TextureName = SourceTextures[0]->GetName();
-	const ETextureSourceFormat SourceFormat = TextureSource.GetFormat();
 	const int32 SizeX = TextureSource.GetSizeX();
 	const int32 SizeY = TextureSource.GetSizeY();
 	const int32 NumSlices = TextureSource.GetNumSlices();
@@ -638,7 +638,6 @@ ENGINE_API bool UTextureCubeArray::CheckArrayTexturesCompatibility()
 	{
 		const FTextureSource& TextureSourceCmp = SourceTextures[TextureCmpIndex]->Source;
 		const FString TextureNameCmp = SourceTextures[TextureCmpIndex]->GetName();
-		const ETextureSourceFormat SourceFormatCmp = TextureSourceCmp.GetFormat();
 
 		check( TextureSourceCmp.IsValid() );
 
@@ -653,20 +652,65 @@ ENGINE_API bool UTextureCubeArray::CheckArrayTexturesCompatibility()
 			UE_LOGF(LogTexture, Warning, "TextureCubeArray creation failed. Textures %ls and %ls have different number of faces (some are long/lat, some are not).", *TextureName, *TextureNameCmp);
 			bError = true;
 		}
+	}
+
+	return (!bError);
+}
+
+ENGINE_API bool UTextureCubeArray::CheckLastTextureCompatibility()
+{
+	// If no last texture exists, return true -- can't have a compatibility mismatch if nothing is there.
+	if (!SourceTextures.Num() || !SourceTextures.Last())
+	{
+		return true;
+	}
 
-		if (SourceFormatCmp != SourceFormat)
+	// Reject an explicitly invalid last texture
+	const UTextureCube* LastTexture = SourceTextures.Last();
+	if (!LastTexture->Source.IsValid())
+	{
+		return false;
+	}
+	const FTextureSource& LastSource = LastTexture->Source;
+
+	// Check if the last matches the first valid
+	for (UTextureCube* SourceTexture : SourceTextures)
+	{
+		if (SourceTexture && SourceTexture->Source.IsValid())
 		{
-			UE_LOGF(LogTexture, Warning, "TextureCubeArray creation failed. Textures %ls and %ls have incompatible pixel formats.", *TextureName, *TextureNameCmp);
-			bError = true;
+			const FTextureSource& FirstValidSource = SourceTexture->Source;
+			return
+				FirstValidSource.GetSizeX() == LastSource.GetSizeX() &&
+				FirstValidSource.GetSizeY() == LastSource.GetSizeY() &&
+				FirstValidSource.GetNumSlices() == LastSource.GetNumSlices();
 		}
 	}
 
-	return (!bError);
+	return true;
 }
 
-
 ENGINE_API bool UTextureCubeArray::UpdateSourceFromSourceTextures(bool bCreatingNewTexture)
 {
+	// This should be false when texture is updated to avoid overriding user settings.
+	if (bCreatingNewTexture)
+	{
+		// Pull settings from the first valid texture.  We do this before CheckArrayTexturesCompatibility, to support
+		// cases where the array is initialized with multiple initially null elements, and the first non-null element
+		// is set (CheckArrayTexturesCompatibility will return false if there are any null elements).
+		for (UTextureCube* SourceTexture : SourceTextures)
+		{
+			if (SourceTexture && SourceTexture->Source.IsValid())
+			{
+				CompressionSettings = SourceTexture->CompressionSettings;
+				MipGenSettings = TMGS_NoMipmaps;
+				PowerOfTwoMode = ETexturePowerOfTwoSetting::None;
+				LODGroup = SourceTexture->LODGroup;
+				NeverStream = true;
+				break;
+			}
+		}
+	}
+
 	if (!CheckArrayTexturesCompatibility())
 	{
 		return false;
@@ -677,70 +721,96 @@ ENGINE_API bool UTextureCubeArray::UpdateSourceFromSourceTextures(bool bCreating
 	const FTextureSource& InitialSource = SourceTextures[0]->Source;
 	check( InitialSource.IsValid() ); // checked by CheckArrayTexturesCompatibility
 
-	// Format and format size.
-	const EPixelFormat PixelFormat = SourceTextures[0]->GetPixelFormat();
-	const ETextureSourceFormat Format = InitialSource.GetFormat();
-	const int32 FormatDataSize = InitialSource.GetBytesPerPixel();
 	// X,Y,Z size of the array.
 	const int32 SizeX = InitialSource.GetSizeX();
 	const int32 SizeY = InitialSource.GetSizeY();
 	const int32 NumSlices = InitialSource.GetNumSlices();
 	const uint32 ArraySize = SourceTextures.Num();
-	// Only copy the first mip from the source textures to array texture.
-	const uint32 NumMips = 1;
 
-	// This should be false when texture is updated to avoid overriding user settings.
-	if (bCreatingNewTexture)
+	// Detect format and gamma mismatches across source textures, finding a common format that can represent all
+	// sources without information loss (ported from UTexture2DArray::UpdateSourceFromSourceTextures).
+	ETextureSourceFormat Format = InitialSource.GetFormat();
+	bool bSRGB = InitialSource.GetGammaSpace(0) == EGammaSpace::sRGB;
+	bool bMismatchedFormats = false;
+	bool bMismatchedGammaSpace = false;
+
+	for (int32 SourceTextureIndex = 1; SourceTextureIndex < SourceTextures.Num(); ++SourceTextureIndex)
 	{
-		CompressionSettings = SourceTextures[0]->CompressionSettings;
-		MipGenSettings = TMGS_NoMipmaps;
-		PowerOfTwoMode = ETexturePowerOfTwoSetting::None;
-		LODGroup = SourceTextures[0]->LODGroup;
-		SRGB = SourceTextures[0]->SRGB;
-		NeverStream = true;
+		const FTextureSource& OtherSource = SourceTextures[SourceTextureIndex]->Source;
+
+		if ((OtherSource.GetGammaSpace(0) == EGammaSpace::sRGB) != bSRGB)
+		{
+			bMismatchedGammaSpace = true;
+			bSRGB = false;
+			Format = TSF_RGBA32F;
+		}
+
+		if (Format != OtherSource.GetFormat())
+		{
+			bMismatchedFormats = true;
+			Format = FImageCoreUtils::GetCommonSourceFormat(Format, OtherSource.GetFormat());
+		}
 	}
 
+	if (bMismatchedGammaSpace)
+	{
+		UE_LOGF(LogTexture, Log, "TextureCubeArray: mismatched source gamma spaces, converting all to RGBA32F/Linear ...");
+	}
+	else if (bMismatchedFormats)
+	{
+		UE_LOGF(LogTexture, Log, "TextureCubeArray: mismatched source pixel formats, converting all to %ls/%s ...",
+			ERawImageFormat::GetName(FImageCoreUtils::ConvertToRawImageFormat(Format)), bSRGB ? "sRGB" : "Linear");
+	}
+
+	SRGB = bSRGB;
+
 	// Create the source texture for this UTexture.
-	Source.Init(SizeX, SizeY, ArraySize * NumSlices, NumMips, Format);
+	// Only copy the first mip from the source textures to array texture.
+	Source.Init(SizeX, SizeY, ArraySize * NumSlices, 1, Format);
 	Source.bLongLatCubemap = (NumSlices == 1);
 
 	// this path sets bLongLatCubemap for CubeArray
 	//  most paths do not, so it is not reliable
 
-	// We only copy the top level Mip map.
-	TArray<uint8*, TInlineAllocator<MAX_TEXTURE_MIP_COUNT> > DestMipData;
-	TArray<uint64, TInlineAllocator<MAX_TEXTURE_MIP_COUNT> > MipSizeBytes;
-	DestMipData.AddZeroed(NumMips);
-	MipSizeBytes.AddZeroed(NumMips);
-
-	for (uint32 MipIndex = 0; MipIndex < NumMips; ++MipIndex)
-	{
-		DestMipData[MipIndex] = Source.LockMip(MipIndex);
-		MipSizeBytes[MipIndex] = Source.CalcMipSize(MipIndex) / ArraySize;
-	}
+	uint8* DestMipData = Source.LockMip(0);
+	const int64 MipSizePerElement = Source.CalcMipSize(0) / ArraySize;
 
-	TArray64<uint8> RefCubeData;
 	for (int32 SourceTexIndex = 0; SourceTexIndex < SourceTextures.Num(); ++SourceTexIndex)
 	{
-		for (uint32 MipIndex = 0; MipIndex < NumMips; ++MipIndex)
+		FTextureSource& TextureSource = SourceTextures[SourceTexIndex]->Source;
+		check( TextureSource.IsValid() ); // checked by CheckArrayTexturesCompatibility
+		void* DestSliceData = DestMipData + MipSizePerElement * SourceTexIndex;
+
+		if (TextureSource.GetFormat() == Format && ((TextureSource.GetGammaSpace(0) == EGammaSpace::sRGB) == bSRGB))
 		{
-			const uint64 Size = MipSizeBytes[MipIndex];
-			const uint64 CheckSize = SourceTextures[SourceTexIndex]->Source.CalcMipSize(MipIndex);
-			check( SourceTextures[SourceTexIndex]->Source.IsValid() ); // checked by CheckArrayTexturesCompatibility
-			check(Size == CheckSize);
-
-			RefCubeData.Reset();
-			SourceTextures[SourceTexIndex]->Source.GetMipData(RefCubeData, MipIndex);
-			void* Dst = DestMipData[MipIndex] + Size * SourceTexIndex;
-			FMemory::Memcpy(Dst, RefCubeData.GetData(), Size);
+			// Fast path: format and gamma match, direct copy.
+			TArray64<uint8> RefCubeData;
+			TextureSource.GetMipData(RefCubeData, 0);
+			check(RefCubeData.Num() == MipSizePerElement);
+			FMemory::Memcpy(DestSliceData, RefCubeData.GetData(), MipSizePerElement);
 		}
-	}
+		else
+		{
+			// Slow path: convert format/gamma via FImage.
+			FImage SourceImage;
+			TextureSource.GetMipImage(SourceImage, 0);
 
-	for (uint32 MipIndex = 0; MipIndex < NumMips; ++MipIndex)
-	{
-		Source.UnlockMip(MipIndex);
+			if (TextureSource.GetFormat() != Format)
+			{
+				// When gamma space is mismatched we fall back to TSF_RGBA32F which is always linear,
+				// so no separate gamma conversion is needed in that case.
+				FImage ConvertedImage;
+				SourceImage.CopyTo(ConvertedImage, FImageCoreUtils::ConvertToRawImageFormat(Format), bSRGB ? EGammaSpace::sRGB : EGammaSpace::Linear);
+				ConvertedImage.Swap(SourceImage);
+			}
+
+			check(SourceImage.RawData.Num() == MipSizePerElement);
+			FMemory::Memcpy(DestSliceData, SourceImage.RawData.GetData(), MipSizePerElement);
+		}
 	}
 
+	Source.UnlockMip(0);
+
 	SetLightingGuid();
 	ValidateSettingsAfterImportOrEdit();
 	UpdateResource();
@@ -863,34 +933,48 @@ ENGINE_API void UTextureCubeArray::PostEditChangeProperty(FPropertyChangedEvent&
 
 	if (PropertyName == GET_MEMBER_NAME_CHECKED(UTextureCubeArray, SourceTextures))
 	{
-		for (int i=0;i<SourceTextures.Num();i++)
+		int32 NumValid = 0;
+		int32 NumRemoved = 0;
+		for (int i=0; i < SourceTextures.Num(); i++)
 		{
 			UTextureCube * SourceTexture = SourceTextures[i];
-			if (SourceTexture && !SourceTexture->Source.IsValid())
+			if (SourceTexture)
 			{
-				UE_LOGF(LogTexture, Warning, "Texture has no Source, cannot be used [%ls]", *SourceTexture->GetFullName());
-				SourceTextures[i] = nullptr;
+				if (SourceTexture->Source.IsValid())
+				{
+					NumValid++;
+				}
+				else
+				{
+					UE_LOGF(LogTexture, Warning, "Texture has no Source, cannot be used [%ls]", *SourceTexture->GetFullName());
+					SourceTextures[i] = nullptr;
+					NumRemoved++;
+				}
 			}
 		}
 
 		// Empty SourceTextures, remove any resources if present.
-		if (SourceTextures.Num() == 0)
+		if (NumValid == 0)
 		{
 			InvalidateTextureSource();
 		}
-		// First entry into an empty texture array.
-		else if (SourceTextures.Num() == 1)
+		// First valid entry into an empty texture array -- indicate to UpdateSourceFromSourceTextures that we are creating a new texture.
+		else if (NumValid == 1)
 		{
-			UpdateSourceFromSourceTextures(true);
+			// Except in the rare situation where an item was removed for being invalid above, don't treat that as creating a new texture.
+			const bool bCreatingNewTexture = NumRemoved == 0;
+			UpdateSourceFromSourceTextures(bCreatingNewTexture);
 		}
-		// Couldn't add to non-empty array (Error msg logged).
+		// Try to update from all source textures.
 		else if (UpdateSourceFromSourceTextures(false) == false)
 		{
+			// If this is a change to the last element, and it's invalid, remove it to make the array valid.  Arrays don't support calling
+			// OnShouldFilterAsset on child properties of the array, so this is a backstop to clean up such cases on newly added elements.
+			// Invalid edits to an existing last element should be blocked by the filters, so this case should only be reached on an add.
 			int32 ChangedIndex = PropertyChangedEvent.GetArrayIndex(PropertyName.ToString());
 			int32 LastIndex = SourceTextures.Num() - 1;
 
-			// But don't remove an empty texture, only an incompatible one.
-			if (SourceTextures[LastIndex] != nullptr && ChangedIndex == LastIndex)
+			if (ChangedIndex == LastIndex && !CheckLastTextureCompatibility())
 			{
 				SourceTextures.RemoveAt(LastIndex);
 			}
