Fixed bug in instanced static mesh where previous transforms were not added if the batch-add APIs were used (AddInstances* vs AddInstance).

#Synced-CL 52619547

diff --git a/Engine/Source/Runtime/Engine/Private/InstancedStaticMesh.cpp b/Engine/Source/Runtime/Engine/Private/InstancedStaticMesh.cpp
--- a/Engine/Source/Runtime/Engine/Private/InstancedStaticMesh.cpp
+++ b/Engine/Source/Runtime/Engine/Private/InstancedStaticMesh.cpp
@@ -3685,6 +3685,10 @@ TArray<int32> UInstancedStaticMeshComponent::AddInstancesInternal(TConstArrayVie
 	const int32 NumInstances = InstanceIndex + Count;
 	PerInstanceSMData.Reserve(NumInstances);
 
+	if (bHasPreviousTransforms)
+	{
+		PerInstancePrevTransform.Reserve(NumInstances);
+	}
 	// Update navigation relevancy
 	// Note that we update navigation relevancy using the base class since the new instances are not added yet
 	// and IsNavigationRelevant() takes the number of instance into account which can be 0 here.
@@ -3711,6 +3715,13 @@ TArray<int32> UInstancedStaticMeshComponent::AddInstancesInternal(TConstArrayVie
 		const FTransform LocalTransform = bWorldSpace ? InstanceTransform.GetRelativeTransform(GetComponentTransform()) : InstanceTransform;
 		SetupNewInstanceData(NewInstanceData, InstanceIndex, LocalTransform);
 
+		if (bHasPreviousTransforms)
+		{
+			// Copy in the current transform - it is somewhat redundant should probably add tracking instead to ensure it gets set, or pass in a struct that can initialize the whole instance 
+			// (same with custom data where we store a mob of zeroes)
+			PerInstancePrevTransform.Emplace(NewInstanceData.Transform);
+		}
+
 		if (bShouldReturnIndices)
 		{
 			NewInstanceIndices.Add(InstanceIndex);
