Material cache, added displacement attribute
- Fixed broken 64-wide block compression

#Synced-CL 52726052

diff --git a/Engine/Shaders/Private/MaterialCache/MaterialCacheABufferPages.usf b/Engine/Shaders/Private/MaterialCache/MaterialCacheABufferPages.usf
--- a/Engine/Shaders/Private/MaterialCache/MaterialCacheABufferPages.usf
+++ b/Engine/Shaders/Private/MaterialCache/MaterialCacheABufferPages.usf
@@ -8,8 +8,13 @@
 
 StructuredBuffer<FMaterialCachePageWriteData> PageWriteData;
 
-RWTexture2D<uint4>  RWVTLayerCompressed;
+#if COMPRESS_MODE == BC_NONE
 RWTexture2D<float4> RWVTLayerUncompressed;
+#elif COMPRESS_MODE == BC1 || COMPRESS_MODE == BC4 // COMPRESS_MODE == BC_NONE
+RWTexture2D<uint2>  RWVTLayerCompressed64;
+#else // COMPRESS_MODE == BC1 || COMPRESS_MODE == BC4
+RWTexture2D<uint4>  RWVTLayerCompressed128;
+#endif
 SamplerState        TextureSampler;
 
 Texture2DArray<float4> ABuffer;
@@ -51,11 +56,11 @@ void WritePagesMain(uint3 DTid : SV_DispatchThreadID)
 
 	if (bSRGB)
 	{
-		RWVTLayerCompressed[VTPhysicalLocation] = float4(CompressBC1Block_SRGB(BlockRGB), 0, 0);
+		RWVTLayerCompressed64[VTPhysicalLocation] = CompressBC1Block_SRGB(BlockRGB);
 	}
 	else
 	{
-		RWVTLayerCompressed[VTPhysicalLocation] = float4(CompressBC1Block(BlockRGB), 0, 0);
+		RWVTLayerCompressed64[VTPhysicalLocation] = CompressBC1Block(BlockRGB);
 	}
 #elif COMPRESS_MODE == BC3 // COMPRESS_MODE == BC1
 	float3 BlockRGB[16];
@@ -64,7 +69,7 @@ void WritePagesMain(uint3 DTid : SV_DispatchThreadID)
 	// Branches aren't problematic
 	if (bYCoCg)
 	{
-		RWVTLayerCompressed[VTPhysicalLocation] = CompressBC3BlockYCoCg(BlockRGB);
+		RWVTLayerCompressed128[VTPhysicalLocation] = CompressBC3BlockYCoCg(BlockRGB);
 	}
 	else
 	{
@@ -73,34 +78,34 @@ void WritePagesMain(uint3 DTid : SV_DispatchThreadID)
 		
 		if (bSRGB)
 		{
-			RWVTLayerCompressed[VTPhysicalLocation] = CompressBC3Block_SRGB(BlockRGB, BlockA);
+			RWVTLayerCompressed128[VTPhysicalLocation] = CompressBC3Block_SRGB(BlockRGB, BlockA);
 		}
 		else
 		{
-			RWVTLayerCompressed[VTPhysicalLocation] = CompressBC3Block(BlockRGB, BlockA);
+			RWVTLayerCompressed128[VTPhysicalLocation] = CompressBC3Block(BlockRGB, BlockA);
 		}
 	}
 #elif COMPRESS_MODE == BC4 // COMPRESS_MODE == BC3
 	float BlockX[16];
 	ReadBlockX(ABuffer, TextureSampler, ABufferPhysicalUV, TexelSize.xy, BlockX);
 
-	RWVTLayerCompressed[VTPhysicalLocation] = float4(CompressBC4Block(BlockX), 0, 0);
+	RWVTLayerCompressed64[VTPhysicalLocation] = CompressBC4Block(BlockX);
 #elif COMPRESS_MODE == BC5 // COMPRESS_MODE == BC4
 	float BlockX[16];
 	float BlockY[16];
 	ReadBlockXY(ABuffer, TextureSampler, ABufferPhysicalUV, TexelSize.xy, BlockX, BlockY);
 
-	RWVTLayerCompressed[VTPhysicalLocation] = CompressBC5Block(BlockX, BlockY);
+	RWVTLayerCompressed128[VTPhysicalLocation] = CompressBC5Block(BlockX, BlockY);
 #elif COMPRESS_MODE == BC6 // COMPRESS_MODE == BC5
 	float3 BlockRGB[16];
 	ReadBlockRGB(ABuffer, TextureSampler, ABufferPhysicalUV, TexelSize.xy, BlockRGB);
 
-	RWVTLayerCompressed[VTPhysicalLocation] = CompressBC6HBlock(BlockRGB);
+	RWVTLayerCompressed128[VTPhysicalLocation] = CompressBC6HBlock(BlockRGB);
 #elif COMPRESS_MODE == BC7 // COMPRESS_MODE == BC6
 	float3 BlockRGB[16];
 	ReadBlockRGB(ABuffer, TextureSampler, ABufferPhysicalUV, TexelSize.xy, BlockRGB);
 
-	RWVTLayerCompressed[VTPhysicalLocation] = CompressBC7Block(BlockRGB);
+	RWVTLayerCompressed128[VTPhysicalLocation] = CompressBC7Block(BlockRGB);
 #else // COMPRESS_MODE == BC7
 	#error Unexpected compression mode
 #endif // COMPRESS_MODE == BC7
diff --git a/Engine/Source/Runtime/Engine/Internal/MaterialCache/MaterialCacheAttribute.h b/Engine/Source/Runtime/Engine/Internal/MaterialCache/MaterialCacheAttribute.h
--- a/Engine/Source/Runtime/Engine/Internal/MaterialCache/MaterialCacheAttribute.h
+++ b/Engine/Source/Runtime/Engine/Internal/MaterialCache/MaterialCacheAttribute.h
@@ -33,6 +33,9 @@ enum class EMaterialCacheAttribute : uint8
 	Opacity,
 	WorldPosition,
 
+	/** Displacement, [0-1] normalized */
+	Displacement,
+
 	/** World height */
 	/** TODO[MP]: Store local to the primitives bounds */
 	WorldHeight,
@@ -53,7 +56,8 @@ enum class EMaterialCacheAttributeIdentity : uint8
 	NormalSpecularOpacity,
 	MetallicWorldPositionOffset,
 	SpecularRoughnessMetallicOpacity,
-	RoughnessMetallic
+	RoughnessMetallic,
+	Displacement
 };
 
 /** All default attributes */
@@ -79,6 +83,7 @@ static bool IsMaterialAttribute(EMaterialCacheAttribute Attribute)
 	case EMaterialCacheAttribute::Metallic:
 	case EMaterialCacheAttribute::Opacity:
 	case EMaterialCacheAttribute::WorldPosition:
+	case EMaterialCacheAttribute::Displacement:
 		return true;
 	}
 }
@@ -178,6 +183,8 @@ static uint8 GetMaterialCacheAttributeComponentCount(EMaterialCacheAttribute Att
 		return 1;
 	case EMaterialCacheAttribute::Float:
 		return 1;
+	case EMaterialCacheAttribute::Displacement:
+		return 1;
 	}
 
 	checkNoEntry();
@@ -226,6 +233,8 @@ static EMaterialValueType GetMaterialCacheAttributeValueType(EMaterialCacheAttri
 		return MCT_Float1;
 	case EMaterialCacheAttribute::Float:
 		return MCT_Float1;
+	case EMaterialCacheAttribute::Displacement:
+		return MCT_Float1;
 	}
 
 	checkNoEntry();
@@ -256,6 +265,8 @@ static FString GetMaterialCacheAttributeDecoration(EMaterialCacheAttribute Attri
 		return TEXT("Mask");
 	case EMaterialCacheAttribute::Float:
 		return TEXT("Float");
+	case EMaterialCacheAttribute::Displacement:
+		return TEXT("Displacement");
 	}
 
 	checkNoEntry();
@@ -276,6 +287,8 @@ static FString GetMaterialCacheLayerDecoration(const FMaterialCacheLayer& Layer)
 			return TEXT("MetallicWorldPositionOffset");
 		case EMaterialCacheAttributeIdentity::RoughnessMetallic:
 			return TEXT("RoughnessMetallic");
+		case EMaterialCacheAttributeIdentity::Displacement:
+			return TEXT("Displacement");
 	}
 
 	FString Composite;
@@ -314,6 +327,9 @@ static FString GetMaterialCacheLayerDecoration(const FMaterialCacheLayer& Layer)
 			case EMaterialCacheAttribute::Float:
 				Composite += TEXT("Float");
 				break;
+			case EMaterialCacheAttribute::Displacement:
+				Composite += TEXT("Displacement");
+				break;
 		}
 	}
 
@@ -400,6 +416,18 @@ static void PackMaterialCacheAttributeLayers(TArray<EMaterialCacheAttribute>& At
 			Layer.ComponentCount = 2;
 			Out.Add(Layer);
 		}
+
+		// Displacement.x
+		if (MaterialCacheAttributeIdentityTest(Attributes, TConstArrayView<EMaterialCacheAttribute>({ EMaterialCacheAttribute::Displacement })))
+		{
+			FMaterialCacheLayer Layer;
+			Layer.Identity = EMaterialCacheAttributeIdentity::Displacement;
+			Layer.Attributes = { EMaterialCacheAttribute::Displacement };
+			Layer.RenderFormat  = PF_G16;
+			Layer.StorageFormat = PF_BC4;
+			Layer.ComponentCount = 1;
+			Out.Add(Layer);
+		}
 	}
 
 	/** After that, try to pack together the attributes automatically */
@@ -449,6 +477,10 @@ static void PackMaterialCacheAttributeLayers(TArray<EMaterialCacheAttribute>& At
 			Layer.RenderFormat  = PF_R32_FLOAT;
 			Layer.StorageFormat = PF_R32_FLOAT;
 			break;
+		case EMaterialCacheAttribute::Displacement:
+			Layer.RenderFormat  = PF_G16;
+			Layer.StorageFormat = PF_BC4;
+			break;
 		}
 
 		Layer.ComponentCount = GetMaterialCacheAttributeComponentCount(Attribute, true);
diff --git a/Engine/Source/Runtime/Engine/Private/Materials/HLSLMaterialTranslator.cpp b/Engine/Source/Runtime/Engine/Private/Materials/HLSLMaterialTranslator.cpp
--- a/Engine/Source/Runtime/Engine/Private/Materials/HLSLMaterialTranslator.cpp
+++ b/Engine/Source/Runtime/Engine/Private/Materials/HLSLMaterialTranslator.cpp
@@ -15415,6 +15415,15 @@ void FHLSLMaterialTranslator::GenerateMaterialCacheSource()
 							);
 							break;
 						}
+						case EMaterialCacheAttributeIdentity::Displacement:
+						{
+							GeneratedCommon += FString::Printf(
+								TEXT("\tABuffer.%s = Layer%i.x;\n"),
+								*FString::Printf(TEXT("%s%u"), *GetMaterialCacheAttributeDecoration(EMaterialCacheAttribute::Displacement), LocalCounter++),
+								i
+							);
+							break;
+						}
 					}
 				}
 			}
@@ -15459,6 +15468,9 @@ void FHLSLMaterialTranslator::GenerateMaterialCacheSource()
 					case EMaterialCacheAttribute::Metallic:
 						Value = TEXT("saturate(PixelMaterialInputs.Metallic)");
 						break;
+					case EMaterialCacheAttribute::Displacement:
+						Value = TEXT("saturate(PixelMaterialInputs.Displacement)");
+						break;
 					case EMaterialCacheAttribute::Opacity:
 					case EMaterialCacheAttribute::WorldPosition:
 					case EMaterialCacheAttribute::WorldHeight:
@@ -15625,6 +15637,17 @@ void FHLSLMaterialTranslator::GenerateMaterialCacheSource()
 							);
 							break;
 						}
+						case EMaterialCacheAttributeIdentity::Displacement:
+						{
+							int32 DisplacementIndex = LocalCounter++;
+							
+							GeneratedPass += FString::Printf(
+								TEXT("\tLayer%u = float4(ABuffer.%s, 0, 0, 0);\n"),
+								i,
+								*FString::Printf(TEXT("%s%u"), *GetMaterialCacheAttributeDecoration(EMaterialCacheAttribute::Displacement), DisplacementIndex)
+							);
+							break;
+						}
 					}
 				}
 
diff --git a/Engine/Source/Runtime/Renderer/Internal/MaterialCache/MaterialCacheShaders.h b/Engine/Source/Runtime/Renderer/Internal/MaterialCache/MaterialCacheShaders.h
--- a/Engine/Source/Runtime/Renderer/Internal/MaterialCache/MaterialCacheShaders.h
+++ b/Engine/Source/Runtime/Renderer/Internal/MaterialCache/MaterialCacheShaders.h
@@ -125,7 +125,8 @@ class FMaterialCacheABufferWritePagesCS : public FGlobalShader
 public:
 	BEGIN_SHADER_PARAMETER_STRUCT(FParameters, )
 		SHADER_PARAMETER_RDG_BUFFER_SRV(StructuredBuffer<FMaterialCacheBinData>, PageWriteData)
-		SHADER_PARAMETER_RDG_TEXTURE_UAV(RWTexture2D<uint4>, RWVTLayerCompressed)
+		SHADER_PARAMETER_RDG_TEXTURE_UAV(RWTexture2D<uint2>, RWVTLayerCompressed64)
+		SHADER_PARAMETER_RDG_TEXTURE_UAV(RWTexture2D<uint4>, RWVTLayerCompressed128)
 		SHADER_PARAMETER_RDG_TEXTURE_UAV(RWTexture2D<float4>, RWVTLayerUncompressed)
 		SHADER_PARAMETER_RDG_TEXTURE_SRV(Texture2DArray<float4>, ABuffer)
 		SHADER_PARAMETER_SAMPLER(SamplerState, TextureSampler)
diff --git a/Engine/Source/Runtime/Renderer/Private/MaterialCache/MaterialCacheRenderer.cpp b/Engine/Source/Runtime/Renderer/Private/MaterialCache/MaterialCacheRenderer.cpp
--- a/Engine/Source/Runtime/Renderer/Private/MaterialCache/MaterialCacheRenderer.cpp
+++ b/Engine/Source/Runtime/Renderer/Private/MaterialCache/MaterialCacheRenderer.cpp
@@ -1781,12 +1781,9 @@ static void MaterialCacheFinalizePages(FSceneRendererBase* Renderer, FRDGBuilder
 		// Destination target
 		IPooledRenderTarget* RenderTarget = RenderData.Bucket->PendingEntries[0].Setup.PhysicalRenderTargets[ABufferIndex];
 
-		// The "compressed" format of this layer
-		// TODO: Maybe just call it storage format?
-		EPixelFormat CompressedFormat = RenderData.Bucket->TagLayout.Layers[ABufferIndex].StorageFormat;
-
 		// Are we writing to a compressed format?
-		const bool bIsCompressed = IsBlockCompressedFormat(CompressedFormat);
+		EPixelFormat Format = RenderTarget->GetDesc().Format;
+		const bool bIsCompressed = IsBlockCompressedFormat(Format);
 		
 		FMaterialCacheABufferWritePagesCS::FParameters* PassParameters = GraphBuilder.AllocParameters<FMaterialCacheABufferWritePagesCS::FParameters>();
 		PassParameters->PageWriteData = GraphBuilder.CreateSRV(PageWriteData);
@@ -1797,20 +1794,29 @@ static void MaterialCacheFinalizePages(FSceneRendererBase* Renderer, FRDGBuilder
 		PassParameters->bSRGB = RenderData.Bucket->TagLayout.Layers[ABufferIndex].bIsSRGB;
 		PassParameters->BlockOrThreadCount = bIsCompressed ? BlockThreadCountX : TexelThreadCountX;
 
+		// 128 wide?
+		const bool bAliasTo64bit = Format == PF_DXT1 || Format == PF_BC4;
+		const EPixelFormat AliasFormat = bAliasTo64bit ? PF_R32G32_UINT : PF_R32G32B32A32_UINT;
+		
 		// If compressed, bind the compressed aliased format, otherwise the uncompressed
 		if (bIsCompressed)
 		{
-			PassParameters->RWVTLayerCompressed = GraphBuilder.CreateUAV(GraphBuilder.RegisterExternalTexture(RenderTarget, ERDGTextureFlags::ForceImmediateFirstBarrier));
-			PassParameters->RWVTLayerUncompressed = PassParameters->RWVTLayerCompressed;
+			if (bAliasTo64bit)
+			{
+				PassParameters->RWVTLayerCompressed64 = GraphBuilder.CreateUAV(GraphBuilder.RegisterExternalTexture(RenderTarget, ERDGTextureFlags::ForceImmediateFirstBarrier), ERDGUnorderedAccessViewFlags::None, AliasFormat);
+			}
+			else
+			{
+				PassParameters->RWVTLayerCompressed128 = GraphBuilder.CreateUAV(GraphBuilder.RegisterExternalTexture(RenderTarget, ERDGTextureFlags::ForceImmediateFirstBarrier), ERDGUnorderedAccessViewFlags::None, AliasFormat);
+			}
 		}
 		else
 		{
-			PassParameters->RWVTLayerUncompressed = GraphBuilder.CreateUAV(GraphBuilder.RegisterExternalTexture(RenderTarget, ERDGTextureFlags::ForceImmediateFirstBarrier),  ERDGUnorderedAccessViewFlags::None, CompressedFormat);
-			PassParameters->RWVTLayerCompressed = PassParameters->RWVTLayerUncompressed;
+			PassParameters->RWVTLayerUncompressed = GraphBuilder.CreateUAV(GraphBuilder.RegisterExternalTexture(RenderTarget, ERDGTextureFlags::ForceImmediateFirstBarrier),  ERDGUnorderedAccessViewFlags::None, Format);
 		}
 
 		FMaterialCacheABufferWritePagesCS::FPermutationDomain PermutationVector;
-		PermutationVector.Set<FMaterialCacheABufferWritePagesCS::FCompressMode>(FMaterialCacheABufferWritePagesCS::GetCompressMode(CompressedFormat));
+		PermutationVector.Set<FMaterialCacheABufferWritePagesCS::FCompressMode>(FMaterialCacheABufferWritePagesCS::GetCompressMode(Format));
 		
 		GraphBuilder.AddPass(
 			RDG_EVENT_NAME("WritePages"),
