Fixed spikes when loading maps with large HISMs that were introduced as part of a fix for a different issue. This new fix is narrower in scope and should only cause rebuilds when ApplyComponentInstanceData is called.

Original fix CL 29842740, and jira UE-197349

#Synced-CL 52728553

diff --git a/Engine/Source/Runtime/Engine/Private/HierarchicalInstancedStaticMesh.cpp b/Engine/Source/Runtime/Engine/Private/HierarchicalInstancedStaticMesh.cpp
--- a/Engine/Source/Runtime/Engine/Private/HierarchicalInstancedStaticMesh.cpp
+++ b/Engine/Source/Runtime/Engine/Private/HierarchicalInstancedStaticMesh.cpp
@@ -2453,8 +2453,8 @@ bool UHierarchicalInstancedStaticMeshComponent::BatchUpdateInstancesData(int32 S
 void UHierarchicalInstancedStaticMeshComponent::ApplyComponentInstanceData(FInstancedStaticMeshComponentInstanceData* InstancedMeshData)
 {
 	UInstancedStaticMeshComponent::ApplyComponentInstanceData(InstancedMeshData);
-
-	BuildTreeIfOutdated(/*Async*/false, /*ForceUpdate*/false);
+	// Force rebuild, since anything may have changed in the ApplyComponentInstanceData.
+	BuildTreeIfOutdated(/*Async*/false, /*ForceUpdate*/true);
 }
 
 void UHierarchicalInstancedStaticMeshComponent::PreAllocateInstancesMemory(int32 AddedInstanceCount)
@@ -2789,7 +2789,7 @@ bool UHierarchicalInstancedStaticMeshComponent::BuildTreeIfOutdated(bool Async,
 
 	if (ForceUpdate 
 		|| bIsOutOfDate
-		|| PrimitiveInstanceDataManager.HasAnyChanges()
+		|| PrimitiveInstanceDataManager.HasAnyInstanceChanges()
 		|| InstanceReorderTable.Num() != PerInstanceSMData.Num()
 		|| NumBuiltInstances != PerInstanceSMData.Num() 
 		|| (GetStaticMesh() != nullptr && CacheMeshExtendedBounds != GetStaticMesh()->GetBounds())
