Splines: Don't create physics state if its a visualization component

CameraRig_Rail has spline components as visualization components. When saving, the spline components are regenerated, causing physics state to be recreated. That causes Sequencer to think that the spawned camera rig rail is modified and then marking the sequence as dirty.

Before 5.6, the sequence didn't get marked dirty, but I wasn't able to track down why.

#Synced-CL 52681601

diff --git a/Engine/Source/Runtime/Engine/Private/Components/SplineMeshComponent.cpp b/Engine/Source/Runtime/Engine/Private/Components/SplineMeshComponent.cpp
--- a/Engine/Source/Runtime/Engine/Private/Components/SplineMeshComponent.cpp
+++ b/Engine/Source/Runtime/Engine/Private/Components/SplineMeshComponent.cpp
@@ -974,6 +974,13 @@ void USplineMeshComponent::GetMeshId(FString& OutMeshId) const
 
 void USplineMeshComponent::OnCreatePhysicsState()
 {
+#if WITH_EDITOR
+	if (IsVisualizationComponent())
+	{
+		return Super::OnCreatePhysicsState();
+	}
+#endif
+
 	// With editor code we can recreate the collision if the mesh changes
 	const FGuid MeshBodySetupGuid = (GetStaticMesh() != nullptr ? GetStaticMesh()->GetBodySetup()->BodySetupGuid : FGuid());
 	if (CachedMeshBodySetupGuid != MeshBodySetupGuid)
